Move InteractiveToolsFramework and EditorInteractiveToolsFramework modules to Experimental
#rnx

diff --git a/Engine/Source/Editor/EditorInteractiveToolsFramework/EditorInteractiveToolsFramework.Build.cs b/Engine/Source/Editor/EditorInteractiveToolsFramework/EditorInteractiveToolsFramework.Build.cs
deleted file mode 100644
diff --git a/Engine/Source/Editor/EditorInteractiveToolsFramework/Private/EdModeInteractiveToolsContext.cpp b/Engine/Source/Editor/EditorInteractiveToolsFramework/Private/EdModeInteractiveToolsContext.cpp
deleted file mode 100644
diff --git a/Engine/Source/Editor/EditorInteractiveToolsFramework/Private/EditorComponentSourceFactory.cpp b/Engine/Source/Editor/EditorInteractiveToolsFramework/Private/EditorComponentSourceFactory.cpp
deleted file mode 100644
diff --git a/Engine/Source/Editor/EditorInteractiveToolsFramework/Private/EditorInteractiveToolsFrameworkModule.cpp b/Engine/Source/Editor/EditorInteractiveToolsFramework/Private/EditorInteractiveToolsFrameworkModule.cpp
deleted file mode 100644
diff --git a/Engine/Source/Editor/EditorInteractiveToolsFramework/Private/EditorToolAssetAPI.cpp b/Engine/Source/Editor/EditorInteractiveToolsFramework/Private/EditorToolAssetAPI.cpp
deleted file mode 100644
diff --git a/Engine/Source/Editor/EditorInteractiveToolsFramework/Private/StandardToolModeCommands.cpp b/Engine/Source/Editor/EditorInteractiveToolsFramework/Private/StandardToolModeCommands.cpp
deleted file mode 100644
diff --git a/Engine/Source/Editor/EditorInteractiveToolsFramework/Public/EdModeInteractiveToolsContext.h b/Engine/Source/Editor/EditorInteractiveToolsFramework/Public/EdModeInteractiveToolsContext.h
deleted file mode 100644
diff --git a/Engine/Source/Editor/EditorInteractiveToolsFramework/Public/EditorComponentSourceFactory.h b/Engine/Source/Editor/EditorInteractiveToolsFramework/Public/EditorComponentSourceFactory.h
deleted file mode 100644
diff --git a/Engine/Source/Editor/EditorInteractiveToolsFramework/Public/EditorInteractiveToolsFrameworkModule.h b/Engine/Source/Editor/EditorInteractiveToolsFramework/Public/EditorInteractiveToolsFrameworkModule.h
deleted file mode 100644
diff --git a/Engine/Source/Editor/EditorInteractiveToolsFramework/Public/EditorToolAssetAPI.h b/Engine/Source/Editor/EditorInteractiveToolsFramework/Public/EditorToolAssetAPI.h
deleted file mode 100644
diff --git a/Engine/Source/Editor/EditorInteractiveToolsFramework/Public/InteractiveToolsCommands.h b/Engine/Source/Editor/EditorInteractiveToolsFramework/Public/InteractiveToolsCommands.h
deleted file mode 100644
diff --git a/Engine/Source/Editor/EditorInteractiveToolsFramework/Public/StandardToolModeCommands.h b/Engine/Source/Editor/EditorInteractiveToolsFramework/Public/StandardToolModeCommands.h
deleted file mode 100644
diff --git a/Engine/Source/Editor/Experimental/EditorInteractiveToolsFramework/EditorInteractiveToolsFramework.Build.cs b/Engine/Source/Editor/Experimental/EditorInteractiveToolsFramework/EditorInteractiveToolsFramework.Build.cs
new file mode 100644
--- /dev/null
+++ b/Engine/Source/Editor/Experimental/EditorInteractiveToolsFramework/EditorInteractiveToolsFramework.Build.cs
@@ -0,0 +1,62 @@
+// Copyright 1998-2019 Epic Games, Inc. All Rights Reserved.
+
+using UnrealBuildTool;
+
+public class EditorInteractiveToolsFramework : ModuleRules
+{
+	public EditorInteractiveToolsFramework(ReadOnlyTargetRules Target) : base(Target)
+	{
+		PCHUsage = ModuleRules.PCHUsageMode.UseExplicitOrSharedPCHs;
+		
+		PublicIncludePaths.AddRange(
+			new string[] {
+                //"ContentBrowser"
+				// ... add public include paths required here ...
+			}
+			);
+				
+		
+		PrivateIncludePaths.AddRange(
+			new string[] {
+				// ... add other private include paths required here ...
+			}
+			);
+			
+		
+		PublicDependencyModuleNames.AddRange(
+			new string[]
+			{
+				"Core"
+				// ... add other public dependencies that you statically link with here ...
+			}
+			);
+			
+		
+		PrivateDependencyModuleNames.AddRange(
+			new string[]
+			{
+				"CoreUObject",
+				"Engine",
+                "Slate",
+                "SlateCore",
+                "InputCore",
+				"UnrealEd",
+                "ContentBrowser",
+                "LevelEditor",
+                "ApplicationCore",
+                "EditorStyle",
+                "InteractiveToolsFramework"
+
+				// ... add private dependencies that you statically link with here ...	
+			}
+            );
+		
+		
+		DynamicallyLoadedModuleNames.AddRange(
+			new string[]
+			{
+				// ... add any modules that your module loads dynamically here ...
+			}
+			);
+	}
+}
diff --git a/Engine/Source/Editor/Experimental/EditorInteractiveToolsFramework/Private/EdModeInteractiveToolsContext.cpp b/Engine/Source/Editor/Experimental/EditorInteractiveToolsFramework/Private/EdModeInteractiveToolsContext.cpp
new file mode 100644
--- /dev/null
+++ b/Engine/Source/Editor/Experimental/EditorInteractiveToolsFramework/Private/EdModeInteractiveToolsContext.cpp
@@ -0,0 +1,788 @@
+// Copyright 1998-2019 Epic Games, Inc. All Rights Reserved.
+
+
+#include "EdModeInteractiveToolsContext.h"
+#include "Editor.h"
+#include "EditorViewportClient.h"
+#include "EditorModeManager.h"
+#include "LevelEditorViewport.h"   // for GCurrentLevelEditingViewportClient
+#include "Engine/Selection.h"
+#include "Misc/ITransaction.h"
+#include "ScopedTransaction.h"
+#include "Materials/Material.h"
+#include "Engine/StaticMesh.h"
+#include "Components/StaticMeshComponent.h"
+
+#include "EditorToolAssetAPI.h"
+#include "EditorComponentSourceFactory.h"
+
+//#include "PhysicsEngine/BodySetup.h"
+//#include "Interfaces/Interface_CollisionDataProvider.h"
+
+//#define ENABLE_DEBUG_PRINTING
+
+
+
+class FEdModeToolsContextQueriesImpl : public IToolsContextQueriesAPI
+{
+public:
+	UEdModeInteractiveToolsContext* ToolsContext;
+	FEdMode* EditorMode;
+	
+	FViewCameraState CachedViewState;
+
+	FEdModeToolsContextQueriesImpl(UEdModeInteractiveToolsContext* Context, FEdMode* EditorModeIn)
+	{
+		ToolsContext = Context;
+		EditorMode = EditorModeIn;
+	}
+
+	void CacheCurrentViewState(FEditorViewportClient* ViewportClient)
+	{
+		FViewportCameraTransform ViewTransform = ViewportClient->GetViewTransform();
+		CachedViewState.Position = ViewTransform.GetLocation();
+		CachedViewState.Orientation = ViewTransform.GetRotation().Quaternion();
+		CachedViewState.bIsOrthographic = ViewportClient->IsOrtho();
+		CachedViewState.bIsVR = false;
+	}
+
+	virtual void GetCurrentSelectionState(FToolBuilderState& StateOut) const override
+	{
+		StateOut.ToolManager = ToolsContext->ToolManager;
+		StateOut.GizmoManager = ToolsContext->GizmoManager;
+		StateOut.World = EditorMode->GetWorld();
+		StateOut.SelectedActors = EditorMode->GetModeManager()->GetSelectedActors();
+		StateOut.SelectedComponents = EditorMode->GetModeManager()->GetSelectedComponents();
+		StateOut.SourceBuilder = ToolsContext->GetComponentSourceFactory();
+	}
+
+
+	virtual void GetCurrentViewState(FViewCameraState& StateOut) const override
+	{
+		StateOut = CachedViewState;
+	}
+
+
+	virtual bool ExecuteSceneSnapQuery(const FSceneSnapQueryRequest& Request, TArray<FSceneSnapQueryResult>& Results) const override
+	{
+		if (Request.RequestType != ESceneSnapQueryType::Position)
+		{
+			return false;		// not supported yet
+		}
+
+		int FoundResultCount = 0;
+
+		//
+		// Run a snap query by casting ray into the world.
+		// If a hit is found, we look up what triangle was hit, and then test its vertices and edges
+		// 
+
+		// cast ray into world
+		FVector RayStart = CachedViewState.Position;
+		FVector RayDirection = Request.Position - RayStart; RayDirection.Normalize();
+		FVector RayEnd = RayStart + 9999999 * RayDirection;
+		FCollisionObjectQueryParams ObjectQueryParams(FCollisionObjectQueryParams::AllObjects);
+		FCollisionQueryParams QueryParams = FCollisionQueryParams::DefaultQueryParam;
+		QueryParams.bTraceComplex = true;
+		QueryParams.bReturnFaceIndex = true;
+		FHitResult HitResult;
+		bool bHitWorld = EditorMode->GetWorld()->LineTraceSingleByObjectType(HitResult, RayStart, RayEnd, ObjectQueryParams, QueryParams);
+		if (bHitWorld && HitResult.FaceIndex >= 0) 
+		{
+			float VisualAngle = OpeningAngleDeg(Request.Position, HitResult.ImpactPoint, RayStart);
+			//UE_LOG(LogTemp, Warning, TEXT("[HIT] visualangle %f faceindex %d"), VisualAngle, HitResult.FaceIndex);
+			if (VisualAngle < Request.VisualAngleThresholdDegrees)
+			{
+				UPrimitiveComponent* Component = HitResult.Component.Get();
+				if (Cast<UStaticMeshComponent>(Component) != nullptr)
+				{
+					// HitResult.FaceIndex is apparently an index into the TriMeshCollisionData, not sure how
+					// to directly access it. Calling GetPhysicsTriMeshData is expensive!
+					//UBodySetup* BodySetup = Cast<UStaticMeshComponent>(Component)->GetBodySetup();
+					//UObject* CDPObj = BodySetup->GetOuter();
+					//IInterface_CollisionDataProvider* CDP = Cast<IInterface_CollisionDataProvider>(CDPObj);
+					//FTriMeshCollisionData TriMesh;
+					//CDP->GetPhysicsTriMeshData(&TriMesh, true);
+					//FTriIndices Triangle = TriMesh.Indices[HitResult.FaceIndex];
+					//FVector Positions[3] = { TriMesh.Vertices[Triangle.v0], TriMesh.Vertices[Triangle.v1], TriMesh.Vertices[Triangle.v2] };
+
+					// physics collision data is created from StaticMesh RenderData
+					// so use HitResult.FaceIndex to extract triangle from the LOD0 mesh
+					// (note: this may be incorrect if there are multiple sections...in that case I think we have to 
+					//  first find section whose accumulated index range would contain .FaceIndexX)
+					UStaticMesh* StaticMesh = Cast<UStaticMeshComponent>(Component)->GetStaticMesh();
+					FStaticMeshLODResources& LOD = StaticMesh->RenderData->LODResources[0];
+					FIndexArrayView Indices = LOD.IndexBuffer.GetArrayView();
+					int32 TriIdx = 3 * HitResult.FaceIndex;
+					FVector Positions[3];
+					Positions[0] = LOD.VertexBuffers.PositionVertexBuffer.VertexPosition(Indices[TriIdx]);
+					Positions[1] = LOD.VertexBuffers.PositionVertexBuffer.VertexPosition(Indices[TriIdx+1]);
+					Positions[2] = LOD.VertexBuffers.PositionVertexBuffer.VertexPosition(Indices[TriIdx+2]);
+					
+					// transform to world space
+					FTransform ComponentTransform = Component->GetComponentTransform();
+					Positions[0] = ComponentTransform.TransformPosition(Positions[0]);
+					Positions[1] = ComponentTransform.TransformPosition(Positions[1]);
+					Positions[2] = ComponentTransform.TransformPosition(Positions[2]);
+
+					FSceneSnapQueryResult SnapResult;
+					SnapResult.TriVertices[0] = Positions[0];
+					SnapResult.TriVertices[1] = Positions[1];
+					SnapResult.TriVertices[2] = Positions[2];
+
+					// try snapping to vertices
+					float SmallestAngle = Request.VisualAngleThresholdDegrees;
+					if ( (Request.TargetTypes & ESceneSnapQueryTargetType::MeshVertex) != ESceneSnapQueryTargetType::None)
+					{
+						for (int j = 0; j < 3; ++j)
+						{
+							VisualAngle = OpeningAngleDeg(Request.Position, Positions[j], RayStart);
+							if (VisualAngle < SmallestAngle)
+							{
+								SmallestAngle = VisualAngle;
+								SnapResult.Position = Positions[j];
+								SnapResult.TargetType = ESceneSnapQueryTargetType::MeshVertex;
+								SnapResult.TriSnapIndex = j;
+							}
+						}
+					}
+
+					// try snapping to nearest points on edges
+					if ( ((Request.TargetTypes & ESceneSnapQueryTargetType::MeshEdge) != ESceneSnapQueryTargetType::None) &&
+						 (SnapResult.TargetType != ESceneSnapQueryTargetType::MeshVertex) )
+					{
+						for (int j = 0; j < 3; ++j)
+						{
+							FVector EdgeNearestPt = NearestSegmentPt(Positions[j], Positions[(j+1)%3], Request.Position);
+							VisualAngle = OpeningAngleDeg(Request.Position, EdgeNearestPt, RayStart);
+							if (VisualAngle < SmallestAngle )
+							{
+								SmallestAngle = VisualAngle;
+								SnapResult.Position = EdgeNearestPt;
+								SnapResult.TargetType = ESceneSnapQueryTargetType::MeshEdge;
+								SnapResult.TriSnapIndex = j;
+							}
+						}
+					}
+
+					// if we found a valid snap, return it
+					if (SmallestAngle < Request.VisualAngleThresholdDegrees)
+					{
+						SnapResult.TargetActor = HitResult.Actor.Get();
+						SnapResult.TargetComponent = HitResult.Component.Get();
+						Results.Add(SnapResult);
+						FoundResultCount++;
+					}
+				}
+			}
+
+		}
+
+		return (FoundResultCount > 0);
+	}
+
+
+	//@ todo this are mirrored from GeometryProcessing, which is still experimental...replace w/ direct calls once GP component is standardized
+	static float OpeningAngleDeg(FVector A, FVector B, const FVector& P)
+	{
+		A -= P;
+		A.Normalize();
+		B -= P;
+		B.Normalize();
+		float Dot = FMath::Clamp(FVector::DotProduct(A,B), -1.0f, 1.0f);
+		return acos(Dot) * (180.0f / 3.141592653589f);
+	}
+	static FVector NearestSegmentPt(FVector A, FVector B, const FVector& P)
+	{
+		FVector Direction = (B - A); 
+		float Length = Direction.Size();
+		Direction /= Length;
+		float t = FVector::DotProduct( (P - A), Direction);
+		if (t >= Length)
+		{
+			return B;
+		}
+		if (t <= 0)
+		{
+			return A;
+		}
+		return A + t * Direction;
+	}
+
+
+
+	virtual UMaterialInterface* GetStandardMaterial(EStandardToolContextMaterials MaterialType) const
+	{
+		if (MaterialType == EStandardToolContextMaterials::VertexColorMaterial)
+		{
+			return ToolsContext->StandardVertexColorMaterial;
+		}
+		check(false);
+		return nullptr;
+	}
+
+};
+
+
+
+
+class FEdModeToolsContextTransactionImpl : public IToolsContextTransactionsAPI
+{
+public:
+	UEdModeInteractiveToolsContext* ToolsContext;
+	FEdMode* EditorMode;
+
+	FEdModeToolsContextTransactionImpl(UEdModeInteractiveToolsContext* Context, FEdMode* EditorModeIn)
+	{
+		ToolsContext = Context;
+		EditorMode = EditorModeIn;
+	}
+
+
+	virtual void PostMessage(const TCHAR* Message, EToolMessageLevel Level) override
+	{
+		UE_LOG(LogTemp, Warning, TEXT("[ToolsContext] %s"), Message);
+	}
+
+
+	virtual void PostInvalidation() override
+	{
+		ToolsContext->PostInvalidation();
+	}
+
+	virtual void BeginUndoTransaction(const FText& Description) override
+	{
+		GEditor->BeginTransaction(Description);
+	}
+
+	virtual void EndUndoTransaction() override
+	{
+		GEditor->EndTransaction();
+	}
+
+	virtual void AppendChange(UObject* TargetObject, TUniquePtr<FChange> Change, const FText& Description) override
+	{
+		FScopedTransaction Transaction(Description);
+		check(GUndo != nullptr);
+		GUndo->StoreUndo(TargetObject, MoveTemp(Change));
+		// end transaction
+	}
+
+
+	virtual bool RequestSelectionChange(const FSelectedOjectsChangeList& SelectionChange) override
+	{
+		checkf(SelectionChange.Components.Num() == 0, TEXT("FEdModeToolsContextTransactionImpl::RequestSelectionChange - Component selection not supported yet"));
+
+		if (SelectionChange.ModificationType == ESelectedObjectsModificationType::Clear)
+		{
+			GEditor->SelectNone(true, true, false);
+			return true;
+		}
+
+		if (SelectionChange.ModificationType == ESelectedObjectsModificationType::Replace )
+		{
+			GEditor->SelectNone(false, true, false);
+		}
+
+		bool bAdd = (SelectionChange.ModificationType != ESelectedObjectsModificationType::Remove);
+		int NumActors = SelectionChange.Actors.Num();
+		for (int k = 0; k < NumActors; ++k)
+		{
+			GEditor->SelectActor(SelectionChange.Actors[k], bAdd, false, true, false);
+		}
+
+		GEditor->NoteSelectionChange(true);
+		return true;
+	}
+
+};
+
+
+
+
+
+UEdModeInteractiveToolsContext::UEdModeInteractiveToolsContext()
+{
+	EditorMode = nullptr;
+	QueriesAPI = nullptr;
+	TransactionAPI = nullptr;
+	AssetAPI = nullptr;
+	SourceFactory = nullptr;
+}
+
+
+
+void UEdModeInteractiveToolsContext::Initialize(IToolsContextQueriesAPI* QueriesAPIIn, IToolsContextTransactionsAPI* TransactionsAPIIn)
+{
+	UInteractiveToolsContext::Initialize(QueriesAPIIn, TransactionsAPIIn);
+
+	BeginPIEDelegateHandle = FEditorDelegates::BeginPIE.AddLambda([this](bool bSimulating)
+	{
+		TerminateActiveToolsOnPIEStart();
+	});
+	PreSaveWorldDelegateHandle = FEditorDelegates::PreSaveWorld.AddLambda([this](uint32 SaveFlags, UWorld* World)
+	{
+		TerminateActiveToolsOnSaveWorld();
+	});
+
+	bInvalidationPending = false;
+}
+
+void UEdModeInteractiveToolsContext::Shutdown()
+{
+	FEditorDelegates::BeginPIE.Remove(BeginPIEDelegateHandle);
+	FEditorDelegates::PreSaveWorld.Remove(PreSaveWorldDelegateHandle);
+
+	UInteractiveToolsContext::Shutdown();
+}
+
+
+
+
+void UEdModeInteractiveToolsContext::InitializeContextFromEdMode(FEdMode* EditorModeIn)
+{
+	this->EditorMode = EditorModeIn;
+
+	this->TransactionAPI = new FEdModeToolsContextTransactionImpl(this, EditorModeIn);
+	this->QueriesAPI = new FEdModeToolsContextQueriesImpl(this, EditorModeIn);
+	this->AssetAPI = new FEditorToolAssetAPI();
+	this->SourceFactory = new FEditorComponentSourceFactory();
+
+	Initialize(QueriesAPI, TransactionAPI);
+
+	// enable auto invalidation in Editor, because invalidating for all hover and capture events is unpleasant
+	this->InputRouter->bAutoInvalidateOnHover = true;
+	this->InputRouter->bAutoInvalidateOnCapture = true;
+
+
+	// set up standard materials
+	StandardVertexColorMaterial = LoadObject<UMaterial>(nullptr, TEXT("/Game/Materials/VertexColor"));
+}
+
+
+void UEdModeInteractiveToolsContext::ShutdownContext()
+{
+	Shutdown();
+
+	if (QueriesAPI != nullptr)
+	{
+		delete QueriesAPI;
+		QueriesAPI = nullptr;
+	}
+
+	if (TransactionAPI != nullptr)
+	{
+		delete TransactionAPI;
+		TransactionAPI = nullptr;
+	}
+
+	if (AssetAPI != nullptr)
+	{
+		delete AssetAPI;
+		AssetAPI = nullptr;
+	}
+
+	if (SourceFactory != nullptr)
+	{
+		delete SourceFactory;
+		SourceFactory = nullptr;
+	}
+
+	this->EditorMode = nullptr;
+}
+
+
+
+
+void UEdModeInteractiveToolsContext::TerminateActiveToolsOnPIEStart()
+{
+	if (ToolManager->HasActiveTool(EToolSide::Left))
+	{
+		EToolShutdownType ShutdownType = ToolManager->CanAcceptActiveTool(EToolSide::Left) ?
+			EToolShutdownType::Accept : EToolShutdownType::Cancel;
+		ToolManager->DeactivateTool(EToolSide::Left, ShutdownType);
+	}
+	if (ToolManager->HasActiveTool(EToolSide::Right))
+	{
+		EToolShutdownType ShutdownType = ToolManager->CanAcceptActiveTool(EToolSide::Right) ?
+			EToolShutdownType::Accept : EToolShutdownType::Cancel;
+		ToolManager->DeactivateTool(EToolSide::Right, ShutdownType);
+	}
+}
+void UEdModeInteractiveToolsContext::TerminateActiveToolsOnSaveWorld()
+{
+	if (ToolManager->HasActiveTool(EToolSide::Left))
+	{
+		EToolShutdownType ShutdownType = ToolManager->CanAcceptActiveTool(EToolSide::Left) ?
+			EToolShutdownType::Accept : EToolShutdownType::Cancel;
+		ToolManager->DeactivateTool(EToolSide::Left, ShutdownType);
+	}
+	if (ToolManager->HasActiveTool(EToolSide::Right))
+	{
+		EToolShutdownType ShutdownType = ToolManager->CanAcceptActiveTool(EToolSide::Right) ?
+			EToolShutdownType::Accept : EToolShutdownType::Cancel;
+		ToolManager->DeactivateTool(EToolSide::Right, ShutdownType);
+	}
+}
+
+
+
+
+
+void UEdModeInteractiveToolsContext::PostInvalidation()
+{
+	bInvalidationPending = true;
+}
+
+
+void UEdModeInteractiveToolsContext::Tick(FEditorViewportClient* ViewportClient, float DeltaTime)
+{
+	ToolManager->Tick(DeltaTime);
+	GizmoManager->Tick(DeltaTime);
+
+	if (bInvalidationPending)
+	{
+		ViewportClient->Invalidate();
+		bInvalidationPending = false;
+	}
+
+	// save this view
+	// Check against GCurrentLevelEditingViewportClient is temporary and should be removed in future.
+	// Current issue is that this ::Tick() is called *per viewport*, so once for each view in a 4-up view.
+	if (ViewportClient == GCurrentLevelEditingViewportClient)
+	{
+		((FEdModeToolsContextQueriesImpl*)this->QueriesAPI)->CacheCurrentViewState(ViewportClient);
+	}
+}
+
+
+
+class TempRenderContext : public IToolsContextRenderAPI
+{
+public:
+	FPrimitiveDrawInterface* PDI;
+
+	virtual FPrimitiveDrawInterface* GetPrimitiveDrawInterface() override
+	{
+		return PDI;
+	}
+};
+
+
+void UEdModeInteractiveToolsContext::Render(const FSceneView* View, FViewport* Viewport, FPrimitiveDrawInterface* PDI)
+{
+	TempRenderContext RenderContext;
+	RenderContext.PDI = PDI;
+	ToolManager->Render(&RenderContext);
+	GizmoManager->Render(&RenderContext);
+}
+
+
+
+
+bool UEdModeInteractiveToolsContext::InputKey(FEditorViewportClient* ViewportClient, FViewport* Viewport, FKey Key, EInputEvent Event)
+{
+#ifdef ENABLE_DEBUG_PRINTING
+	if (Event == IE_Pressed) { UE_LOG(LogTemp, Warning, TEXT("PRESSED EVENT")); }
+	else if (Event == IE_Released) { UE_LOG(LogTemp, Warning, TEXT("RELEASED EVENT")); }
+	else if (Event == IE_Repeat) { UE_LOG(LogTemp, Warning, TEXT("REPEAT EVENT")); }
+	else if (Event == IE_Axis) { UE_LOG(LogTemp, Warning, TEXT("AXIS EVENT")); }
+	else if (Event == IE_DoubleClick) { UE_LOG(LogTemp, Warning, TEXT("DOUBLECLICK EVENT")); }
+#endif
+
+	bool bHandled = false;
+
+
+	// escape key cancels current tool
+	if (Key == EKeys::Escape && Event == IE_Released )
+	{
+		if (ToolManager->HasAnyActiveTool())
+		{
+			if (ToolManager->HasActiveTool(EToolSide::Mouse))
+			{
+				ToolManager->DeactivateTool(EToolSide::Mouse, EToolShutdownType::Cancel);
+			}
+			return true;
+		}
+	}
+
+	// enter key accepts current tool, or ends tool if it does not have accept state
+	if (Key == EKeys::Enter && Event == IE_Released && ToolManager->HasAnyActiveTool())
+	{
+		if (ToolManager->HasActiveTool(EToolSide::Mouse)) 
+		{
+			if (ToolManager->GetActiveTool(EToolSide::Mouse)->HasAccept())
+			{
+				if (ToolManager->CanAcceptActiveTool(EToolSide::Mouse))
+				{
+					ToolManager->DeactivateTool(EToolSide::Mouse, EToolShutdownType::Accept);
+					return true;
+				}
+			}
+			else
+			{
+				ToolManager->DeactivateTool(EToolSide::Mouse, EToolShutdownType::Completed);
+				return true;
+			}
+		}
+	}
+
+	// if alt is down we do not process mouse event
+	if (ViewportClient->IsAltPressed())
+	{
+		return false;
+	}
+
+	if (Event == IE_Pressed || Event == IE_Released)
+	{
+		if (Key.IsMouseButton())
+		{
+			bool bIsLeftMouse = (Key == EKeys::LeftMouseButton);
+			bool bIsMiddleMouse = (Key == EKeys::MiddleMouseButton);
+			bool bIsRightMouse = (Key == EKeys::RightMouseButton);
+
+			if (bIsLeftMouse || bIsMiddleMouse || bIsRightMouse)
+			{
+
+				// early-out here if we are going to do camera manipulation
+				if (ViewportClient->IsAltPressed())
+				{
+					return bHandled;
+				}
+
+				FInputDeviceState InputState = CurrentMouseState;
+				InputState.InputDevice = EInputDevices::Mouse;
+				InputState.SetModifierKeyStates(
+					ViewportClient->IsShiftPressed(), ViewportClient->IsAltPressed(),
+					ViewportClient->IsCtrlPressed(), ViewportClient->IsCmdPressed());
+
+				if (bIsLeftMouse)
+				{
+					InputState.Mouse.Left.SetStates(
+						(Event == IE_Pressed), (Event == IE_Pressed), (Event == IE_Released));
+					CurrentMouseState.Mouse.Left.bDown = (Event == IE_Pressed);
+				}
+				else if (bIsMiddleMouse)
+				{
+					InputState.Mouse.Middle.SetStates(
+						(Event == IE_Pressed), (Event == IE_Pressed), (Event == IE_Released));
+					CurrentMouseState.Mouse.Middle.bDown = (Event == IE_Pressed);
+				}
+				else
+				{
+					InputState.Mouse.Right.SetStates(
+						(Event == IE_Pressed), (Event == IE_Pressed), (Event == IE_Released));
+					CurrentMouseState.Mouse.Right.bDown = (Event == IE_Pressed);
+				}
+
+				InputRouter->PostInputEvent(InputState);
+
+				if (InputRouter->HasActiveMouseCapture())
+				{
+					// what is this about? MeshPaintMode has it...
+					ViewportClient->bLockFlightCamera = true;
+					bHandled = true;   // indicate that we handled this event,
+									   // which will disable camera movement/etc ?
+				}
+				else
+				{
+					//ViewportClient->bLockFlightCamera = false;
+				}
+
+			}
+		}
+		else if (Key.IsGamepadKey())
+		{
+			// not supported yet
+		}
+		else if (Key.IsTouch())
+		{
+			// not supported yet
+		}
+		else if (Key.IsFloatAxis() || Key.IsVectorAxis())
+		{
+			// not supported yet
+		}
+		else    // is this definitely a keyboard key?
+		{
+			FInputDeviceState InputState;
+			InputState.InputDevice = EInputDevices::Keyboard;
+			InputState.SetModifierKeyStates(
+				ViewportClient->IsShiftPressed(), ViewportClient->IsAltPressed(),
+				ViewportClient->IsCtrlPressed(), ViewportClient->IsCmdPressed());
+			InputState.Keyboard.ActiveKey.Button = Key;
+			bool bPressed = (Event == IE_Pressed);
+			InputState.Keyboard.ActiveKey.SetStates(bPressed, bPressed, !bPressed);
+			InputRouter->PostInputEvent(InputState);
+		}
+
+	}
+
+	return bHandled;
+}
+
+
+
+
+bool UEdModeInteractiveToolsContext::MouseEnter(FEditorViewportClient* ViewportClient, FViewport* Viewport, int32 x, int32 y)
+{
+#ifdef ENABLE_DEBUG_PRINTING
+	UE_LOG(LogTemp, Warning, TEXT("MOUSE ENTER"));
+#endif
+
+	CurrentMouseState.Mouse.Position2D = FVector2D(x, y);
+	CurrentMouseState.Mouse.WorldRay = GetRayFromMousePos(ViewportClient, Viewport, x, y);
+
+	return false;
+}
+
+
+bool UEdModeInteractiveToolsContext::MouseMove(FEditorViewportClient* ViewportClient, FViewport* Viewport, int32 x, int32 y)
+{
+#ifdef ENABLE_DEBUG_PRINTING
+	//UE_LOG(LogTemp, Warning, TEXT("MOUSE MOVE"));
+#endif
+
+	CurrentMouseState.Mouse.Position2D = FVector2D(x, y);
+	CurrentMouseState.Mouse.WorldRay = GetRayFromMousePos(ViewportClient, Viewport, x, y);
+	FInputDeviceState InputState = CurrentMouseState;
+	InputState.InputDevice = EInputDevices::Mouse;
+
+	InputState.SetModifierKeyStates(
+		ViewportClient->IsShiftPressed(), ViewportClient->IsAltPressed(),
+		ViewportClient->IsCtrlPressed(), ViewportClient->IsCmdPressed());
+
+	if (InputRouter->HasActiveMouseCapture())
+	{
+		// This state occurs if InputBehavior did not release capture on mouse release. 
+		// UMultiClickSequenceInputBehavior does this, eg for multi-click draw-polygon sequences.
+		// It's not ideal though and maybe would be better done via multiple captures + hover...?
+		InputRouter->PostInputEvent(InputState);
+	}
+	else
+	{
+		InputRouter->PostHoverInputEvent(InputState);
+	}
+
+	return false;
+}
+
+
+bool UEdModeInteractiveToolsContext::MouseLeave(FEditorViewportClient* ViewportClient, FViewport* Viewport)
+{
+#ifdef ENABLE_DEBUG_PRINTING
+	UE_LOG(LogTemp, Warning, TEXT("MOUSE LEAVE"));
+#endif
+
+	return false;
+}
+
+
+
+bool UEdModeInteractiveToolsContext::StartTracking(FEditorViewportClient* InViewportClient, FViewport* InViewport)
+{
+	return false;
+}
+
+bool UEdModeInteractiveToolsContext::CapturedMouseMove(FEditorViewportClient* InViewportClient, FViewport* InViewport, int32 InMouseX, int32 InMouseY)
+{
+#ifdef ENABLE_DEBUG_PRINTING
+	//UE_LOG(LogTemp, Warning, TEXT("CAPTURED MOUSE MOVE"));
+#endif
+
+	// if alt is down we will not allow client to see this event
+	if (InViewportClient->IsAltPressed())
+	{
+		return false;
+	}
+
+	FVector2D OldPosition = CurrentMouseState.Mouse.Position2D;
+	CurrentMouseState.Mouse.Position2D = FVector2D(InMouseX, InMouseY);
+	CurrentMouseState.Mouse.WorldRay = GetRayFromMousePos(InViewportClient, InViewport, InMouseX, InMouseY);
+
+	if (InputRouter->HasActiveMouseCapture())
+	{
+		FInputDeviceState InputState = CurrentMouseState;
+		InputState.InputDevice = EInputDevices::Mouse;
+		InputState.SetModifierKeyStates(
+			InViewportClient->IsShiftPressed(), InViewportClient->IsAltPressed(),
+			InViewportClient->IsCtrlPressed(), InViewportClient->IsCmdPressed());
+		InputState.Mouse.Delta2D = CurrentMouseState.Mouse.Position2D - OldPosition;
+		InputRouter->PostInputEvent(InputState);
+		return true;
+	}
+
+	return false;
+}
+
+
+bool UEdModeInteractiveToolsContext::EndTracking(FEditorViewportClient* InViewportClient, FViewport* InViewport)
+{
+#ifdef ENABLE_DEBUG_PRINTING
+	UE_LOG(LogTemp, Warning, TEXT("END TRACKING"));
+#endif
+
+	return true;
+}
+
+
+
+
+
+
+FRay UEdModeInteractiveToolsContext::GetRayFromMousePos(FEditorViewportClient* ViewportClient, FViewport* Viewport, int MouseX, int MouseY)
+{
+	FSceneViewFamilyContext ViewFamily(FSceneViewFamily::ConstructionValues(
+		ViewportClient->Viewport,
+		ViewportClient->GetScene(),
+		ViewportClient->EngineShowFlags)
+		.SetRealtimeUpdate(ViewportClient->IsRealtime()));
+	FSceneView* View = ViewportClient->CalcSceneView(&ViewFamily);
+	FViewportCursorLocation MouseViewportRay(View, (FEditorViewportClient*)Viewport->GetClient(), MouseX, MouseY);
+
+	return FRay(MouseViewportRay.GetOrigin(), MouseViewportRay.GetDirection(), true);
+}
+
+
+
+
+
+
+bool UEdModeInteractiveToolsContext::CanStartTool(const FString& ToolTypeIdentifier) const
+{
+	return (ToolManager->HasActiveTool(EToolSide::Mouse) == false) &&
+		(ToolManager->CanActivateTool(EToolSide::Mouse, ToolTypeIdentifier) == true);
+}
+bool UEdModeInteractiveToolsContext::ActiveToolHasAccept() const
+{
+	return ToolManager->HasActiveTool(EToolSide::Mouse) &&
+		ToolManager->GetActiveTool(EToolSide::Mouse)->HasAccept();
+}
+bool UEdModeInteractiveToolsContext::CanAcceptActiveTool() const
+{
+	return ToolManager->CanAcceptActiveTool(EToolSide::Mouse);
+}
+bool UEdModeInteractiveToolsContext::CanCancelActiveTool() const
+{
+	return ToolManager->CanCancelActiveTool(EToolSide::Mouse);
+}
+
+bool UEdModeInteractiveToolsContext::CanCompleteActiveTool() const
+{
+	return ToolManager->HasActiveTool(EToolSide::Mouse) && CanCancelActiveTool() == false;
+}
+void UEdModeInteractiveToolsContext::StartTool(const FString& ToolTypeIdentifier)
+{
+	if (ToolManager->SelectActiveToolType(EToolSide::Mouse, ToolTypeIdentifier) == false)
+	{
+		UE_LOG(LogTemp, Warning, TEXT("ToolManager: Unknown Tool Type %s"), *ToolTypeIdentifier);
+	}
+	else
+	{
+		ToolManager->ActivateTool(EToolSide::Mouse);
+	}
+}
+void UEdModeInteractiveToolsContext::EndTool(EToolShutdownType ShutdownType)
+{
+	ToolManager->DeactivateTool(EToolSide::Mouse, ShutdownType);
+}
\ No newline at end of file
diff --git a/Engine/Source/Editor/Experimental/EditorInteractiveToolsFramework/Private/EditorComponentSourceFactory.cpp b/Engine/Source/Editor/Experimental/EditorInteractiveToolsFramework/Private/EditorComponentSourceFactory.cpp
new file mode 100644
--- /dev/null
+++ b/Engine/Source/Editor/Experimental/EditorInteractiveToolsFramework/Private/EditorComponentSourceFactory.cpp
@@ -0,0 +1,108 @@
+// Copyright 1998-2019 Epic Games, Inc. All Rights Reserved.
+
+#include "EditorComponentSourceFactory.h"
+
+#include "Engine/StaticMesh.h"
+#include "Components/ActorComponent.h"
+#include "Components/StaticMeshComponent.h"
+#include "PhysicsEngine/BodySetup.h"
+
+
+TUniquePtr<IMeshDescriptionSource> FEditorComponentSourceFactory::MakeMeshDescriptionSource(UActorComponent* Component)
+{
+	UStaticMeshComponent* StaticMeshComp = Cast<UStaticMeshComponent>(Component);
+	if (StaticMeshComp != nullptr)
+	{
+		auto NewSource = new FStaticMeshComponentMeshDescriptionSource(StaticMeshComp, 0);
+		return TUniquePtr<IMeshDescriptionSource>(NewSource);
+	}
+	return nullptr;
+}
+
+
+
+
+
+FStaticMeshComponentMeshDescriptionSource::FStaticMeshComponentMeshDescriptionSource(
+	UStaticMeshComponent* ComponentIn, int LODIndex)
+{
+	this->Component = ComponentIn;
+	this->LODIndex = LODIndex;
+}
+
+
+AActor* FStaticMeshComponentMeshDescriptionSource::GetOwnerActor() const
+{
+	return Component->GetOwner();
+}
+
+UActorComponent* FStaticMeshComponentMeshDescriptionSource::GetOwnerComponent() const
+{
+	return Component;
+}
+
+
+void FStaticMeshComponentMeshDescriptionSource::SetOwnerVisibility(bool bVisible) const
+{
+	Component->SetVisibility(bVisible);
+}
+
+FMeshDescription* FStaticMeshComponentMeshDescriptionSource::GetMeshDescription() const
+{
+	return Component->GetStaticMesh()->GetMeshDescription(0);
+}
+
+UMaterialInterface* FStaticMeshComponentMeshDescriptionSource::GetMaterial(int32 MaterialIndex) const
+{
+	return Component->GetMaterial(MaterialIndex);
+}
+
+FTransform FStaticMeshComponentMeshDescriptionSource::GetWorldTransform() const
+{
+	//return Component->GetOwner()->GetActorTransform();
+	return Component->GetComponentTransform();
+}
+
+
+
+bool FStaticMeshComponentMeshDescriptionSource::IsReadOnly() const
+{
+	return false;
+}
+
+
+void FStaticMeshComponentMeshDescriptionSource::CommitInPlaceModification(const TFunction<void(FMeshDescription*)>& ModifyFunction) 
+{
+	//bool bSaved = Component->Modify();
+	//check(bSaved);
+	UStaticMesh* StaticMesh = Component->GetStaticMesh();
+
+	// make sure transactional flag is on
+	StaticMesh->SetFlags(RF_Transactional);
+
+	bool bSavedToTransactionBuffer = StaticMesh->Modify();
+	check(bSavedToTransactionBuffer);
+	FMeshDescription* MeshDescription = StaticMesh->GetMeshDescription(LODIndex);
+
+	ModifyFunction(MeshDescription);
+
+	StaticMesh->CommitMeshDescription(LODIndex);
+	StaticMesh->PostEditChange();
+
+	// this rebuilds physics, but it doesn't undo!
+	Component->RecreatePhysicsState();
+
+	//Component->PostEditChange();
+}
+
+
+bool FStaticMeshComponentMeshDescriptionSource::HitTest(const FRay& WorldRay, FHitResult& OutHit) const
+{
+	FVector End = WorldRay.PointAt(HALF_WORLD_MAX);
+	if (Component->LineTraceComponent(OutHit, WorldRay.Origin, End, FCollisionQueryParams(SCENE_QUERY_STAT(HitTest), true)))
+	{
+		return true;
+	}
+
+	return false;
+}
\ No newline at end of file
diff --git a/Engine/Source/Editor/Experimental/EditorInteractiveToolsFramework/Private/EditorInteractiveToolsFrameworkModule.cpp b/Engine/Source/Editor/Experimental/EditorInteractiveToolsFramework/Private/EditorInteractiveToolsFrameworkModule.cpp
new file mode 100644
--- /dev/null
+++ b/Engine/Source/Editor/Experimental/EditorInteractiveToolsFramework/Private/EditorInteractiveToolsFrameworkModule.cpp
@@ -0,0 +1,20 @@
+// Copyright 1998-2019 Epic Games, Inc. All Rights Reserved.
+
+#include "EditorInteractiveToolsFrameworkModule.h"
+
+#define LOCTEXT_NAMESPACE "FEditorInteractiveToolsFrameworkModule"
+
+void FEditorInteractiveToolsFrameworkModule::StartupModule()
+{
+	// This code will execute after your module is loaded into memory; the exact timing is specified in the .uplugin file per-module
+}
+
+void FEditorInteractiveToolsFrameworkModule::ShutdownModule()
+{
+	// This function may be called during shutdown to clean up your module.  For modules that support dynamic reloading,
+	// we call this function before unloading the module.
+}
+
+#undef LOCTEXT_NAMESPACE
+	
+IMPLEMENT_MODULE(FEditorInteractiveToolsFrameworkModule, EditorInteractiveToolsFramework)
\ No newline at end of file
diff --git a/Engine/Source/Editor/Experimental/EditorInteractiveToolsFramework/Private/EditorToolAssetAPI.cpp b/Engine/Source/Editor/Experimental/EditorInteractiveToolsFramework/Private/EditorToolAssetAPI.cpp
new file mode 100644
--- /dev/null
+++ b/Engine/Source/Editor/Experimental/EditorInteractiveToolsFramework/Private/EditorToolAssetAPI.cpp
@@ -0,0 +1,93 @@
+// Copyright 1998-2019 Epic Games, Inc. All Rights Reserved.
+
+#include "EditorToolAssetAPI.h"
+#include "CoreMinimal.h"
+#include "GameFramework/Actor.h"
+
+
+#include "AssetRegistryModule.h"
+#include "FileHelpers.h"
+#include "PackageTools.h"
+#include "ObjectTools.h"
+
+// for content-browser things
+#include "ContentBrowserModule.h"
+#include "IContentBrowserSingleton.h"
+
+
+FString FEditorToolAssetAPI::GetActiveAssetFolderPath()
+{
+	check(false);  	// [RMS] this only works if we have an asset selected!
+
+	IContentBrowserSingleton& ContentBrowser = FModuleManager::LoadModuleChecked<FContentBrowserModule>("ContentBrowser").Get();
+
+	//ContentBrowser.CreatePathPicker()
+
+	TArray<FAssetData> SelectedAssets;
+	ContentBrowser.GetSelectedAssets(SelectedAssets);
+	return SelectedAssets[0].PackagePath.ToString();
+}
+
+
+
+FString FEditorToolAssetAPI::MakePackageName(const FString& AssetName, const FString& FolderPath)
+{
+	return FolderPath + TEXT("/") + AssetName;
+}
+
+
+FString FEditorToolAssetAPI::MakeUniqueAssetName(const FString& AssetName, const FString& FolderPath)
+{
+	FString NewPackageName = MakePackageName(AssetName, FolderPath);
+	if (FPackageName::DoesPackageExist(NewPackageName) == false)
+	{
+		return AssetName;
+	}
+
+	int Counter = 1;
+	FString UniqueName = AssetName;
+	while ( FPackageName::DoesPackageExist(NewPackageName) )
+	{
+		UniqueName = AssetName + FString::Printf(TEXT("_%d"), Counter++);
+		NewPackageName = MakePackageName(UniqueName, FolderPath);
+	}
+	return UniqueName;
+}
+
+
+UPackage* FEditorToolAssetAPI::CreateNewPackage(const FString& AssetName, const FString& FolderPath)
+{
+	FString NewPackageName = MakePackageName(AssetName, FolderPath);
+	NewPackageName = UPackageTools::SanitizePackageName(NewPackageName);
+	UPackage* AssetPackage = CreatePackage(nullptr, *NewPackageName);
+	return AssetPackage;
+}
+
+
+
+void FEditorToolAssetAPI::InteractiveSaveGeneratedAsset(UObject* Asset, UPackage* AssetPackage)
+{
+	Asset->MarkPackageDirty();
+
+	FAssetRegistryModule::AssetCreated(Asset);
+
+	TArray<UPackage*> PackagesToSave;
+	PackagesToSave.Add(AssetPackage);
+	bool bCheckDirty = true;
+	bool bPromptToSave = true;
+	FEditorFileUtils::PromptForCheckoutAndSave(PackagesToSave, bCheckDirty, bPromptToSave);
+}
+
+
+void FEditorToolAssetAPI::AutoSaveGeneratedAsset(UObject* Asset, UPackage* AssetPackage)
+{
+	Asset->MarkPackageDirty();
+
+	FAssetRegistryModule::AssetCreated(Asset);
+
+	TArray<UPackage*> PackagesToSave;
+	PackagesToSave.Add(AssetPackage);
+	bool bCheckDirty = true;
+	bool bPromptToSave = false;
+	FEditorFileUtils::PromptForCheckoutAndSave(PackagesToSave, bCheckDirty, bPromptToSave);
+}
\ No newline at end of file
diff --git a/Engine/Source/Editor/Experimental/EditorInteractiveToolsFramework/Private/StandardToolModeCommands.cpp b/Engine/Source/Editor/Experimental/EditorInteractiveToolsFramework/Private/StandardToolModeCommands.cpp
new file mode 100644
--- /dev/null
+++ b/Engine/Source/Editor/Experimental/EditorInteractiveToolsFramework/Private/StandardToolModeCommands.cpp
@@ -0,0 +1,38 @@
+// Copyright 1998-2019 Epic Games, Inc. All Rights Reserved.
+
+#include "StandardToolModeCommands.h"
+#include "EditorStyleSet.h"
+
+#define LOCTEXT_NAMESPACE "StandardToolModeCommands"
+
+
+FStandardToolModeCommands::FStandardToolModeCommands() :
+	TCommands<FStandardToolModeCommands>(
+		"StandardToolCommands", // Context name for fast lookup
+		NSLOCTEXT("Contexts", "StandardToolCommands", "Standard Tool Commands"), // Localized context name for displaying
+		NAME_None, // Parent
+		FEditorStyle::GetStyleSetName() // Icon Style Set
+		)
+{}
+
+
+void FStandardToolModeCommands::RegisterCommands()
+{
+	TSharedPtr<FUICommandInfo> IncreaseBrushSize;
+	UI_COMMAND(IncreaseBrushSize, "Increase Brush Size", "Increases the size of the brush", EUserInterfaceActionType::Button, FInputChord(EKeys::RightBracket));
+	Commands.Add(EStandardToolModeCommands::IncreaseBrushSize, IncreaseBrushSize);
+
+	TSharedPtr<FUICommandInfo> DecreaseBrushSize;
+	UI_COMMAND(DecreaseBrushSize, "Decrease Brush Size", "Decreases the size of the brush", EUserInterfaceActionType::Button, FInputChord(EKeys::LeftBracket));
+	Commands.Add(EStandardToolModeCommands::DecreaseBrushSize, DecreaseBrushSize);
+}
+
+
+TSharedPtr<FUICommandInfo> FStandardToolModeCommands::FindStandardCommand(EStandardToolModeCommands Command) const
+{
+	const TSharedPtr<FUICommandInfo>* Found = Commands.Find(Command);
+	checkf(Found != nullptr, TEXT("FStandardToolModeCommands::FindStandardCommand: standard command %d was not found!"), Command);
+	return (Found != nullptr) ? *Found : nullptr;
+}
+
+
diff --git a/Engine/Source/Editor/Experimental/EditorInteractiveToolsFramework/Public/EdModeInteractiveToolsContext.h b/Engine/Source/Editor/Experimental/EditorInteractiveToolsFramework/Public/EdModeInteractiveToolsContext.h
new file mode 100644
--- /dev/null
+++ b/Engine/Source/Editor/Experimental/EditorInteractiveToolsFramework/Public/EdModeInteractiveToolsContext.h
@@ -0,0 +1,99 @@
+// Copyright 1998-2019 Epic Games, Inc. All Rights Reserved.
+
+#pragma once
+
+#include "CoreMinimal.h"
+#include "EdMode.h"
+#include "InteractiveToolsContext.h"
+#include "EdModeInteractiveToolsContext.generated.h"
+
+/**
+ * EdModeInteractiveToolsContext is an extension/adapter of an InteractiveToolsContext which 
+ * allows it to be easily embedded inside an FEdMode. A set of functions are provided which can be
+ * called from the FEdMode functions of the same name. These will handle the data type
+ * conversions and forwarding calls necessary to operate the ToolsContext
+ */
+UCLASS(Transient)
+class EDITORINTERACTIVETOOLSFRAMEWORK_API UEdModeInteractiveToolsContext : public UInteractiveToolsContext
+{
+	GENERATED_BODY()
+
+public:
+	UEdModeInteractiveToolsContext();
+
+
+	virtual void InitializeContextFromEdMode(FEdMode* EditorMode);
+	virtual void ShutdownContext();
+
+	// default behavior is to accept active tool
+	virtual void TerminateActiveToolsOnPIEStart();
+
+	// default behavior is to accept active tool
+	virtual void TerminateActiveToolsOnSaveWorld();
+
+	IToolsContextQueriesAPI* GetQueriesAPI() const { return QueriesAPI; }
+	IToolsContextTransactionsAPI* GetTransactionAPI() const { return TransactionAPI; }
+	IToolsContextAssetAPI* GetAssetAPI() const { return AssetAPI; }
+	IComponentSourceFactory* GetComponentSourceFactory() const { return SourceFactory; }
+
+	virtual void PostInvalidation();
+
+
+	// call these from your FEdMode functions of the same name
+
+	virtual void Tick(FEditorViewportClient* ViewportClient, float DeltaTime);
+	virtual void Render(const FSceneView* View, FViewport* Viewport, FPrimitiveDrawInterface* PDI);
+
+	virtual bool InputKey(FEditorViewportClient* ViewportClient, FViewport* Viewport, FKey Key, EInputEvent Event);
+
+	virtual bool MouseEnter(FEditorViewportClient* ViewportClient, FViewport* Viewport, int32 x, int32 y);
+	virtual bool MouseLeave(FEditorViewportClient* ViewportClient, FViewport* Viewport);
+	virtual bool MouseMove(FEditorViewportClient* ViewportClient, FViewport* Viewport, int32 x, int32 y);
+
+	virtual bool StartTracking(FEditorViewportClient* InViewportClient, FViewport* InViewport);
+	virtual bool CapturedMouseMove(FEditorViewportClient* InViewportClient, FViewport* InViewport, int32 InMouseX, int32 InMouseY);
+	virtual bool EndTracking(FEditorViewportClient* InViewportClient, FViewport* InViewport);
+
+
+	//
+	// Utility functions useful for hooking up to UICommand/etc
+	//
+
+	virtual bool CanStartTool(const FString& ToolTypeIdentifier) const;
+	virtual bool ActiveToolHasAccept() const;
+	virtual bool CanAcceptActiveTool() const;
+	virtual bool CanCancelActiveTool() const;
+	virtual bool CanCompleteActiveTool() const;
+	virtual void StartTool(const FString& ToolTypeIdentifier);
+	virtual void EndTool(EToolShutdownType ShutdownType);
+
+
+protected:
+	// we hide these 
+	virtual void Initialize(IToolsContextQueriesAPI* QueriesAPI, IToolsContextTransactionsAPI* TransactionsAPI) override;
+	virtual void Shutdown() override;
+
+
+public:
+	UPROPERTY()
+	UMaterialInterface* StandardVertexColorMaterial;
+
+protected:
+	FEdMode* EditorMode;
+
+	FDelegateHandle BeginPIEDelegateHandle;
+	FDelegateHandle PreSaveWorldDelegateHandle;
+
+	IToolsContextQueriesAPI* QueriesAPI;
+	IToolsContextTransactionsAPI* TransactionAPI;
+	IToolsContextAssetAPI* AssetAPI;
+	IComponentSourceFactory* SourceFactory;
+
+	bool bInvalidationPending;
+
+	/** Input event instance used to keep track of various button states, etc, that we cannot directly query on-demand */
+	FInputDeviceState CurrentMouseState;
+
+	static FRay GetRayFromMousePos(FEditorViewportClient* ViewportClient, FViewport* Viewport, int MouseX, int MouseY);
+
+};
diff --git a/Engine/Source/Editor/Experimental/EditorInteractiveToolsFramework/Public/EditorComponentSourceFactory.h b/Engine/Source/Editor/Experimental/EditorInteractiveToolsFramework/Public/EditorComponentSourceFactory.h
new file mode 100644
--- /dev/null
+++ b/Engine/Source/Editor/Experimental/EditorInteractiveToolsFramework/Public/EditorComponentSourceFactory.h
@@ -0,0 +1,54 @@
+// Copyright 1998-2019 Epic Games, Inc. All Rights Reserved.
+
+#pragma once
+
+#include "CoreMinimal.h"
+#include "ComponentSourceInterfaces.h"
+
+
+
+
+/**
+ * Component Source Factory for use in the Editor (ie inside FEdMode, etc)
+ */
+class EDITORINTERACTIVETOOLSFRAMEWORK_API FEditorComponentSourceFactory : public IComponentSourceFactory
+{
+public:
+	virtual ~FEditorComponentSourceFactory() {}
+
+	virtual TUniquePtr<IMeshDescriptionSource> MakeMeshDescriptionSource(UActorComponent* Component);
+};
+
+
+
+
+
+class UStaticMeshComponent;
+
+/**
+ * This MeshDescriptionSource provides a specific LOD from a StaticMeshComponent
+ */
+class EDITORINTERACTIVETOOLSFRAMEWORK_API FStaticMeshComponentMeshDescriptionSource : public IMeshDescriptionSource
+{
+public:
+	UStaticMeshComponent* Component;
+	int LODIndex;
+
+	FStaticMeshComponentMeshDescriptionSource(
+		UStaticMeshComponent* ComponentIn,
+		int LODIndex = 0);
+	virtual ~FStaticMeshComponentMeshDescriptionSource() {}
+
+	virtual AActor* GetOwnerActor() const override;
+	virtual UActorComponent* GetOwnerComponent() const override;
+	virtual FMeshDescription* GetMeshDescription() const override;
+	virtual UMaterialInterface* GetMaterial(int32 MaterialIndex) const override;
+	virtual FTransform GetWorldTransform() const override;
+	virtual bool HitTest(const FRay& WorldRay, FHitResult& OutHit) const override;
+
+	virtual bool IsReadOnly() const override;
+
+	virtual void SetOwnerVisibility(bool bVisible) const override;
+	virtual void CommitInPlaceModification(const TFunction<void(FMeshDescription*)>& ModifyFunction) override;
+};
+
diff --git a/Engine/Source/Editor/Experimental/EditorInteractiveToolsFramework/Public/EditorInteractiveToolsFrameworkModule.h b/Engine/Source/Editor/Experimental/EditorInteractiveToolsFramework/Public/EditorInteractiveToolsFrameworkModule.h
new file mode 100644
--- /dev/null
+++ b/Engine/Source/Editor/Experimental/EditorInteractiveToolsFramework/Public/EditorInteractiveToolsFrameworkModule.h
@@ -0,0 +1,15 @@
+// Copyright 1998-2019 Epic Games, Inc. All Rights Reserved.
+
+#pragma once
+
+#include "CoreMinimal.h"
+#include "Modules/ModuleManager.h"
+
+class FEditorInteractiveToolsFrameworkModule : public IModuleInterface
+{
+public:
+
+	/** IModuleInterface implementation */
+	virtual void StartupModule() override;
+	virtual void ShutdownModule() override;
+};
diff --git a/Engine/Source/Editor/Experimental/EditorInteractiveToolsFramework/Public/EditorToolAssetAPI.h b/Engine/Source/Editor/Experimental/EditorInteractiveToolsFramework/Public/EditorToolAssetAPI.h
new file mode 100644
--- /dev/null
+++ b/Engine/Source/Editor/Experimental/EditorInteractiveToolsFramework/Public/EditorToolAssetAPI.h
@@ -0,0 +1,28 @@
+// Copyright 1998-2019 Epic Games, Inc. All Rights Reserved.
+
+#pragma once
+
+#include "CoreMinimal.h"
+#include "ToolContextInterfaces.h"
+
+/**
+ * Implementation of ToolsContext Asset management API that is suitable for use
+ * inside the Editor (eg inside an FEdMode)
+ */
+class EDITORINTERACTIVETOOLSFRAMEWORK_API FEditorToolAssetAPI : public IToolsContextAssetAPI
+{
+public:
+	virtual FString GetActiveAssetFolderPath() override;
+
+	virtual FString MakePackageName(const FString& AssetName, const FString& FolderPath) override;
+
+	virtual FString MakeUniqueAssetName(const FString& AssetName, const FString& FolderPath) override;
+
+	// FolderPath should neither start with, nor end with, path separators
+	// Generated path will be /Game/FolderPath/AssetName
+	virtual UPackage* CreateNewPackage(const FString& AssetName, const FString& FolderPath) override;
+
+	virtual void InteractiveSaveGeneratedAsset(UObject* Asset, UPackage* AssetPackage) override;
+
+	virtual void AutoSaveGeneratedAsset(UObject* Asset, UPackage* AssetPackage) override;
+};
\ No newline at end of file
diff --git a/Engine/Source/Editor/Experimental/EditorInteractiveToolsFramework/Public/InteractiveToolsCommands.h b/Engine/Source/Editor/Experimental/EditorInteractiveToolsFramework/Public/InteractiveToolsCommands.h
new file mode 100644
--- /dev/null
+++ b/Engine/Source/Editor/Experimental/EditorInteractiveToolsFramework/Public/InteractiveToolsCommands.h
@@ -0,0 +1,302 @@
+// Copyright 1998-2019 Epic Games, Inc. All Rights Reserved.
+
+#pragma once
+
+#include "CoreMinimal.h"
+#include "Framework/Commands/Commands.h"
+#include "Framework/Commands/UICommandList.h"
+#include "InteractiveTool.h"
+#include "StandardToolModeCommands.h"
+
+
+
+/**
+ * TInteractiveToolCommands is a base class that handles connecting up Tool Actions 
+ * (ie the FInteractiveToolAction provided by a UInteractiveTool) to the UnrealEditor
+ * Command system, which allows for remappable hotkeys, etc
+ * 
+ * Usage is as follows:
+ *    - in your EdMode Module, create a subclass-instance of this, say FMyToolCommands and call FMyToolCommands::Register() in your ::StartupModule() function
+ *        - subclass must implement ::GetToolDefaultObjectList(), here you just add GetMutableDefault<MyToolX> to the input list for all your Tools
+ *        - subclass must implement ::GetNamespaceString(), likely you can just return TEXT(LOCTEXT_NAMESPACE)
+ *    - add a member TSharedPtr<FUICommandList> UICommandList; to your EdMode impl
+ *    - when you start a new Tool, call FMyToolCommands::Get().BindCommandsForCurrentTool(UICommandList, NewTool)
+ *    - when you end a Tool, call FMyToolCommands::Get().UnbindActiveCommands()
+ *    - in your EdModeImpl::InputKey override, call UICommandList->ProcessCommandBindings()
+ */
+template<typename CommandContextType>
+class TInteractiveToolCommands : public TCommands<CommandContextType>
+{
+
+public:
+	/**
+	 * Initialize commands. Collects possible Actions from the available UInteractiveTools
+	 * (provided by GetToolDefaultObjectList) and then registers a FUICommandInfo for each.
+	 * This needs to be called in ::StartupModule for your EdMode Module
+	 */
+	virtual void RegisterCommands() override;
+
+	/**
+	 * Bind any of the registered UICommands to the given Tool, if they are compatible.
+	 */
+	virtual void BindCommandsForCurrentTool(TSharedPtr<FUICommandList> UICommandList, UInteractiveTool* Tool) const;
+
+	/**
+	 * Unbind all of the currently-registered commands
+	 */
+	virtual void UnbindActiveCommands(TSharedPtr<FUICommandList> UICommandList) const;
+
+
+	//
+	// Interface that subclasses need to implement
+	// 
+
+
+	/**
+	 * RegisterCommands() needs actual UInteractiveTool instances for all the Tools that want to
+	 * provide Actions which will be connected to hotkeys. We can do this based on the CDO's for
+	 * each tool, returned by GetMutableDefault<UYourInteractiveToolType>(). The Tool CDOs are
+	 * not owned by a ToolManager and we will only call .GetActionSet() on them.
+	 */
+	virtual void GetToolDefaultObjectList(TArray<UInteractiveTool*>& ToolCDOs) = 0;
+
+	/**
+	 * The Namespace used to identify this Command set
+	 */
+	virtual const TCHAR * GetNamespaceString() = 0;
+
+
+protected:
+
+	/** Forwarding constructor */
+	TInteractiveToolCommands(const FName InContextName, const FText& InContextDesc, const FName InContextParent, const FName InStyleSetName)
+		: TCommands<CommandContextType>(InContextName, InContextDesc, InContextParent, InStyleSetName)
+	{
+	}
+
+	/**
+	 * Query FStandardToolModeCommands to find an existing command/hotkey for this standard tool action
+	 */
+	virtual TSharedPtr<FUICommandInfo> FindStandardCommand(EStandardToolActions ToolActionID);
+
+	/**
+	 * List of pairs of known Tool Actions and their associated UICommands.
+	 * This list is used at bind-time to connect up the already-registered UICommands to the active Tool
+	 */
+	struct FToolActionCommand
+	{
+		FInteractiveToolAction ToolAction;
+		TSharedPtr<FUICommandInfo> UICommand;
+	};
+	TArray<FToolActionCommand> ActionCommands;
+
+
+	//
+	// Support for sharing common commands between Tools where the commands
+	// are not shared across Tool Modes (and hence not in FStandardToolModeCommands)
+	// 
+
+	/** List FToolActionCommands for standard Tool Actions */
+	TMap<EStandardToolActions, FToolActionCommand> SharedActionCommands;
+
+	/**
+	 * Find or Create a UICommand for a standard Tool Action, that will be shared across Tools
+	 */
+	bool FindOrCreateSharedCommand(EStandardToolActions ActionID, FToolActionCommand& FoundOut);
+
+	/**
+	 * Create a suitable FInteractiveToolAction for the standard Tool Action, ie with suitable
+	 * command name and description strings and hotkeys. This info will be used to create the shared UICommand.
+	 */
+	virtual bool IntializeStandardToolAction(EStandardToolActions ActionID, FInteractiveToolAction& ToolActionOut);
+
+	/**
+	 * Utility function that registeres a Tool Aciton as a UICommand
+	 */
+	void RegisterUIToolCommand(const FInteractiveToolAction& ToolAction, TSharedPtr<FUICommandInfo>& UICommandInfo);
+};
+
+
+
+
+
+
+template<typename CommandContextType>
+void TInteractiveToolCommands<CommandContextType>::RegisterCommands()
+{
+	// make sure standard commands are registered
+	FStandardToolModeCommands::Register();
+
+	// get list of all tools used by command set
+	TArray<UInteractiveTool*> AllToolsCDOs;
+	GetToolDefaultObjectList(AllToolsCDOs);
+
+	// collect all the actions in these tools
+	TArray<FInteractiveToolAction> ToolActions;
+	for (UInteractiveTool* Tool : AllToolsCDOs)
+	{
+		Tool->GetActionSet()->CollectActions(ToolActions);
+	}
+
+	// register all the commands in all the actions
+	int NumActions = ToolActions.Num();
+	ActionCommands.SetNum(NumActions);
+	for (int k = 0; k < NumActions; ++k)
+	{
+		const FInteractiveToolAction& ToolAction = ToolActions[k];
+		ActionCommands[k].ToolAction = ToolAction;
+
+		bool bRegistered = false;
+		if ((int32)ToolAction.ActionID < (int32)EStandardToolActions::BaseClientDefinedActionID)
+		{
+			TSharedPtr<FUICommandInfo> FoundStandard = FindStandardCommand((EStandardToolActions)ToolAction.ActionID);
+			if (FoundStandard.IsValid())
+			{
+				ActionCommands[k].UICommand = FoundStandard;
+				bRegistered = true;
+			}
+			else
+			{
+				FToolActionCommand StandardActionCommand;
+				bool bFoundRemap = FindOrCreateSharedCommand((EStandardToolActions)ToolAction.ActionID, StandardActionCommand);
+				if (bFoundRemap)
+				{
+					ActionCommands[k].UICommand = StandardActionCommand.UICommand;
+					bRegistered = true;
+				}
+			}
+		}
+
+		if (bRegistered == false)
+		{
+			RegisterUIToolCommand(ActionCommands[k].ToolAction, ActionCommands[k].UICommand);
+		}
+	}
+}
+
+
+template<typename CommandContextType>
+void TInteractiveToolCommands<CommandContextType>::BindCommandsForCurrentTool(TSharedPtr<FUICommandList> UICommandList, UInteractiveTool* Tool) const
+{
+	UClass* ToolClassType = Tool->GetClass();
+
+	int NumActionCommands = ActionCommands.Num();
+	for (int32 k = 0; k < NumActionCommands; ++k)
+	{
+		const FToolActionCommand& ActionCommand = ActionCommands[k];
+		if (ActionCommand.ToolAction.ClassType == ToolClassType)
+		{
+			UICommandList->MapAction(
+				ActionCommand.UICommand,
+				FExecuteAction::CreateUObject(Tool, &UInteractiveTool::ExecuteAction, ActionCommand.ToolAction.ActionID));
+			//FCanExecuteAction::CreateRaw(this, &FEdModeFoliage::CurrentToolUsesBrush));
+		}
+	}
+}
+
+
+template<typename CommandContextType>
+void TInteractiveToolCommands<CommandContextType>::UnbindActiveCommands(TSharedPtr<FUICommandList> UICommandList) const
+{
+	// @todo would be more efficient if we kept track of which commands were mapped.
+	// However currently this function must be const because TCommands::Get() returns a const
+	for (const FToolActionCommand& ActionCommand : ActionCommands)
+	{
+		if (UICommandList->IsActionMapped(ActionCommand.UICommand))
+		{
+			UICommandList->UnmapAction(ActionCommand.UICommand);
+		}
+	}
+}
+
+
+
+
+
+
+//
+// Support for re-using known commands from FStandardToolModeCommands
+//
+
+template<typename CommandContextType>
+TSharedPtr<FUICommandInfo> TInteractiveToolCommands<CommandContextType>::FindStandardCommand(EStandardToolActions ToolActionID)
+{
+	const FStandardToolModeCommands& StdCommands = FStandardToolModeCommands::Get();
+	switch (ToolActionID)
+	{
+	case EStandardToolActions::IncreaseBrushSize:
+		return StdCommands.FindStandardCommand(EStandardToolModeCommands::IncreaseBrushSize);
+	case EStandardToolActions::DecreaseBrushSize:
+		return StdCommands.FindStandardCommand(EStandardToolModeCommands::DecreaseBrushSize);
+	default:
+		return nullptr;
+	}
+}
+
+
+
+
+
+
+//
+// Support for sharing UI commands between multiple Tools, where the
+// shared command is not at the shared-between-modes level
+// 
+
+template<typename CommandContextType>
+bool TInteractiveToolCommands<CommandContextType>::FindOrCreateSharedCommand(EStandardToolActions ActionID, FToolActionCommand& FoundOut)
+{
+	if (SharedActionCommands.Contains(ActionID))
+	{
+		FoundOut = SharedActionCommands[ActionID];
+		return true;
+	}
+
+	FInteractiveToolAction NewAction;
+	bool bIsKnownAction = IntializeStandardToolAction(ActionID, NewAction);
+
+	if (bIsKnownAction)
+	{
+		check(NewAction.ActionID == (int32)ActionID);
+		FToolActionCommand NewActionCommand;
+		NewActionCommand.ToolAction = NewAction;
+		RegisterUIToolCommand(NewActionCommand.ToolAction, NewActionCommand.UICommand);
+		FoundOut = SharedActionCommands.Add(ActionID, NewActionCommand);
+		return true;
+	}
+	return false;
+}
+
+
+template<typename CommandContextType>
+bool TInteractiveToolCommands<CommandContextType>::IntializeStandardToolAction(EStandardToolActions ActionID, FInteractiveToolAction& ToolActionOut)
+{
+	// base class does not have any standard Tool actions. Subclasses can override
+	// this function, populating ToolActionOut and returning true, if the ActionID should be shared between multiple Tools
+	
+	return false;
+}
+
+
+
+
+template<typename CommandContextType>
+void TInteractiveToolCommands<CommandContextType>::RegisterUIToolCommand(const FInteractiveToolAction& ToolAction, TSharedPtr<FUICommandInfo>& UICommandInfo)
+{
+	FString TooltipString = ToolAction.ActionName + FString("_ToolTip");
+	FString DotString = FString(".") + ToolAction.ActionName;
+
+	UI_COMMAND_Function(this,
+		UICommandInfo,
+		GetNamespaceString(),
+		*ToolAction.ActionName,
+		*TooltipString,
+		TCHAR_TO_ANSI(*DotString),
+		*ToolAction.ShortName, *ToolAction.Description,
+		EUserInterfaceActionType::Button,
+		FInputChord(ToolAction.DefaultModifiers, ToolAction.DefaultKey));
+}
+
+
+
+
diff --git a/Engine/Source/Editor/Experimental/EditorInteractiveToolsFramework/Public/StandardToolModeCommands.h b/Engine/Source/Editor/Experimental/EditorInteractiveToolsFramework/Public/StandardToolModeCommands.h
new file mode 100644
--- /dev/null
+++ b/Engine/Source/Editor/Experimental/EditorInteractiveToolsFramework/Public/StandardToolModeCommands.h
@@ -0,0 +1,54 @@
+// Copyright 1998-2019 Epic Games, Inc. All Rights Reserved.
+
+#pragma once
+
+#include "CoreMinimal.h"
+#include "Framework/Commands/Commands.h"
+#include "Framework/Commands/UICommandInfo.h"
+
+
+/**
+ * Integer identifiers for "Standard" Tool commands, that many EditorModes / AssetEditors / etc may share.
+ * This allows a single hotkey binding to be used across many contexts.
+ */
+enum class EStandardToolModeCommands
+{
+	IncreaseBrushSize = 100,
+	DecreaseBrushSize = 101
+};
+
+
+
+/**
+ * FStandardToolModeCommands provides standard commands for Tools. This allows
+ * for "sharing" common hotkey bindings between multiple EditorModes
+ * 
+ * The set of standard commands is defined by EStandardToolModeCommands.
+ * You must call FStandardToolModeCommands::Register() on module startup to register these commands
+ * (note that they may already be registered by another module, but you should call this to be safe).
+ * 
+ * Then, when you want to Bind a standard command in your FUICommandList,
+ * call FStandardToolModeCommands::Get().FindStandardCommand() to get the registered UICommandInfo.
+ * 
+ */
+class EDITORINTERACTIVETOOLSFRAMEWORK_API FStandardToolModeCommands : public TCommands<FStandardToolModeCommands>
+{
+public:
+	FStandardToolModeCommands();
+
+	/**
+	 * Registers the set of standard commands. Call on module startup.
+	 */
+	virtual void RegisterCommands() override;
+
+	/**
+	 * Look up the UICommandInfo for a standard command
+	 */
+	TSharedPtr<FUICommandInfo> FindStandardCommand(EStandardToolModeCommands Command) const;
+
+protected:
+	TMap<EStandardToolModeCommands, TSharedPtr<FUICommandInfo>> Commands;
+
+};
+
+
diff --git a/Engine/Source/Runtime/Experimental/InteractiveToolsFramework/InteractiveToolsFramework.Build.cs b/Engine/Source/Runtime/Experimental/InteractiveToolsFramework/InteractiveToolsFramework.Build.cs
new file mode 100644
--- /dev/null
+++ b/Engine/Source/Runtime/Experimental/InteractiveToolsFramework/InteractiveToolsFramework.Build.cs
@@ -0,0 +1,55 @@
+// Copyright 1998-2019 Epic Games, Inc. All Rights Reserved.
+
+using UnrealBuildTool;
+
+public class InteractiveToolsFramework : ModuleRules
+{
+	public InteractiveToolsFramework(ReadOnlyTargetRules Target) : base(Target)
+	{
+        PCHUsage = ModuleRules.PCHUsageMode.UseExplicitOrSharedPCHs;
+		
+		PublicIncludePaths.AddRange(
+			new string[] {
+				// ... add public include paths required here ...
+			}
+			);
+				
+		
+		PrivateIncludePaths.AddRange(
+			new string[] {
+				// ... add other private include paths required here ...
+			}
+			);
+			
+		
+		PublicDependencyModuleNames.AddRange(
+			new string[]
+			{
+                "Core",
+                "CoreUObject",
+                "InputCore",
+                "ApplicationCore",
+				// ... add other public dependencies that you statically link with here ...
+			}
+            );
+			
+		
+		PrivateDependencyModuleNames.AddRange(
+			new string[]
+			{
+				"Engine"
+				//"Slate",
+				//"SlateCore",
+				// ... add private dependencies that you statically link with here ...	
+			}
+			);
+		
+		
+		DynamicallyLoadedModuleNames.AddRange(
+			new string[]
+			{
+				// ... add any modules that your module loads dynamically here ...
+			}
+			);
+	}
+}
diff --git a/Engine/Source/Runtime/Experimental/InteractiveToolsFramework/Private/BaseBehaviors/AnyButtonInputBehavior.cpp b/Engine/Source/Runtime/Experimental/InteractiveToolsFramework/Private/BaseBehaviors/AnyButtonInputBehavior.cpp
new file mode 100644
--- /dev/null
+++ b/Engine/Source/Runtime/Experimental/InteractiveToolsFramework/Private/BaseBehaviors/AnyButtonInputBehavior.cpp
@@ -0,0 +1,112 @@
+// Copyright 1998-2019 Epic Games, Inc. All Rights Reserved.
+
+#include "BaseBehaviors/AnyButtonInputBehavior.h"
+
+
+
+UAnyButtonInputBehavior::UAnyButtonInputBehavior()
+{
+	ButtonNumber = 0;
+}
+
+
+EInputDevices UAnyButtonInputBehavior::GetSupportedDevices()
+{
+	return EInputDevices::Mouse;
+}
+
+
+bool UAnyButtonInputBehavior::IsPressed(const FInputDeviceState& input)
+{
+	if (input.IsFromDevice(EInputDevices::Mouse)) 
+	{
+		ActiveDevice = EInputDevices::Mouse;
+		return GetMouseButtonState(input).bPressed;
+	} 
+	else if (input.IsFromDevice(EInputDevices::TabletFingers))
+	{
+		ActiveDevice = EInputDevices::TabletFingers;
+		//return input.TouchPressed;  // not implemented yet
+		return false;
+	}
+	return false;
+}
+
+bool UAnyButtonInputBehavior::IsDown(const FInputDeviceState& input)
+{
+	if (input.IsFromDevice(EInputDevices::Mouse))
+	{
+		ActiveDevice = EInputDevices::Mouse;
+		return GetMouseButtonState(input).bDown;
+	}
+	return false;
+}
+
+bool UAnyButtonInputBehavior::IsReleased(const FInputDeviceState& input)
+{
+	if (input.IsFromDevice(EInputDevices::Mouse))
+	{
+		ActiveDevice = EInputDevices::Mouse;
+		return GetMouseButtonState(input).bReleased;
+	}
+	return false;
+}
+
+FVector2D UAnyButtonInputBehavior::GetClickPoint(const FInputDeviceState& input)
+{
+	if (input.IsFromDevice(EInputDevices::Mouse))
+	{
+		ActiveDevice = EInputDevices::Mouse;
+		return input.Mouse.Position2D;
+	}
+	return FVector2D::ZeroVector;
+}
+
+
+FRay UAnyButtonInputBehavior::GetWorldRay(const FInputDeviceState& input)
+{
+	if (input.IsFromDevice(EInputDevices::Mouse))
+	{
+		ActiveDevice = EInputDevices::Mouse;
+		return input.Mouse.WorldRay;
+	}
+	return FRay(FVector::ZeroVector, FVector(0, 0, 1), true);
+
+}
+
+
+FInputDeviceRay UAnyButtonInputBehavior::GetDeviceRay(const FInputDeviceState& input)
+{
+	if (input.IsFromDevice(EInputDevices::Mouse))
+	{
+		ActiveDevice = EInputDevices::Mouse;
+		return FInputDeviceRay(input.Mouse.WorldRay, input.Mouse.Position2D);
+	}
+	return FInputDeviceRay(FRay(FVector::ZeroVector, FVector(0, 0, 1), true));
+
+}
+
+
+EInputDevices UAnyButtonInputBehavior::GetActiveDevice() const
+{
+	return ActiveDevice;
+}
+
+
+
+
+FDeviceButtonState UAnyButtonInputBehavior::GetMouseButtonState(const FInputDeviceState& input)
+{
+	if (ButtonNumber == 2)
+	{
+		return input.Mouse.Right;
+	}
+	else if (ButtonNumber == 1)
+	{
+		return input.Mouse.Middle;
+	}
+	else
+	{
+		return input.Mouse.Left;
+	}
+}
\ No newline at end of file
diff --git a/Engine/Source/Runtime/Experimental/InteractiveToolsFramework/Private/BaseBehaviors/ClickDragBehavior.cpp b/Engine/Source/Runtime/Experimental/InteractiveToolsFramework/Private/BaseBehaviors/ClickDragBehavior.cpp
new file mode 100644
--- /dev/null
+++ b/Engine/Source/Runtime/Experimental/InteractiveToolsFramework/Private/BaseBehaviors/ClickDragBehavior.cpp
@@ -0,0 +1,77 @@
+// Copyright 1998-2019 Epic Games, Inc. All Rights Reserved.
+
+#include "BaseBehaviors/ClickDragBehavior.h"
+
+
+
+UClickDragInputBehavior::UClickDragInputBehavior()
+{
+}
+
+
+void UClickDragInputBehavior::Initialize(IClickDragBehaviorTarget* TargetIn)
+{
+	check(TargetIn != nullptr);
+	this->Target = TargetIn;
+}
+
+
+FInputCaptureRequest UClickDragInputBehavior::WantsCapture(const FInputDeviceState& Input)
+{
+	if (IsPressed(Input) && (ModifierCheckFunc == nullptr || ModifierCheckFunc(Input)) )
+	{
+		if ( Target->CanBeginClickDragSequence(GetDeviceRay(Input)) )
+		{
+			return FInputCaptureRequest::Begin(this, EInputCaptureSide::Any);
+		}
+	}
+	return FInputCaptureRequest::Ignore();
+}
+
+
+FInputCaptureUpdate UClickDragInputBehavior::BeginCapture(const FInputDeviceState& Input, EInputCaptureSide Side)
+{
+	Modifiers.UpdateModifiers(Input, Target);
+	OnClickPress(Input, Side);
+	return FInputCaptureUpdate::Begin(this, EInputCaptureSide::Any);
+}
+
+
+FInputCaptureUpdate UClickDragInputBehavior::UpdateCapture(const FInputDeviceState& Input, const FInputCaptureData& Data)
+{
+	Modifiers.UpdateModifiers(Input, Target);
+
+	if (IsReleased(Input)) 
+	{
+		OnClickRelease(Input, Data);
+		return FInputCaptureUpdate::End();
+	}
+	else
+	{
+		OnClickDrag(Input, Data);
+		return FInputCaptureUpdate::Continue();
+	}
+}
+
+
+void UClickDragInputBehavior::ForceEndCapture(const FInputCaptureData& Data)
+{
+	Target->OnTerminateDragSequence();
+}
+
+
+
+void UClickDragInputBehavior::OnClickPress(const FInputDeviceState& Input, EInputCaptureSide Side)
+{
+	Target->OnClickPress(GetDeviceRay(Input));
+}
+
+void UClickDragInputBehavior::OnClickDrag(const FInputDeviceState& Input, const FInputCaptureData& Data)
+{
+	Target->OnClickDrag(GetDeviceRay(Input));
+}
+
+void UClickDragInputBehavior::OnClickRelease(const FInputDeviceState& Input, const FInputCaptureData& Data)
+{
+	Target->OnClickRelease(GetDeviceRay(Input));
+}
diff --git a/Engine/Source/Runtime/Experimental/InteractiveToolsFramework/Private/BaseBehaviors/KeyAsModifierInputBehavior.cpp b/Engine/Source/Runtime/Experimental/InteractiveToolsFramework/Private/BaseBehaviors/KeyAsModifierInputBehavior.cpp
new file mode 100644
--- /dev/null
+++ b/Engine/Source/Runtime/Experimental/InteractiveToolsFramework/Private/BaseBehaviors/KeyAsModifierInputBehavior.cpp
@@ -0,0 +1,56 @@
+// Copyright 1998-2019 Epic Games, Inc. All Rights Reserved.
+
+#include "BaseBehaviors/KeyAsModifierInputBehavior.h"
+
+UKeyAsModifierInputBehavior::UKeyAsModifierInputBehavior()
+{
+}
+
+void UKeyAsModifierInputBehavior::Initialize(IModifierToggleBehaviorTarget* TargetIn, int ModifierID, const FKey& ModifierKeyIn)
+{
+	this->Target = TargetIn;
+	ModifierKey = ModifierKeyIn;
+	FKey TempKey = ModifierKeyIn;
+	Modifiers.RegisterModifier(ModifierID, [TempKey](const FInputDeviceState& Input)
+	{
+		return Input.Keyboard.ActiveKey.Button == TempKey;
+	});
+}
+
+FInputCaptureRequest UKeyAsModifierInputBehavior::WantsCapture(const FInputDeviceState& Input)
+{
+	if ((ModifierCheckFunc == nullptr || ModifierCheckFunc(Input)))
+	{
+		if (Input.Keyboard.ActiveKey.Button == ModifierKey && Input.Keyboard.ActiveKey.bPressed)
+		{
+			return FInputCaptureRequest::Begin(this, EInputCaptureSide::Any);
+		}
+	}
+	return FInputCaptureRequest::Ignore();
+}
+
+
+FInputCaptureUpdate UKeyAsModifierInputBehavior::BeginCapture(const FInputDeviceState& Input, EInputCaptureSide Side)
+{
+	Modifiers.UpdateModifiers(Input, Target);
+	return FInputCaptureUpdate::Begin(this, EInputCaptureSide::Any);
+}
+
+
+FInputCaptureUpdate UKeyAsModifierInputBehavior::UpdateCapture(const FInputDeviceState& Input, const FInputCaptureData& Data)
+{
+	Modifiers.UpdateModifiers(Input, Target);
+
+	if (Input.Keyboard.ActiveKey.bReleased)
+	{
+		return FInputCaptureUpdate::End();
+	}
+
+	return FInputCaptureUpdate::Continue();
+}
+
+
+void UKeyAsModifierInputBehavior::ForceEndCapture(const FInputCaptureData& Data)
+{
+}
+
diff --git a/Engine/Source/Runtime/Experimental/InteractiveToolsFramework/Private/BaseBehaviors/MouseHoverBehavior.cpp b/Engine/Source/Runtime/Experimental/InteractiveToolsFramework/Private/BaseBehaviors/MouseHoverBehavior.cpp
new file mode 100644
--- /dev/null
+++ b/Engine/Source/Runtime/Experimental/InteractiveToolsFramework/Private/BaseBehaviors/MouseHoverBehavior.cpp
@@ -0,0 +1,39 @@
+// Copyright 1998-2019 Epic Games, Inc. All Rights Reserved.
+
+#include "BaseBehaviors/MouseHoverBehavior.h"
+
+
+UMouseHoverBehavior::UMouseHoverBehavior()
+{
+	Target = nullptr;
+}
+
+EInputDevices UMouseHoverBehavior::GetSupportedDevices()
+{
+	return EInputDevices::Mouse;
+}
+
+void UMouseHoverBehavior::Initialize(IHoverBehaviorTarget* TargetIn)
+{
+	this->Target = TargetIn;
+}
+
+bool UMouseHoverBehavior::WantsHoverEvents()
+{
+	return true;
+}
+
+void UMouseHoverBehavior::UpdateHover(const FInputDeviceState& input)
+{
+	if (Target != nullptr)
+	{
+		Modifiers.UpdateModifiers(input, Target);
+
+		Target->OnUpdateHover( FInputDeviceRay(input.Mouse.WorldRay, input.Mouse.Position2D) );
+	}
+}
+
+void UMouseHoverBehavior::EndHover(const FInputDeviceState& input)
+{
+}
+
diff --git a/Engine/Source/Runtime/Experimental/InteractiveToolsFramework/Private/BaseBehaviors/MultiClickSequenceInputBehavior.cpp b/Engine/Source/Runtime/Experimental/InteractiveToolsFramework/Private/BaseBehaviors/MultiClickSequenceInputBehavior.cpp
new file mode 100644
--- /dev/null
+++ b/Engine/Source/Runtime/Experimental/InteractiveToolsFramework/Private/BaseBehaviors/MultiClickSequenceInputBehavior.cpp
@@ -0,0 +1,102 @@
+// Copyright 1998-2019 Epic Games, Inc. All Rights Reserved.
+
+#include "BaseBehaviors/MultiClickSequenceInputBehavior.h"
+
+
+
+UMultiClickSequenceInputBehavior::UMultiClickSequenceInputBehavior()
+{
+	bInActiveSequence = false;
+}
+
+
+void UMultiClickSequenceInputBehavior::Initialize(IClickSequenceBehaviorTarget* TargetIn)
+{
+	this->Target = TargetIn;
+	bInActiveSequence = false;
+}
+
+
+FInputCaptureRequest UMultiClickSequenceInputBehavior::WantsCapture(const FInputDeviceState& input)
+{
+	check(bInActiveSequence == false);   // should not happen...
+	bInActiveSequence = false;
+
+	if ( IsPressed(input) && (ModifierCheckFunc == nullptr || ModifierCheckFunc(input) ) )
+	{
+		if ( Target->CanBeginClickSequence(GetDeviceRay(input)) )
+		{
+			return FInputCaptureRequest::Begin(this, EInputCaptureSide::Any);
+		}
+	}
+	return FInputCaptureRequest::Ignore();
+}
+
+
+FInputCaptureUpdate UMultiClickSequenceInputBehavior::BeginCapture(const FInputDeviceState& input, EInputCaptureSide eSide)
+{
+	Modifiers.UpdateModifiers(input, Target);
+
+	Target->OnBeginClickSequence(GetDeviceRay(input));
+	bInActiveSequence = true;
+	
+	return FInputCaptureUpdate::Begin(this, EInputCaptureSide::Any);
+}
+
+
+FInputCaptureUpdate UMultiClickSequenceInputBehavior::UpdateCapture(const FInputDeviceState& input, const FInputCaptureData& data)
+{
+	check(bInActiveSequence == true);   // should always be the case!
+
+	Modifiers.UpdateModifiers(input, Target);
+
+	// allow target to abort click sequence
+	if (Target->RequestAbortClickSequence())
+	{
+		Target->OnTerminateClickSequence();
+		bInActiveSequence = false;
+		return FInputCaptureUpdate::End();
+	}
+
+	if (IsReleased(input)) 
+	{
+		bool bContinue = Target->OnNextSequenceClick(GetDeviceRay(input));
+		if (bContinue == false)
+		{
+			bInActiveSequence = false;
+			return FInputCaptureUpdate::End();
+		}
+	}
+	else
+	{
+		Target->OnNextSequencePreview(GetDeviceRay(input));
+	}
+
+	return FInputCaptureUpdate::Continue();
+}
+
+
+void UMultiClickSequenceInputBehavior::ForceEndCapture(const FInputCaptureData& data)
+{
+	Target->OnTerminateClickSequence();
+	bInActiveSequence = false;
+}
+
+
+bool UMultiClickSequenceInputBehavior::WantsHoverEvents()
+{
+	return true;
+}
+
+void UMultiClickSequenceInputBehavior::UpdateHover(const FInputDeviceState& Input)
+{
+	if (Target != nullptr)
+	{
+		Modifiers.UpdateModifiers(Input, Target);
+		Target->OnBeginSequencePreview(FInputDeviceRay(Input.Mouse.WorldRay, Input.Mouse.Position2D));
+	}
+}
+
+void UMultiClickSequenceInputBehavior::EndHover(const FInputDeviceState& input)
+{
+}
\ No newline at end of file
diff --git a/Engine/Source/Runtime/Experimental/InteractiveToolsFramework/Private/BaseBehaviors/SingleClickBehavior.cpp b/Engine/Source/Runtime/Experimental/InteractiveToolsFramework/Private/BaseBehaviors/SingleClickBehavior.cpp
new file mode 100644
--- /dev/null
+++ b/Engine/Source/Runtime/Experimental/InteractiveToolsFramework/Private/BaseBehaviors/SingleClickBehavior.cpp
@@ -0,0 +1,65 @@
+// Copyright 1998-2019 Epic Games, Inc. All Rights Reserved.
+
+#include "BaseBehaviors/SingleClickBehavior.h"
+
+
+
+USingleClickInputBehavior::USingleClickInputBehavior()
+{
+	HitTestOnRelease = true;
+}
+
+
+void USingleClickInputBehavior::Initialize(IClickBehaviorTarget* TargetIn)
+{
+	this->Target = TargetIn;
+}
+
+
+FInputCaptureRequest USingleClickInputBehavior::WantsCapture(const FInputDeviceState& input)
+{
+	if (IsPressed(input) && (ModifierCheckFunc == nullptr || ModifierCheckFunc(input)) )
+	{
+		if ( Target->IsHitByClick(GetDeviceRay(input)) )
+		{
+			return FInputCaptureRequest::Begin(this, EInputCaptureSide::Any);
+		}
+	}
+	return FInputCaptureRequest::Ignore();
+}
+
+
+FInputCaptureUpdate USingleClickInputBehavior::BeginCapture(const FInputDeviceState& input, EInputCaptureSide eSide)
+{
+	return FInputCaptureUpdate::Begin(this, EInputCaptureSide::Any);
+}
+
+
+FInputCaptureUpdate USingleClickInputBehavior::UpdateCapture(const FInputDeviceState& input, const FInputCaptureData& data)
+{
+	if (IsReleased(input)) 
+	{
+		if (HitTestOnRelease == false || 
+			Target->IsHitByClick(GetDeviceRay(input)) )
+		{
+			Clicked(input, data);
+		}
+
+		return FInputCaptureUpdate::End();
+	}
+
+	return FInputCaptureUpdate::Continue();
+}
+
+
+void USingleClickInputBehavior::ForceEndCapture(const FInputCaptureData& data)
+{
+	// nothing to do
+}
+
+
+void USingleClickInputBehavior::Clicked(const FInputDeviceState& input, const FInputCaptureData& data)
+{
+	Target->OnClicked(GetDeviceRay(input));
+}
+
diff --git a/Engine/Source/Runtime/Experimental/InteractiveToolsFramework/Private/BaseTools/ClickDragTool.cpp b/Engine/Source/Runtime/Experimental/InteractiveToolsFramework/Private/BaseTools/ClickDragTool.cpp
new file mode 100644
--- /dev/null
+++ b/Engine/Source/Runtime/Experimental/InteractiveToolsFramework/Private/BaseTools/ClickDragTool.cpp
@@ -0,0 +1,72 @@
+// Copyright 1998-2019 Epic Games, Inc. All Rights Reserved.
+
+#include "BaseTools/ClickDragTool.h"
+#include "InteractiveToolManager.h"
+#include "BaseBehaviors/ClickDragBehavior.h"
+
+
+/*
+ * ToolBuilder
+ */
+
+bool UClickDragToolBuilder::CanBuildTool(const FToolBuilderState& SceneState) const
+{
+	return true;
+}
+
+UInteractiveTool* UClickDragToolBuilder::BuildTool(const FToolBuilderState& SceneState) const
+{
+	UClickDragTool* NewTool = NewObject<UClickDragTool>(SceneState.ToolManager);
+	return NewTool;
+}
+
+
+
+/*
+ * Tool
+ */
+
+
+void UClickDragTool::Setup()
+{
+	UInteractiveTool::Setup();
+
+	// add default mouse input behavior
+	UClickDragInputBehavior* MouseBehavior = NewObject<UClickDragInputBehavior>();
+	MouseBehavior->Initialize(this);
+	AddInputBehavior(MouseBehavior);
+}
+
+
+bool UClickDragTool::CanBeginClickDragSequence(const FInputDeviceRay& ClickPos)
+{
+	return true;
+}
+
+
+void UClickDragTool::OnClickPress(const FInputDeviceRay& ClickPos)
+{
+	// print debug message
+	GetToolManager()->PostMessage( 
+		FString::Printf( TEXT("UClickDragTool::OnClickPress at (%f,%f)"), ClickPos.ScreenPosition.X, ClickPos.ScreenPosition.Y), 
+		EToolMessageLevel::Internal );
+}
+
+void UClickDragTool::OnClickDrag(const FInputDeviceRay& ClickPos)
+{
+}
+
+
+void UClickDragTool::OnClickRelease(const FInputDeviceRay& ClickPos)
+{
+	// print debug message
+	GetToolManager()->PostMessage(
+		FString::Printf(TEXT("UClickDragTool::OnClickRelease at (%f,%f)"), ClickPos.ScreenPosition.X, ClickPos.ScreenPosition.Y),
+		EToolMessageLevel::Internal);
+}
+
+
+void UClickDragTool::OnTerminateDragSequence()
+{
+
+}
\ No newline at end of file
diff --git a/Engine/Source/Runtime/Experimental/InteractiveToolsFramework/Private/BaseTools/MeshSurfacePointTool.cpp b/Engine/Source/Runtime/Experimental/InteractiveToolsFramework/Private/BaseTools/MeshSurfacePointTool.cpp
new file mode 100644
--- /dev/null
+++ b/Engine/Source/Runtime/Experimental/InteractiveToolsFramework/Private/BaseTools/MeshSurfacePointTool.cpp
@@ -0,0 +1,166 @@
+// Copyright 1998-2019 Epic Games, Inc. All Rights Reserved.
+
+#include "BaseTools/MeshSurfacePointTool.h"
+#include "BaseBehaviors/MouseHoverBehavior.h"
+#include "InteractiveToolManager.h"
+#include "ToolBuilderUtil.h"
+
+
+/*
+ * ToolBuilder
+ */
+
+bool UMeshSurfacePointToolBuilder::CanBuildTool(const FToolBuilderState& SceneState) const
+{
+	return ToolBuilderUtil::CountComponents(SceneState, ToolBuilderUtil::IsMeshDescriptionSourceComponent) == 1;
+}
+
+UInteractiveTool* UMeshSurfacePointToolBuilder::BuildTool(const FToolBuilderState& SceneState) const
+{
+	UMeshSurfacePointTool* NewTool = CreateNewTool(SceneState);
+	InitializeNewTool(NewTool, SceneState);
+	return NewTool;
+}
+
+UMeshSurfacePointTool* UMeshSurfacePointToolBuilder::CreateNewTool(const FToolBuilderState& SceneState) const
+{
+	return NewObject<UMeshSurfacePointTool>(SceneState.ToolManager);
+}
+
+void UMeshSurfacePointToolBuilder::InitializeNewTool(UMeshSurfacePointTool* NewTool, const FToolBuilderState& SceneState) const
+{
+	UActorComponent* MeshComponent = ToolBuilderUtil::FindFirstComponent(SceneState, ToolBuilderUtil::IsMeshDescriptionSourceComponent);
+	check(MeshComponent != nullptr);
+	NewTool->SetMeshSource(
+		SceneState.SourceBuilder->MakeMeshDescriptionSource(MeshComponent) );
+}
+
+
+/*
+ * Tool
+ */
+
+void UMeshSurfacePointTool::SetMeshSource(TUniquePtr<IMeshDescriptionSource> MeshSourceIn)
+{
+	this->MeshSource = TUniquePtr<IMeshDescriptionSource>(MoveTemp(MeshSourceIn));
+}
+
+
+void UMeshSurfacePointTool::Setup()
+{
+	UInteractiveTool::Setup();
+
+	bShiftToggle = false;
+
+	// add input behaviors
+	UMeshSurfacePointToolMouseBehavior* mouseBehavior = NewObject<UMeshSurfacePointToolMouseBehavior>();
+	mouseBehavior->Initialize(this);
+	AddInputBehavior(mouseBehavior);
+
+	UMouseHoverBehavior* hoverBehavior = NewObject<UMouseHoverBehavior>();
+	hoverBehavior->Initialize(this);
+	AddInputBehavior(hoverBehavior);
+}
+
+
+bool UMeshSurfacePointTool::HitTest(const FRay& Ray, FHitResult& OutHit)
+{
+	return MeshSource->HitTest(Ray, OutHit);
+}
+
+
+void UMeshSurfacePointTool::OnBeginDrag(const FRay& Ray)
+{
+
+}
+
+
+void UMeshSurfacePointTool::OnUpdateDrag(const FRay& Ray)
+{
+	FHitResult OutHit;
+	if ( HitTest(Ray, OutHit) ) 
+	{
+		GetToolManager()->PostMessage( 
+			FString::Printf(TEXT("[UMeshSurfacePointTool::OnUpdateDrag] Hit triangle index %d at ray distance %f"), OutHit.FaceIndex, OutHit.Distance), EToolMessageLevel::Internal);
+	}
+}
+
+void UMeshSurfacePointTool::OnEndDrag(const FRay& Ray)
+{
+	//GetToolManager()->PostMessage(TEXT("UMeshSurfacePointTool::OnEndDrag!"), EToolMessageLevel::Internal);
+}
+
+
+void UMeshSurfacePointTool::SetShiftToggle(bool bShiftDown)
+{
+	bShiftToggle = bShiftDown;
+}
+
+
+
+
+/*
+ * Mouse Input Behavior
+ */
+
+
+void UMeshSurfacePointToolMouseBehavior::Initialize(UMeshSurfacePointTool* ToolIn)
+{
+	this->Tool = ToolIn;
+	bInDragCapture = false;
+}
+
+
+FInputCaptureRequest UMeshSurfacePointToolMouseBehavior::WantsCapture(const FInputDeviceState& input)
+{
+	if ( IsPressed(input) )
+	{
+		FHitResult OutHit;
+		if (Tool->HitTest(input.Mouse.WorldRay, OutHit))
+		{
+			return FInputCaptureRequest::Begin(this, EInputCaptureSide::Any, OutHit.Distance);
+		}
+	}
+	return FInputCaptureRequest::Ignore();
+}
+
+FInputCaptureUpdate UMeshSurfacePointToolMouseBehavior::BeginCapture(const FInputDeviceState& input, EInputCaptureSide eSide)
+{
+	Tool->SetShiftToggle(input.bShiftKeyDown);
+	Tool->OnBeginDrag(input.Mouse.WorldRay);
+	LastWorldRay = input.Mouse.WorldRay;
+	bInDragCapture = true;
+	return FInputCaptureUpdate::Begin(this, EInputCaptureSide::Any);
+}
+
+FInputCaptureUpdate UMeshSurfacePointToolMouseBehavior::UpdateCapture(const FInputDeviceState& input, const FInputCaptureData& data)
+{
+	LastWorldRay = input.Mouse.WorldRay;
+
+	if ( IsReleased(input) ) 
+	{
+		Tool->OnEndDrag(input.Mouse.WorldRay);
+		bInDragCapture = false;
+		return FInputCaptureUpdate::End();
+	}
+
+	Tool->OnUpdateDrag(input.Mouse.WorldRay);
+	return FInputCaptureUpdate::Continue();
+}
+
+void UMeshSurfacePointToolMouseBehavior::ForceEndCapture(const FInputCaptureData& data)
+{
+	if (bInDragCapture)
+	{
+		Tool->OnEndDrag(LastWorldRay);
+		bInDragCapture = false;
+	}
+
+	// nothing to do
+}
+
+
+
+
+
+
diff --git a/Engine/Source/Runtime/Experimental/InteractiveToolsFramework/Private/BaseTools/SingleClickTool.cpp b/Engine/Source/Runtime/Experimental/InteractiveToolsFramework/Private/BaseTools/SingleClickTool.cpp
new file mode 100644
--- /dev/null
+++ b/Engine/Source/Runtime/Experimental/InteractiveToolsFramework/Private/BaseTools/SingleClickTool.cpp
@@ -0,0 +1,56 @@
+// Copyright 1998-2019 Epic Games, Inc. All Rights Reserved.
+
+#include "BaseTools/SingleClickTool.h"
+#include "InteractiveToolManager.h"
+
+
+/*
+ * ToolBuilder
+ */
+
+bool USingleClickToolBuilder::CanBuildTool(const FToolBuilderState& SceneState) const
+{
+	return true;
+}
+
+UInteractiveTool* USingleClickToolBuilder::BuildTool(const FToolBuilderState& SceneState) const
+{
+	USingleClickTool* NewTool = NewObject<USingleClickTool>(SceneState.ToolManager);
+	return NewTool;
+}
+
+
+
+/*
+ * Tool
+ */
+
+
+void USingleClickTool::Setup()
+{
+	UInteractiveTool::Setup();
+
+	// add default button input behaviors for devices
+	USingleClickInputBehavior* MouseBehavior = NewObject<USingleClickInputBehavior>();
+	MouseBehavior->Initialize(this);
+	AddInputBehavior(MouseBehavior);
+}
+
+
+bool USingleClickTool::IsHitByClick(const FInputDeviceRay& ClickPos)
+{
+	return true;
+}
+
+
+void USingleClickTool::OnClicked(const FInputDeviceRay& ClickPos)
+{
+	// print debug message
+	GetToolManager()->PostMessage( 
+		FString::Printf( TEXT("USingleClickTool::OnClicked at (%f,%f)"), ClickPos.ScreenPosition.X, ClickPos.ScreenPosition.Y), 
+		EToolMessageLevel::Internal );
+}
+
+
+
+
diff --git a/Engine/Source/Runtime/Experimental/InteractiveToolsFramework/Private/InputBehavior.cpp b/Engine/Source/Runtime/Experimental/InteractiveToolsFramework/Private/InputBehavior.cpp
new file mode 100644
--- /dev/null
+++ b/Engine/Source/Runtime/Experimental/InteractiveToolsFramework/Private/InputBehavior.cpp
@@ -0,0 +1,77 @@
+// Copyright 1998-2019 Epic Games, Inc. All Rights Reserved.
+
+
+#include "InputBehavior.h"
+
+
+bool operator<(const FInputCaptureRequest& l, const FInputCaptureRequest& r)
+{
+	if (l.Source->GetPriority() == r.Source->GetPriority())
+	{
+		return l.HitDepth < r.HitDepth;
+	}
+	else
+	{
+		return l.Source->GetPriority() < r.Source->GetPriority();
+	}
+}
+
+
+
+UInputBehavior::UInputBehavior()
+{
+	DefaultPriority = FInputCapturePriority(100);
+}
+
+FInputCapturePriority UInputBehavior::GetPriority()
+{
+	return DefaultPriority;
+}
+
+void UInputBehavior::SetDefaultPriority(const FInputCapturePriority& Priority)
+{
+	DefaultPriority = Priority;
+}
+
+EInputDevices UInputBehavior::GetSupportedDevices()
+{
+	return EInputDevices::Mouse;
+}
+
+
+FInputCaptureRequest UInputBehavior::WantsCapture(const FInputDeviceState& input)
+{
+	return FInputCaptureRequest::Ignore();
+}
+
+FInputCaptureUpdate UInputBehavior::BeginCapture(const FInputDeviceState& input, EInputCaptureSide eSide)
+{
+	return FInputCaptureUpdate::Ignore();
+}
+
+FInputCaptureUpdate UInputBehavior::UpdateCapture(const FInputDeviceState& input, const FInputCaptureData& data)
+{
+	return FInputCaptureUpdate::End();
+}
+
+void UInputBehavior::ForceEndCapture(const FInputCaptureData& data)
+{
+	return;
+}
+
+
+bool UInputBehavior::WantsHoverEvents()
+{
+	return false;
+}
+
+void UInputBehavior::UpdateHover(const FInputDeviceState& input)
+{
+	return;
+}
+
+void UInputBehavior::EndHover(const FInputDeviceState& input)
+{
+	return;
+}
+
diff --git a/Engine/Source/Runtime/Experimental/InteractiveToolsFramework/Private/InputBehaviorSet.cpp b/Engine/Source/Runtime/Experimental/InteractiveToolsFramework/Private/InputBehaviorSet.cpp
new file mode 100644
--- /dev/null
+++ b/Engine/Source/Runtime/Experimental/InteractiveToolsFramework/Private/InputBehaviorSet.cpp
@@ -0,0 +1,164 @@
+// Copyright 1998-2019 Epic Games, Inc. All Rights Reserved.
+
+#include "InputBehaviorSet.h"
+
+
+UInputBehaviorSet::UInputBehaviorSet()
+{
+}
+
+UInputBehaviorSet::~UInputBehaviorSet()
+{
+}
+
+
+bool UInputBehaviorSet::IsEmpty() const
+{
+	return Behaviors.Num() == 0;
+}
+
+void UInputBehaviorSet::Add(UInputBehavior* Behavior, void* Source, const FString& Group)
+{
+	//if (source == null)
+	//	source = DefaultSource;
+
+	FBehaviorInfo Info;
+	Info.Behavior = Behavior;
+	Info.Source = Source;
+	Info.Group = Group;
+
+	Behaviors.Add(Info);
+	BehaviorsModified();
+}
+
+
+void UInputBehaviorSet::Add(const UInputBehaviorSet* OtherSet, void* NewSource, const FString& NewGroup)
+{
+	for (const FBehaviorInfo& OtherInfo : OtherSet->Behaviors)
+	{
+		FBehaviorInfo Info;
+		Info.Behavior = OtherInfo.Behavior;
+		Info.Source = (NewSource == nullptr) ? OtherInfo.Source : NewSource;
+		Info.Group = (NewGroup.IsEmpty()) ? OtherInfo.Group : NewGroup;
+		Behaviors.Add(Info);
+	}
+	if (OtherSet->Behaviors.Num() > 0)
+	{
+		BehaviorsModified();
+	}
+}
+
+
+bool UInputBehaviorSet::Remove(UInputBehavior* behavior)
+{
+	int32 idx = Behaviors.IndexOfByPredicate( 
+		[behavior] (const FBehaviorInfo& b) { return b.Behavior == behavior; }
+	);
+	if ( idx != INDEX_NONE ) {
+		Behaviors.RemoveAt(idx);
+		BehaviorsModified();
+		return true;
+	}
+	return false;
+}
+
+
+bool UInputBehaviorSet::RemoveByGroup(const FString& Group)
+{
+	int RemovedCount = Behaviors.RemoveAll(
+		[Group](const FBehaviorInfo& b) { return b.Group == Group; }
+	);
+	if (RemovedCount > 0) 
+	{
+		BehaviorsModified();
+	}
+	return (RemovedCount > 0);
+}
+
+
+
+bool UInputBehaviorSet::RemoveBySource(void* Source)
+{
+	int RemovedCount = Behaviors.RemoveAll(
+		[Source](const FBehaviorInfo& b) { return b.Source == Source; }
+	);
+	if (RemovedCount > 0)
+	{
+		BehaviorsModified();
+	}
+	return (RemovedCount > 0);	
+}
+
+
+void UInputBehaviorSet::RemoveAll()
+{
+	Behaviors.Reset();
+	BehaviorsModified();
+}
+
+
+void UInputBehaviorSet::CollectWantsCapture(const FInputDeviceState& input, TArray<FInputCaptureRequest>& result)
+{
+	for (auto b : Behaviors)
+	{
+		// only call WantsCapture if the Behavior supports the current input device
+		if (SupportsInputType(b.Behavior, input)) 
+		{
+			FInputCaptureRequest request = b.Behavior->WantsCapture(input);
+			if ( request.Type != EInputCaptureRequestType::Ignore ) 
+			{
+				request.Owner = b.Source;
+				result.Add(request);
+			}
+		}
+	}
+}
+
+
+
+bool UInputBehaviorSet::UpdateHover(const FInputDeviceState& input)
+{
+	bool bAnyWantedHover = false;
+	for (auto b : Behaviors) 
+	{
+		if (SupportsInputType(b.Behavior, input)) 
+		{
+			if (b.Behavior->WantsHoverEvents()) 
+			{
+				b.Behavior->UpdateHover(input);
+				bAnyWantedHover = true;
+			}
+		}
+	}
+	return bAnyWantedHover;
+}
+
+bool UInputBehaviorSet::EndHover(const FInputDeviceState& input)
+{
+	bool bAnyWantedHover = false;
+	for (auto b : Behaviors)
+	{
+		if (SupportsInputType(b.Behavior, input)) 
+		{
+			if (b.Behavior->WantsHoverEvents()) 
+			{
+				b.Behavior->EndHover(input);
+				bAnyWantedHover = true;
+			}
+		}
+	}
+	return bAnyWantedHover;
+}
+
+
+
+
+
+void UInputBehaviorSet::BehaviorsModified()
+{
+	// sort by priority
+	Behaviors.StableSort();
+
+	//send some kind of event...
+	//FUtil.SafeSendAnyEvent(OnSetChanged, this);
+}
diff --git a/Engine/Source/Runtime/Experimental/InteractiveToolsFramework/Private/InputRouter.cpp b/Engine/Source/Runtime/Experimental/InteractiveToolsFramework/Private/InputRouter.cpp
new file mode 100644
--- /dev/null
+++ b/Engine/Source/Runtime/Experimental/InteractiveToolsFramework/Private/InputRouter.cpp
@@ -0,0 +1,318 @@
+// Copyright 1998-2019 Epic Games, Inc. All Rights Reserved.
+
+
+#include "InputRouter.h"
+
+
+UInputRouter::UInputRouter()
+{
+	ActiveLeftCapture = nullptr;
+	ActiveLeftCaptureOwner = nullptr;
+	ActiveRightCapture = nullptr;
+	ActiveRightCaptureOwner = nullptr;
+	ActiveKeyboardCapture = nullptr;
+	ActiveKeyboardCaptureOwner = nullptr;
+
+	bAutoInvalidateOnHover = false;
+	bAutoInvalidateOnCapture = false;
+
+	ActiveInputBehaviors = NewObject<UInputBehaviorSet>(this, "InputBehaviors");
+}
+
+
+
+void UInputRouter::Initialize(IToolsContextTransactionsAPI* TransactionsAPIIn)
+{
+	this->TransactionsAPI = TransactionsAPIIn;
+}
+
+void UInputRouter::Shutdown()
+{
+	this->TransactionsAPI = nullptr;
+}
+
+
+
+
+void UInputRouter::RegisterSource(IInputBehaviorSource* Source)
+{
+	ActiveInputBehaviors->Add(Source->GetInputBehaviors(), Source);
+}
+
+
+void UInputRouter::DeregisterSource(IInputBehaviorSource* Source)
+{
+	ActiveInputBehaviors->RemoveBySource(Source);
+}
+
+
+
+void UInputRouter::PostInputEvent(const FInputDeviceState& Input)
+{
+	if (ActiveInputBehaviors->IsEmpty())
+	{
+		return;
+	}
+
+	if (Input.IsFromDevice(EInputDevices::Mouse))
+	{
+		PostInputEvent_Mouse(Input);
+	}
+	else if (Input.IsFromDevice(EInputDevices::Keyboard))
+	{
+		PostInputEvent_Keyboard(Input);
+	}
+	else
+	{
+		unimplemented();
+		TransactionsAPI->PostMessage(TEXT("UInteractiveToolManager::PostInputEvent - input device is not currently supported."), EToolMessageLevel::Internal);
+		return;
+	}
+}
+
+
+
+//
+// Keyboard event handling
+// 
+
+
+void UInputRouter::PostInputEvent_Keyboard(const FInputDeviceState& Input)
+{
+	if (ActiveKeyboardCapture != nullptr)
+	{
+		HandleCapturedKeyboardInput(Input);
+	}
+	else
+	{
+		ActiveKeyboardCaptureData = FInputCaptureData();
+		CheckForKeyboardCaptures(Input);
+	}
+}
+
+
+void UInputRouter::CheckForKeyboardCaptures(const FInputDeviceState& Input)
+{
+	TArray<FInputCaptureRequest> CaptureRequests;
+	ActiveInputBehaviors->CollectWantsCapture(Input, CaptureRequests);
+	if (CaptureRequests.Num() == 0)
+	{
+		return;
+	}
+
+	CaptureRequests.StableSort();
+
+	bool bAccepted = false;
+	for (int i = 0; i < CaptureRequests.Num() && bAccepted == false; ++i)
+	{
+		FInputCaptureUpdate Result =
+			CaptureRequests[i].Source->BeginCapture(Input, EInputCaptureSide::Left);
+		if (Result.State == EInputCaptureState::Begin)
+		{
+			ActiveKeyboardCapture = Result.Source;
+			ActiveKeyboardCaptureOwner = CaptureRequests[i].Owner;
+			ActiveKeyboardCaptureData = Result.Data;
+			bAccepted = true;
+		}
+	}
+}
+
+
+void UInputRouter::HandleCapturedKeyboardInput(const FInputDeviceState& Input)
+{
+	if (ActiveKeyboardCapture == nullptr)
+	{
+		return;
+	}
+
+	FInputCaptureUpdate Result =
+		ActiveKeyboardCapture->UpdateCapture(Input, ActiveKeyboardCaptureData);
+
+	if (Result.State == EInputCaptureState::End)
+	{
+		ActiveKeyboardCapture = nullptr;
+		ActiveKeyboardCaptureOwner = nullptr;
+		ActiveKeyboardCaptureData = FInputCaptureData();
+	}
+	else if (Result.State != EInputCaptureState::Continue)
+	{
+		TransactionsAPI->PostMessage(TEXT("UInteractiveToolManager::HandleCapturedKeyboardInput - unexpected capture state!"), EToolMessageLevel::Internal);
+	}
+
+	if (bAutoInvalidateOnCapture)
+	{
+		TransactionsAPI->PostInvalidation();
+	}
+}
+
+
+
+//
+// Mouse event handling
+//
+
+
+
+void UInputRouter::PostInputEvent_Mouse(const FInputDeviceState& Input)
+{
+	if (ActiveLeftCapture != nullptr)
+	{
+		HandleCapturedMouseInput(Input);
+	}
+	else
+	{
+		ActiveLeftCaptureData = FInputCaptureData();
+		CheckForMouseCaptures(Input);
+	}
+
+	// update hover if nobody is capturing
+	if (ActiveLeftCapture == nullptr && ActiveRightCapture == nullptr)
+	{
+		bool bProcessed = ActiveInputBehaviors->UpdateHover(Input);
+		if (bProcessed && bAutoInvalidateOnHover)
+		{
+			TransactionsAPI->PostInvalidation();
+		}
+	}
+
+}
+
+
+void UInputRouter::PostHoverInputEvent(const FInputDeviceState& Input)
+{
+	LastHoverInput = Input;
+	bool bProcessed = ActiveInputBehaviors->UpdateHover(Input);
+	if (bProcessed && bAutoInvalidateOnHover)
+	{
+		TransactionsAPI->PostInvalidation();
+	}
+}
+
+
+
+
+bool UInputRouter::HasActiveMouseCapture() const
+{
+	return (ActiveLeftCapture != nullptr);
+}
+
+
+void UInputRouter::CheckForMouseCaptures(const FInputDeviceState& Input)
+{
+	TArray<FInputCaptureRequest> CaptureRequests;
+	ActiveInputBehaviors->CollectWantsCapture(Input, CaptureRequests);
+	if (CaptureRequests.Num() == 0)
+	{
+		return;
+	}
+
+	CaptureRequests.StableSort();
+
+	bool bAccepted = false;
+	for (int i = 0; i < CaptureRequests.Num() && bAccepted == false; ++i)
+	{
+		FInputCaptureUpdate Result =
+			CaptureRequests[i].Source->BeginCapture(Input, EInputCaptureSide::Left);
+		if (Result.State == EInputCaptureState::Begin)
+		{
+			// end outstanding hovers
+			ActiveInputBehaviors->EndHover(Input);
+
+			ActiveLeftCapture = Result.Source;
+			ActiveLeftCaptureOwner = CaptureRequests[i].Owner;
+			ActiveLeftCaptureData = Result.Data;
+			bAccepted = true;
+		}
+
+	}
+}
+
+
+void UInputRouter::HandleCapturedMouseInput(const FInputDeviceState& Input)
+{
+	if (ActiveLeftCapture == nullptr)
+	{
+		return;
+	}
+
+	// have active capture - give it this event
+
+	FInputCaptureUpdate Result =
+		ActiveLeftCapture->UpdateCapture(Input, ActiveLeftCaptureData);
+
+	if (Result.State == EInputCaptureState::End)
+	{
+		ActiveLeftCapture = nullptr;
+		ActiveLeftCaptureOwner = nullptr;
+		ActiveLeftCaptureData = FInputCaptureData();
+	}
+	else if (Result.State != EInputCaptureState::Continue)
+	{
+		TransactionsAPI->PostMessage(TEXT("UInteractiveToolManager::HandleCapturedMouseInput - unexpected capture state!"), EToolMessageLevel::Internal);
+	}
+
+	if (bAutoInvalidateOnCapture)
+	{
+		TransactionsAPI->PostInvalidation();
+	}
+}
+
+
+
+void UInputRouter::ForceTerminateAll()
+{
+	if (ActiveKeyboardCapture != nullptr)
+	{
+		ActiveKeyboardCapture->ForceEndCapture(ActiveKeyboardCaptureData);
+		ActiveKeyboardCapture = nullptr;
+		ActiveKeyboardCaptureOwner = nullptr;
+		ActiveKeyboardCaptureData = FInputCaptureData();
+	}
+
+	if (ActiveLeftCapture != nullptr)
+	{
+		ActiveLeftCapture->ForceEndCapture(ActiveLeftCaptureData);
+		ActiveLeftCapture = nullptr;
+		ActiveLeftCaptureOwner = nullptr;
+		ActiveLeftCaptureData = FInputCaptureData();
+	}
+
+	if (ActiveRightCapture != nullptr)
+	{
+		ActiveRightCapture->ForceEndCapture(ActiveRightCaptureData);
+		ActiveRightCapture = nullptr;
+		ActiveRightCaptureOwner = nullptr;
+		ActiveRightCaptureData = FInputCaptureData();
+	}
+
+	ActiveInputBehaviors->EndHover(LastHoverInput);
+}
+
+
+void UInputRouter::ForceTerminateSource(IInputBehaviorSource* Source)
+{
+	if (ActiveKeyboardCapture != nullptr && ActiveKeyboardCaptureOwner == Source)
+	{
+		ActiveKeyboardCapture->ForceEndCapture(ActiveKeyboardCaptureData);
+		ActiveKeyboardCapture = nullptr;
+		ActiveKeyboardCaptureOwner = nullptr;
+		ActiveKeyboardCaptureData = FInputCaptureData();
+	}
+
+	if (ActiveLeftCapture != nullptr && ActiveLeftCaptureOwner == Source)
+	{
+		ActiveLeftCapture->ForceEndCapture(ActiveLeftCaptureData);
+		ActiveLeftCapture = nullptr;
+		ActiveLeftCaptureOwner = nullptr;
+		ActiveLeftCaptureData = FInputCaptureData();
+	}
+
+	if (ActiveRightCapture != nullptr && ActiveRightCaptureOwner == Source)
+	{
+		ActiveRightCapture->ForceEndCapture(ActiveRightCaptureData);
+		ActiveRightCapture = nullptr;
+		ActiveRightCaptureOwner = nullptr;
+		ActiveRightCaptureData = FInputCaptureData();
+	}
+}
+
diff --git a/Engine/Source/Runtime/Experimental/InteractiveToolsFramework/Private/InputState.cpp b/Engine/Source/Runtime/Experimental/InteractiveToolsFramework/Private/InputState.cpp
new file mode 100644
--- /dev/null
+++ b/Engine/Source/Runtime/Experimental/InteractiveToolsFramework/Private/InputState.cpp
@@ -0,0 +1,4 @@
+// Copyright 1998-2019 Epic Games, Inc. All Rights Reserved.
+
+#include "InputState.h"
+
diff --git a/Engine/Source/Runtime/Experimental/InteractiveToolsFramework/Private/InteractiveGizmo.cpp b/Engine/Source/Runtime/Experimental/InteractiveToolsFramework/Private/InteractiveGizmo.cpp
new file mode 100644
--- /dev/null
+++ b/Engine/Source/Runtime/Experimental/InteractiveToolsFramework/Private/InteractiveGizmo.cpp
@@ -0,0 +1,52 @@
+// Copyright 1998-2019 Epic Games, Inc. All Rights Reserved.
+
+
+#include "InteractiveGizmo.h"
+#include "InteractiveGizmoManager.h"
+
+
+UInteractiveGizmo::UInteractiveGizmo()
+{
+	// Gizmos don't get saved but this isn't necessary because they are created in the transient package...
+	SetFlags(RF_Transient);
+
+	InputBehaviors = NewObject<UInputBehaviorSet>(this, TEXT("GizmoInputBehaviors"));
+}
+
+void UInteractiveGizmo::Setup()
+{
+}
+
+void UInteractiveGizmo::Shutdown()
+{
+	InputBehaviors->RemoveAll();
+}
+
+void UInteractiveGizmo::Render(IToolsContextRenderAPI* RenderAPI)
+{
+}
+
+
+void UInteractiveGizmo::AddInputBehavior(UInputBehavior* Behavior)
+{
+	InputBehaviors->Add(Behavior);
+}
+
+const UInputBehaviorSet* UInteractiveGizmo::GetInputBehaviors() const
+{
+	return InputBehaviors;
+}
+
+
+
+void UInteractiveGizmo::Tick(float DeltaTime)
+{
+}
+
+
+UInteractiveGizmoManager* UInteractiveGizmo::GetGizmoManager() const
+{
+	UInteractiveGizmoManager* GizmoManager = Cast<UInteractiveGizmoManager>(GetOuter());
+	check(GizmoManager != nullptr);
+	return GizmoManager;
+}
\ No newline at end of file
diff --git a/Engine/Source/Runtime/Experimental/InteractiveToolsFramework/Private/InteractiveGizmoManager.cpp b/Engine/Source/Runtime/Experimental/InteractiveToolsFramework/Private/InteractiveGizmoManager.cpp
new file mode 100644
--- /dev/null
+++ b/Engine/Source/Runtime/Experimental/InteractiveToolsFramework/Private/InteractiveGizmoManager.cpp
@@ -0,0 +1,228 @@
+// Copyright 1998-2019 Epic Games, Inc. All Rights Reserved.
+
+
+#include "InteractiveGizmoManager.h"
+
+
+UInteractiveGizmoManager::UInteractiveGizmoManager()
+{
+	QueriesAPI = nullptr;
+	TransactionsAPI = nullptr;
+	InputRouter = nullptr;
+}
+
+
+void UInteractiveGizmoManager::Initialize(IToolsContextQueriesAPI* QueriesAPIIn, IToolsContextTransactionsAPI* TransactionsAPIIn, UInputRouter* InputRouterIn)
+{
+	this->QueriesAPI = QueriesAPIIn;
+	this->TransactionsAPI = TransactionsAPIIn;
+	this->InputRouter = InputRouterIn;
+}
+
+
+void UInteractiveGizmoManager::Shutdown()
+{
+	this->QueriesAPI = nullptr;
+
+	TArray<FActiveGizmo> AllGizmos = ActiveGizmos;
+	for (FActiveGizmo& ActiveGizmo : AllGizmos)
+	{
+		DestroyGizmo(ActiveGizmo.Gizmo);
+	}
+	ActiveGizmos.Reset();
+
+	this->TransactionsAPI = nullptr;
+}
+
+
+
+void UInteractiveGizmoManager::RegisterGizmoType(const FString& Identifier, UInteractiveGizmoBuilder* Builder)
+{
+	check(GizmoBuilders.Contains(Identifier) == false);
+	GizmoBuilders.Add(Identifier, Builder );
+}
+
+
+bool UInteractiveGizmoManager::DeregisterGizmoType(const FString& BuilderIdentifier)
+{
+	if (GizmoBuilders.Contains(BuilderIdentifier) == false)
+	{
+		PostMessage(FString::Printf(TEXT("UInteractiveGizmoManager::DeregisterGizmoType: could not find requested type %s"), *BuilderIdentifier), EToolMessageLevel::Internal);
+		return false;
+	}
+	GizmoBuilders.Remove(BuilderIdentifier);
+	return true;
+}
+
+
+
+
+UInteractiveGizmo* UInteractiveGizmoManager::CreateGizmo(const FString& BuilderIdentifier, const FString& InstanceIdentifier)
+{
+	if ( GizmoBuilders.Contains(BuilderIdentifier) == false )
+	{
+		PostMessage(FString::Printf(TEXT("UInteractiveGizmoManager::CreateGizmo: could not find requested type %s"), *BuilderIdentifier), EToolMessageLevel::Internal);
+		return nullptr;
+	}
+	UInteractiveGizmoBuilder* FoundBuilder = GizmoBuilders[BuilderIdentifier];
+
+	// check if we have used this instance identifier
+	for (FActiveGizmo& ActiveGizmo : ActiveGizmos)
+	{
+		if (ActiveGizmo.InstanceIdentifier == InstanceIdentifier)
+		{
+			PostMessage(FString::Printf(TEXT("UInteractiveGizmoManager::CreateGizmo: instance identifier %s already in use!"), *InstanceIdentifier), EToolMessageLevel::Internal);
+			return nullptr;
+		}
+	}
+
+	FToolBuilderState CurrentSceneState;
+	QueriesAPI->GetCurrentSelectionState(CurrentSceneState);
+
+	UInteractiveGizmo* NewGizmo = FoundBuilder->BuildGizmo(CurrentSceneState);
+	if (NewGizmo == nullptr)
+	{
+		PostMessage(FString::Printf(TEXT("UInteractiveGizmoManager::CreateGizmo: BuildGizmo() returned null")), EToolMessageLevel::Internal);
+		return nullptr;
+	}
+
+	NewGizmo->Setup();
+
+	// register new active input behaviors
+	InputRouter->RegisterSource(NewGizmo);
+
+	PostInvalidation();
+
+	FActiveGizmo ActiveGizmo = { NewGizmo, BuilderIdentifier, InstanceIdentifier };
+	ActiveGizmos.Add(ActiveGizmo);
+
+	return NewGizmo;
+}
+
+
+
+bool UInteractiveGizmoManager::DestroyGizmo(UInteractiveGizmo* Gizmo)
+{
+	int FoundIndex = -1;
+	for ( int i = 0; i < ActiveGizmos.Num(); ++i )
+	{
+		if (ActiveGizmos[i].Gizmo == Gizmo)
+		{
+			FoundIndex = i;
+			break;
+		}
+	}
+	if (FoundIndex == -1)
+	{
+		return false;
+	}
+
+	InputRouter->ForceTerminateSource(Gizmo);
+
+	Gizmo->Shutdown();
+
+	InputRouter->DeregisterSource(Gizmo);
+
+	ActiveGizmos.RemoveAt(FoundIndex);
+
+	PostInvalidation();
+
+	return true;
+}
+
+
+
+
+TArray<UInteractiveGizmo*> UInteractiveGizmoManager::FindAllGizmosOfType(const FString& BuilderIdentifier)
+{
+	TArray<UInteractiveGizmo*> Found;
+	for (int i = 0; i < ActiveGizmos.Num(); ++i)
+	{
+		if (ActiveGizmos[i].BuilderIdentifier == BuilderIdentifier)
+		{
+			Found.Add(ActiveGizmos[i].Gizmo);
+		}
+	}
+	return Found;
+}
+
+
+void UInteractiveGizmoManager::DestroyAllGizmosOfType(const FString& BuilderIdentifier)
+{
+	TArray<UInteractiveGizmo*> ToRemove = FindAllGizmosOfType(BuilderIdentifier);
+
+	for (int i = 0; i < ToRemove.Num(); ++i)
+	{
+		DestroyGizmo(ToRemove[i]);
+	}
+}
+
+
+
+UInteractiveGizmo* UInteractiveGizmoManager::FindGizmoByInstanceIdentifier(const FString& Identifier)
+{
+	for (int i = 0; i < ActiveGizmos.Num(); ++i)
+	{
+		if (ActiveGizmos[i].InstanceIdentifier == Identifier)
+		{
+			return ActiveGizmos[i].Gizmo;
+		}
+	}
+	return nullptr;
+}
+
+
+
+void UInteractiveGizmoManager::Tick(float DeltaTime)
+{
+	for (FActiveGizmo& ActiveGizmo : ActiveGizmos)
+	{
+		ActiveGizmo.Gizmo->Tick(DeltaTime);
+	}
+}
+
+
+void UInteractiveGizmoManager::Render(IToolsContextRenderAPI* RenderAPI)
+{
+	for (FActiveGizmo& ActiveGizmo : ActiveGizmos)
+	{
+		ActiveGizmo.Gizmo->Render(RenderAPI);
+	}
+
+}
+
+
+void UInteractiveGizmoManager::PostMessage(const TCHAR* Message, EToolMessageLevel Level)
+{
+	TransactionsAPI->PostMessage(Message, Level);
+}
+
+void UInteractiveGizmoManager::PostMessage(const FString& Message, EToolMessageLevel Level)
+{
+	TransactionsAPI->PostMessage(*Message, Level);
+}
+
+void UInteractiveGizmoManager::PostInvalidation()
+{
+	TransactionsAPI->PostInvalidation();
+}
+
+
+void UInteractiveGizmoManager::BeginUndoTransaction(const FText& Description)
+{
+	TransactionsAPI->BeginUndoTransaction(Description);
+}
+
+void UInteractiveGizmoManager::EndUndoTransaction()
+{
+	TransactionsAPI->EndUndoTransaction();
+}
+
+
+
+void UInteractiveGizmoManager::EmitObjectChange(UObject* TargetObject, TUniquePtr<FChange> Change, const FText& Description)
+{
+	TransactionsAPI->AppendChange(TargetObject, MoveTemp(Change), Description );
+}
+
+
diff --git a/Engine/Source/Runtime/Experimental/InteractiveToolsFramework/Private/InteractiveTool.cpp b/Engine/Source/Runtime/Experimental/InteractiveToolsFramework/Private/InteractiveTool.cpp
new file mode 100644
--- /dev/null
+++ b/Engine/Source/Runtime/Experimental/InteractiveToolsFramework/Private/InteractiveTool.cpp
@@ -0,0 +1,115 @@
+// Copyright 1998-2019 Epic Games, Inc. All Rights Reserved.
+
+
+#include "InteractiveTool.h"
+#include "InteractiveToolManager.h"
+
+
+UInteractiveTool::UInteractiveTool()
+{
+	// tools need to be transactional or undo/redo won't work on their uproperties
+	SetFlags(RF_Transactional);
+
+	// tools don't get saved but this isn't necessary because they are created in the transient package...
+	//SetFlags(RF_Transient);
+
+	InputBehaviors = NewObject<UInputBehaviorSet>(this, TEXT("InputBehaviors"));
+}
+
+void UInteractiveTool::Setup()
+{
+}
+
+void UInteractiveTool::Shutdown(EToolShutdownType ShutdownType)
+{
+	InputBehaviors->RemoveAll();
+	ToolPropertyObjects.Reset();
+}
+
+void UInteractiveTool::Render(IToolsContextRenderAPI* RenderAPI)
+{
+}
+
+
+void UInteractiveTool::AddInputBehavior(UInputBehavior* Behavior)
+{
+	InputBehaviors->Add(Behavior);
+}
+
+const UInputBehaviorSet* UInteractiveTool::GetInputBehaviors() const
+{
+	return InputBehaviors;
+}
+
+
+void UInteractiveTool::AddToolPropertySource(UObject* PropertyObject)
+{
+	check(ToolPropertyObjects.Contains(PropertyObject) == false);
+	ToolPropertyObjects.Add(PropertyObject);
+}
+
+void UInteractiveTool::AddToolPropertySource(UInteractiveToolPropertySet* PropertySet)
+{
+	check(ToolPropertyObjects.Contains(PropertySet) == false);
+	ToolPropertyObjects.Add(PropertySet);
+	// @todo do we need to create a lambda every time for this?
+	PropertySet->GetOnModified().AddLambda([this](UObject* PropertySetArg, UProperty* PropertyArg)
+	{
+		OnPropertyModified(PropertySetArg, PropertyArg);
+	});
+}
+
+
+const TArray<UObject*>& UInteractiveTool::GetToolProperties() const
+{
+	return ToolPropertyObjects;
+}
+
+
+void UInteractiveTool::RegisterActions(FInteractiveToolActionSet& ActionSet)
+{
+}
+
+FInteractiveToolActionSet* UInteractiveTool::GetActionSet()
+{
+	if (ToolActionSet == nullptr)
+	{
+		ToolActionSet = new FInteractiveToolActionSet();
+		RegisterActions(*ToolActionSet);
+	}
+	return ToolActionSet;
+}
+
+void UInteractiveTool::ExecuteAction(int32 ActionID)
+{
+	GetActionSet()->ExecuteAction(ActionID);
+}
+
+
+
+bool UInteractiveTool::HasCancel() const
+{
+	return false;
+}
+
+bool UInteractiveTool::HasAccept() const
+{
+	return false;
+}
+
+bool UInteractiveTool::CanAccept() const
+{
+	return false;
+}
+
+
+void UInteractiveTool::Tick(float DeltaTime)
+{
+}
+
+UInteractiveToolManager* UInteractiveTool::GetToolManager() const
+{
+	UInteractiveToolManager* ToolManager = Cast<UInteractiveToolManager>(GetOuter());
+	check(ToolManager != nullptr);
+	return ToolManager;
+}
\ No newline at end of file
diff --git a/Engine/Source/Runtime/Experimental/InteractiveToolsFramework/Private/InteractiveToolActionSet.cpp b/Engine/Source/Runtime/Experimental/InteractiveToolsFramework/Private/InteractiveToolActionSet.cpp
new file mode 100644
--- /dev/null
+++ b/Engine/Source/Runtime/Experimental/InteractiveToolsFramework/Private/InteractiveToolActionSet.cpp
@@ -0,0 +1,59 @@
+// Copyright 1998-2019 Epic Games, Inc. All Rights Reserved.
+
+#include "InteractiveToolActionSet.h"
+#include "InteractiveTool.h"
+
+
+
+void FInteractiveToolActionSet::RegisterAction(UInteractiveTool* Tool, int32 ActionID,
+	const FString& ActionName, const FString& ShortUIName, const FString& DescriptionText,
+	EModifierKey::Type Modifiers, const FKey& ShortcutKey,
+	TFunction<void()> ActionFunction )
+{
+	checkf(FindActionByID(ActionID) == nullptr, TEXT("InteractiveToolActionSet::RegisterAction: Action ID is already registered!"));
+
+	FInteractiveToolAction NewAction;
+	NewAction.ClassType = Tool->GetClass();
+	NewAction.ActionID = ActionID;
+	NewAction.ActionName = ActionName;
+	NewAction.ShortName = ShortUIName;
+	NewAction.Description = DescriptionText;
+	NewAction.DefaultModifiers = Modifiers;
+	NewAction.DefaultKey = ShortcutKey;
+	NewAction.OnAction = ActionFunction;
+
+	Actions.Add(NewAction);
+}
+
+
+const FInteractiveToolAction* FInteractiveToolActionSet::FindActionByID(int32 ActionID) const
+{
+	for (const FInteractiveToolAction& Action : Actions)
+	{
+		if (Action.ActionID == ActionID)
+		{
+			return &Action;
+		}
+	}
+	return nullptr;
+}
+
+
+
+void FInteractiveToolActionSet::CollectActions(TArray<FInteractiveToolAction>& OutActions) const
+{
+	for (const FInteractiveToolAction& Action : Actions)
+	{
+		OutActions.Add(Action);
+	}
+}
+
+
+void FInteractiveToolActionSet::ExecuteAction(int32 ActionID) const
+{
+	const FInteractiveToolAction* Found = FindActionByID(ActionID);
+	if (Found != nullptr)
+	{
+		Found->OnAction();
+	}
+}
\ No newline at end of file
diff --git a/Engine/Source/Runtime/Experimental/InteractiveToolsFramework/Private/InteractiveToolBuilder.cpp b/Engine/Source/Runtime/Experimental/InteractiveToolsFramework/Private/InteractiveToolBuilder.cpp
new file mode 100644
--- /dev/null
+++ b/Engine/Source/Runtime/Experimental/InteractiveToolsFramework/Private/InteractiveToolBuilder.cpp
@@ -0,0 +1,6 @@
+// Copyright 1998-2019 Epic Games, Inc. All Rights Reserved.
+
+
+#include "InteractiveToolBuilder.h"
+
+// Add default functionality here for any UInteractiveToolBuilder functions that are not pure virtual.
diff --git a/Engine/Source/Runtime/Experimental/InteractiveToolsFramework/Private/InteractiveToolManager.cpp b/Engine/Source/Runtime/Experimental/InteractiveToolsFramework/Private/InteractiveToolManager.cpp
new file mode 100644
--- /dev/null
+++ b/Engine/Source/Runtime/Experimental/InteractiveToolsFramework/Private/InteractiveToolManager.cpp
@@ -0,0 +1,265 @@
+// Copyright 1998-2019 Epic Games, Inc. All Rights Reserved.
+
+
+#include "InteractiveToolManager.h"
+
+
+UInteractiveToolManager::UInteractiveToolManager()
+{
+	QueriesAPI = nullptr;
+	TransactionsAPI = nullptr;
+	InputRouter = nullptr;
+
+	ActiveLeftBuilder = nullptr;
+	ActiveLeftTool = nullptr;
+
+	ActiveRightBuilder = nullptr;
+	ActiveRightTool = nullptr;
+}
+
+
+void UInteractiveToolManager::Initialize(IToolsContextQueriesAPI* queriesAPI, IToolsContextTransactionsAPI* transactionsAPI, UInputRouter* InputRouterIn)
+{
+	this->QueriesAPI = queriesAPI;
+	this->TransactionsAPI = transactionsAPI;
+	this->InputRouter = InputRouterIn;
+}
+
+
+void UInteractiveToolManager::Shutdown()
+{
+	this->QueriesAPI = nullptr;
+
+	if (ActiveLeftTool != nullptr)
+	{
+		DeactivateTool(EToolSide::Left, EToolShutdownType::Cancel);
+	}
+	if (ActiveRightTool != nullptr)
+	{
+		DeactivateTool(EToolSide::Right, EToolShutdownType::Cancel);
+	}
+
+	this->TransactionsAPI = nullptr;
+}
+
+
+
+void UInteractiveToolManager::RegisterToolType(const FString& Identifier, UInteractiveToolBuilder* Builder)
+{
+	check(ToolBuilders.Contains(Identifier) == false);
+	ToolBuilders.Add(Identifier, Builder );
+}
+
+
+bool UInteractiveToolManager::SelectActiveToolType(EToolSide Side, const FString& Identifier)
+{
+	if (ToolBuilders.Contains(Identifier))
+	{
+		UInteractiveToolBuilder* Builder = ToolBuilders[Identifier];
+		if (Side == EToolSide::Right)
+		{
+			ActiveRightBuilder = Builder;
+		}
+		else
+		{
+			ActiveLeftBuilder = Builder;
+		}
+		return true;
+	}
+	return false;
+}
+
+
+
+bool UInteractiveToolManager::CanActivateTool(EToolSide Side, const FString& Identifier)
+{
+	check(Side == EToolSide::Left);   // TODO: support right-side tool
+
+	if (ActiveLeftTool != nullptr)
+	{
+		return false;
+	}
+
+	if (ToolBuilders.Contains(Identifier))
+	{
+		FToolBuilderState InputState;
+		QueriesAPI->GetCurrentSelectionState(InputState);
+
+		UInteractiveToolBuilder* Builder = ToolBuilders[Identifier];
+		return Builder->CanBuildTool(InputState);
+	}
+
+	return false;
+}
+
+
+bool UInteractiveToolManager::ActivateTool(EToolSide Side)
+{
+	check(Side == EToolSide::Left);   // TODO: support right-side tool
+
+	if (ActiveLeftTool != nullptr) 
+	{
+		DeactivateTool(EToolSide::Left, EToolShutdownType::Cancel);
+	}
+
+	if (ActiveLeftBuilder == nullptr) 
+	{
+		return false;
+	}
+
+	// construct input state we will pass to tools
+	FToolBuilderState InputState;
+	QueriesAPI->GetCurrentSelectionState(InputState);
+
+	if (ActiveLeftBuilder->CanBuildTool(InputState) == false)
+	{
+		TransactionsAPI->PostMessage( TEXT("UInteractiveToolManager::ActivateTool: CanBuildTool returned false."), EToolMessageLevel::Internal);
+		return false;
+	}
+
+	ActiveLeftTool = ActiveLeftBuilder->BuildTool(InputState);
+	if (ActiveLeftTool == nullptr)
+	{
+		return false;
+	}
+
+	ActiveLeftTool->Setup();
+
+	// register new active input behaviors
+	InputRouter->RegisterSource(ActiveLeftTool);
+
+	PostInvalidation();
+
+	OnToolStarted.Broadcast(this, ActiveLeftTool);
+
+	return true;
+}
+
+
+void UInteractiveToolManager::DeactivateTool(EToolSide Side, EToolShutdownType ShutdownType)
+{
+	check(Side == EToolSide::Left);   // TODO: support right-side tool
+
+	if (ActiveLeftTool != nullptr)
+	{
+		InputRouter->ForceTerminateSource(ActiveLeftTool);
+
+		ActiveLeftTool->Shutdown(ShutdownType);
+
+		InputRouter->DeregisterSource(ActiveLeftTool);
+
+		UInteractiveTool* DoneTool = ActiveLeftTool;
+		ActiveLeftTool = nullptr;
+
+		PostInvalidation();
+
+		OnToolEnded.Broadcast(this, DoneTool);
+	}
+}
+
+
+bool UInteractiveToolManager::HasActiveTool(EToolSide Side) const
+{
+	return (Side == EToolSide::Left) ? (ActiveLeftTool != nullptr) : (ActiveRightTool != nullptr);
+}
+
+bool UInteractiveToolManager::HasAnyActiveTool() const
+{
+	return ActiveLeftTool != nullptr || ActiveRightTool != nullptr;
+}
+
+
+UInteractiveTool* UInteractiveToolManager::GetActiveTool(EToolSide Side)
+{
+	return (Side == EToolSide::Left) ? ActiveLeftTool : ActiveRightTool;
+}
+
+
+bool UInteractiveToolManager::CanAcceptActiveTool(EToolSide Side)
+{
+	if (ActiveLeftTool != nullptr)
+	{
+		return ActiveLeftTool->HasAccept() && ActiveLeftTool->CanAccept();
+	}
+	return false;
+}
+
+bool UInteractiveToolManager::CanCancelActiveTool(EToolSide Side)
+{
+	if (ActiveLeftTool != nullptr)
+	{
+		return ActiveLeftTool->HasCancel();
+	}
+	return false;
+}
+
+
+
+
+
+
+void UInteractiveToolManager::Tick(float DeltaTime)
+{
+	if (ActiveLeftTool != nullptr)
+	{
+		ActiveLeftTool->Tick(DeltaTime);
+	}
+
+	if (ActiveRightTool!= nullptr)
+	{
+		ActiveRightTool->Tick(DeltaTime);
+	}
+}
+
+
+void UInteractiveToolManager::Render(IToolsContextRenderAPI* RenderAPI)
+{
+	if (ActiveLeftTool != nullptr)
+	{
+		ActiveLeftTool->Render(RenderAPI);
+	}
+
+	if (ActiveRightTool != nullptr)
+	{
+		ActiveRightTool->Render(RenderAPI);
+	}
+}
+
+
+void UInteractiveToolManager::PostMessage(const TCHAR* Message, EToolMessageLevel Level)
+{
+	TransactionsAPI->PostMessage(Message, Level);
+}
+
+void UInteractiveToolManager::PostMessage(const FString& Message, EToolMessageLevel Level)
+{
+	TransactionsAPI->PostMessage(*Message, Level);
+}
+
+void UInteractiveToolManager::PostInvalidation()
+{
+	TransactionsAPI->PostInvalidation();
+}
+
+
+void UInteractiveToolManager::BeginUndoTransaction(const FText& Description)
+{
+	TransactionsAPI->BeginUndoTransaction(Description);
+}
+
+void UInteractiveToolManager::EndUndoTransaction()
+{
+	TransactionsAPI->EndUndoTransaction();
+}
+
+
+
+void UInteractiveToolManager::EmitObjectChange(UObject* TargetObject, TUniquePtr<FChange> Change, const FText& Description)
+{
+	TransactionsAPI->AppendChange(TargetObject, MoveTemp(Change), Description );
+}
+
+bool UInteractiveToolManager::RequestSelectionChange(const FSelectedOjectsChangeList& SelectionChange)
+{
+	return TransactionsAPI->RequestSelectionChange(SelectionChange);
+}
diff --git a/Engine/Source/Runtime/Experimental/InteractiveToolsFramework/Private/InteractiveToolsContext.cpp b/Engine/Source/Runtime/Experimental/InteractiveToolsFramework/Private/InteractiveToolsContext.cpp
new file mode 100644
--- /dev/null
+++ b/Engine/Source/Runtime/Experimental/InteractiveToolsFramework/Private/InteractiveToolsContext.cpp
@@ -0,0 +1,38 @@
+// Copyright 1998-2019 Epic Games, Inc. All Rights Reserved.
+
+
+#include "InteractiveToolsContext.h"
+
+
+UInteractiveToolsContext::UInteractiveToolsContext()
+{
+	InputRouter = nullptr;
+	ToolManager = nullptr;
+}
+
+void UInteractiveToolsContext::Initialize(IToolsContextQueriesAPI* QueriesAPI, IToolsContextTransactionsAPI* TransactionsAPI)
+{
+	InputRouter = NewObject<UInputRouter>(this);
+	InputRouter->Initialize(TransactionsAPI);
+
+	ToolManager = NewObject<UInteractiveToolManager>(this);
+	ToolManager->Initialize(QueriesAPI, TransactionsAPI, InputRouter);
+
+	GizmoManager = NewObject<UInteractiveGizmoManager>(this);
+	GizmoManager->Initialize(QueriesAPI, TransactionsAPI, InputRouter);
+}
+
+
+void UInteractiveToolsContext::Shutdown()
+{
+	// force-terminate any remaining captures/hovers/etc
+	InputRouter->ForceTerminateAll();
+	InputRouter->Shutdown();
+	InputRouter = nullptr;
+
+	GizmoManager->Shutdown();
+	GizmoManager = nullptr;
+
+	ToolManager->Shutdown();
+	ToolManager = nullptr;
+}
diff --git a/Engine/Source/Runtime/Experimental/InteractiveToolsFramework/Private/InteractiveToolsFramework.cpp b/Engine/Source/Runtime/Experimental/InteractiveToolsFramework/Private/InteractiveToolsFramework.cpp
new file mode 100644
--- /dev/null
+++ b/Engine/Source/Runtime/Experimental/InteractiveToolsFramework/Private/InteractiveToolsFramework.cpp
@@ -0,0 +1,20 @@
+// Copyright 1998-2019 Epic Games, Inc. All Rights Reserved.
+
+#include "InteractiveToolsFramework.h"
+
+#define LOCTEXT_NAMESPACE "FInteractiveToolsFrameworkModule"
+
+void FInteractiveToolsFrameworkModule::StartupModule()
+{
+	// This code will execute after your module is loaded into memory; the exact timing is specified in the .uplugin file per-module
+}
+
+void FInteractiveToolsFrameworkModule::ShutdownModule()
+{
+	// This function may be called during shutdown to clean up your module.  For modules that support dynamic reloading,
+	// we call this function before unloading the module.
+}
+
+#undef LOCTEXT_NAMESPACE
+	
+IMPLEMENT_MODULE(FInteractiveToolsFrameworkModule, InteractiveToolsFramework)
\ No newline at end of file
diff --git a/Engine/Source/Runtime/Experimental/InteractiveToolsFramework/Private/ToolBuilderUtil.cpp b/Engine/Source/Runtime/Experimental/InteractiveToolsFramework/Private/ToolBuilderUtil.cpp
new file mode 100644
--- /dev/null
+++ b/Engine/Source/Runtime/Experimental/InteractiveToolsFramework/Private/ToolBuilderUtil.cpp
@@ -0,0 +1,137 @@
+// Copyright 1998-2019 Epic Games, Inc. All Rights Reserved.
+
+#include "ToolBuilderUtil.h"
+#include "CoreMinimal.h"
+#include "GameFramework/Actor.h"
+#include "Engine/Selection.h"
+#include "Components/StaticMeshComponent.h"
+
+
+
+bool ToolBuilderUtil::IsMeshDescriptionSourceComponent(UActorComponent* ComponentObject)
+{
+	UStaticMeshComponent* StaticMeshComp = Cast<UStaticMeshComponent>(ComponentObject);
+	if (StaticMeshComp != nullptr)
+	{
+		return true;
+	}
+	return false;
+}
+
+
+
+int ToolBuilderUtil::CountComponents(const FToolBuilderState& InputState, const TFunction<bool(UActorComponent*)>& Predicate)
+{
+	int nTypedComponents = 0;
+
+	if (InputState.SelectedComponents != nullptr && InputState.SelectedComponents->Num() > 0)
+	{
+		for (FSelectionIterator Iter(*InputState.SelectedComponents); Iter; ++Iter)
+		{
+			if (UActorComponent* Component = Cast<UActorComponent>(*Iter))
+			{
+				if (Predicate(Component))
+				{
+					nTypedComponents++;
+				}
+			}
+		}
+	}
+	else
+	{
+		for (FSelectionIterator Iter(*InputState.SelectedActors); Iter; ++Iter)
+		{
+			if (AActor* Actor = Cast<AActor>(*Iter))
+			{
+				for (UActorComponent* Component : Actor->GetComponents())
+				{
+					if (Predicate(Component))
+					{
+						nTypedComponents++;
+					}
+				}
+			}
+		}
+	}
+
+	return nTypedComponents;
+}
+
+
+
+
+UActorComponent* ToolBuilderUtil::FindFirstComponent(const FToolBuilderState& InputState, const TFunction<bool(UActorComponent*)>& Predicate)
+{
+	if (InputState.SelectedComponents != nullptr && InputState.SelectedComponents->Num() > 0)
+	{
+		for (FSelectionIterator Iter(*InputState.SelectedComponents); Iter; ++Iter)
+		{
+			if (UActorComponent* Component = Cast<UActorComponent>(*Iter))
+			{
+				if (Predicate(Component))
+				{
+					return Component;
+				}
+			}
+		}
+	}
+	else
+	{
+		for (FSelectionIterator Iter(*InputState.SelectedActors); Iter; ++Iter)
+		{
+			if (AActor* Actor = Cast<AActor>(*Iter))
+			{
+				for (UActorComponent* Component : Actor->GetComponents())
+				{
+					if (Predicate(Component))
+					{
+						return Component;
+					}
+				}
+			}
+		}
+	}
+
+	return nullptr;
+}
+
+
+
+
+TArray<UActorComponent*> ToolBuilderUtil::FindAllComponents(const FToolBuilderState& InputState, const TFunction<bool(UActorComponent*)>& Predicate)
+{
+	TArray<UActorComponent*> Components;
+
+	if (InputState.SelectedComponents != nullptr && InputState.SelectedComponents->Num() > 0)
+	{
+		for (FSelectionIterator Iter(*InputState.SelectedComponents); Iter; ++Iter)
+		{
+			if (UActorComponent* Component = Cast<UActorComponent>(*Iter))
+			{
+				if (Predicate(Component))
+				{
+					Components.AddUnique(Component);
+				}
+			}
+		}
+	}
+	else
+	{
+		for (FSelectionIterator Iter(*InputState.SelectedActors); Iter; ++Iter)
+		{
+			if (AActor* Actor = Cast<AActor>(*Iter))
+			{
+				for (UActorComponent* Component : Actor->GetComponents())
+				{
+					if (Predicate(Component))
+					{
+						Components.AddUnique(Component);
+					}
+				}
+			}
+		}
+	}
+
+	return Components;
+}
+
diff --git a/Engine/Source/Runtime/Experimental/InteractiveToolsFramework/Public/BaseBehaviors/AnyButtonInputBehavior.h b/Engine/Source/Runtime/Experimental/InteractiveToolsFramework/Public/BaseBehaviors/AnyButtonInputBehavior.h
new file mode 100644
--- /dev/null
+++ b/Engine/Source/Runtime/Experimental/InteractiveToolsFramework/Public/BaseBehaviors/AnyButtonInputBehavior.h
@@ -0,0 +1,70 @@
+// Copyright 1998-2019 Epic Games, Inc. All Rights Reserved.
+
+#pragma once
+
+#include "CoreMinimal.h"
+#include "InputBehavior.h"
+#include "AnyButtonInputBehavior.generated.h"
+
+/**
+ * UAnyButtonInputBehavior is a base behavior that provides a generic
+ * interface to a TargetButton on a physical Input Device. You can subclass
+ * UAnyButtonInputBehavior to write InputBehaviors that can work independent
+ * of a particular device type or button, by using the UAnyButtonInputBehavior functions below.
+ * 
+ * The target device button is selected using the .ButtonNumber property, or you can
+ * override the relevant GetXButtonState() function if you need more control.
+ * 
+ *  @todo spatial controllers
+ *  @todo support tablet fingers
+ *  @todo support gamepad?
+ */
+UCLASS()
+class INTERACTIVETOOLSFRAMEWORK_API UAnyButtonInputBehavior : public UInputBehavior
+{
+	GENERATED_BODY()
+
+public:
+	UAnyButtonInputBehavior();
+
+	/** Return set of devices supported by this behavior */
+	virtual EInputDevices GetSupportedDevices() override;
+
+	/** @return true if Target Button has been pressed this frame */
+	virtual bool IsPressed(const FInputDeviceState& input);
+
+	/** @return true if Target Button is currently held down */
+	virtual bool IsDown(const FInputDeviceState& input);
+
+	/** @return true if Target Button was released this frame */
+	virtual bool IsReleased(const FInputDeviceState& input);
+	
+	/** @return current 2D position of Target Device, or zero if device does not have 2D position */
+	virtual FVector2D GetClickPoint(const FInputDeviceState& input);
+
+	/** @return current 3D world ray for Target Device position */
+	virtual FRay GetWorldRay(const FInputDeviceState& input);
+	
+	/** @return current 3D world ray and optional 2D position for Target Device */
+	virtual FInputDeviceRay GetDeviceRay(const FInputDeviceState& input);
+
+	/** @return last-active supported Device */
+	EInputDevices GetActiveDevice() const;
+
+protected:
+
+	/** 
+	 * Button number on target device. Button 0 is "default" button on all devices.
+	 * Mouse: Left=0, Middle=1, Right=2
+	 */
+	UPROPERTY()
+	int ButtonNumber;
+
+
+protected:
+	/** Which device is currently active */
+	EInputDevices ActiveDevice;
+
+	/** @return mouse button state for Target Button */
+	virtual FDeviceButtonState GetMouseButtonState(const FInputDeviceState& input);
+};
\ No newline at end of file
diff --git a/Engine/Source/Runtime/Experimental/InteractiveToolsFramework/Public/BaseBehaviors/BehaviorTargetInterfaces.h b/Engine/Source/Runtime/Experimental/InteractiveToolsFramework/Public/BaseBehaviors/BehaviorTargetInterfaces.h
new file mode 100644
--- /dev/null
+++ b/Engine/Source/Runtime/Experimental/InteractiveToolsFramework/Public/BaseBehaviors/BehaviorTargetInterfaces.h
@@ -0,0 +1,173 @@
+// Copyright 1998-2019 Epic Games, Inc. All Rights Reserved.
+
+#pragma once
+
+#include "CoreMinimal.h"
+
+
+/**
+ * IModifierToggleBehaviorTarget is an interface that InputBehaviors can use to notify
+ * a target about modifier toggle states (like shift key being down, etc).
+ * The meaning of the modifier ID is client-defined (generally provided to the InputBehavior in a registration step)
+ */
+class IModifierToggleBehaviorTarget
+{
+public:
+	virtual ~IModifierToggleBehaviorTarget() {}
+
+	/**
+	 * Notify target of current modifier state
+	 * @param ModifierID client-defined integer that identifiers modifier
+	 * @param bIsOn is modifier current on or off
+	 */
+	virtual void OnUpdateModifierState(int ModifierID, bool bIsOn)
+	{
+	}
+};
+
+
+/**
+ * Functions required to apply standard "Click" state machines to a target object.
+ * See USingleClickBehavior for an example of this kind of state machine.
+ */
+class IClickBehaviorTarget : public IModifierToggleBehaviorTarget
+{
+public:
+	virtual ~IClickBehaviorTarget() {}
+
+	/**
+	 * Test if target is hit by a click
+	 * @param ClickPos device position/ray at click point
+	 * @return true if target was hit by click ray/point
+	 */
+	virtual bool IsHitByClick(const FInputDeviceRay& ClickPos) = 0;
+
+
+	/**
+	 * Notify Target that click ocurred
+	 * @param ClickPos device position/ray at click point
+	 */
+	virtual void OnClicked(const FInputDeviceRay& ClickPos) = 0;
+};
+
+
+
+
+/**
+ * Functions required to apply standard "Click-Drag" state machines to a target object.
+ * See UClickDragBehavior for an example of this kind of state machine.
+ */
+class IClickDragBehaviorTarget : public IModifierToggleBehaviorTarget
+{
+public:
+	virtual ~IClickDragBehaviorTarget() {}
+
+	/**
+	 * Test if target can begin click-drag interaction at this point
+	 * @param PressPos device position/ray at click point
+	 * @return true if target wants to begin sequence
+	 */
+	virtual bool CanBeginClickDragSequence(const FInputDeviceRay& PressPos) = 0;
+
+
+	/**
+	 * Notify Target that click press ocurred
+	 * @param PressPos device position/ray at click point
+	 */
+	virtual void OnClickPress(const FInputDeviceRay& PressPos) = 0;
+
+	/**
+	 * Notify Target that input position has changed
+	 * @param DragPos device position/ray at click point
+	 */
+	virtual void OnClickDrag(const FInputDeviceRay& DragPos) = 0;
+
+	/**
+	 * Notify Target that click release occurred
+	 * @param ReleasePos device position/ray at click point
+	 */
+	virtual void OnClickRelease(const FInputDeviceRay& ReleasePos) = 0;
+
+	/**
+	 * Notify Target that click-drag sequence has been explicitly terminated (eg by escape key)
+	 */
+	virtual void OnTerminateDragSequence() = 0;
+};
+
+
+
+
+
+
+/**
+ * Target interface used by InputBehaviors that want to implement a multi-click sequence
+ * (eg such as drawing a polygon with multiple clicks)
+ */
+class IClickSequenceBehaviorTarget : public IModifierToggleBehaviorTarget
+{
+public:
+	virtual ~IClickSequenceBehaviorTarget() {}
+
+	/**
+	 * Notify Target device position has changed but a click sequence hasn't begun yet (eg for interactive previews)
+	 * @param ClickPos device position/ray at click point
+	 */
+	virtual void OnBeginSequencePreview(const FInputDeviceRay& ClickPos) { }
+
+	/**
+	 * Test if target would like to begin sequence based on this click
+	 * @param ClickPos device position/ray at click point
+	 * @return true if target wants to begin sequence
+	 */
+	virtual bool CanBeginClickSequence(const FInputDeviceRay& ClickPos) = 0;
+
+	/**
+	 * Notify Target that click sequence can begin at click point
+	 * @param ClickPos device position/ray at click point
+	 */
+	virtual void OnBeginClickSequence(const FInputDeviceRay& ClickPos) = 0;
+	
+	/**
+	 * Notify Target device position has changed but a click hasn't ocurred yet (eg for interactive previews)
+	 * @param ClickPos device position/ray at click point
+	 */
+	virtual void OnNextSequencePreview(const FInputDeviceRay& ClickPos) { }
+
+	/**
+	 * Notify Target about next click in sqeuence
+	 * @param ClickPos device position/ray at click point
+	 * @return false if Target wants to terminate sequence
+	 */
+	virtual bool OnNextSequenceClick(const FInputDeviceRay& ClickPos) = 0;
+
+	/**
+	 * Notify Target that click sequence has been explicitly terminated (eg by escape key, cancel tool, etc)
+	 */
+	virtual void OnTerminateClickSequence() = 0;
+
+	/**
+	 * Target overrides this and returns true if it wants to abort click sequence.
+	 * Behavior checks every update and if this ever returns true, terminates sequence
+	 */
+	virtual bool RequestAbortClickSequence() { return false; }
+};
+
+
+
+
+
+/**
+ * IHoverBehaviorTarget allows Behaviors to notify Tools/etc about
+ * device event data in a generic way, without requiring that all Tools
+ * know about the concept of Hovering.
+ */
+class IHoverBehaviorTarget : public IModifierToggleBehaviorTarget
+{
+public:
+	virtual ~IHoverBehaviorTarget() {}
+
+	/**
+	 * Notify Target about a hover event
+	 */
+	virtual void OnUpdateHover(const FInputDeviceRay& DevicePos) = 0;
+};
diff --git a/Engine/Source/Runtime/Experimental/InteractiveToolsFramework/Public/BaseBehaviors/ClickDragBehavior.h b/Engine/Source/Runtime/Experimental/InteractiveToolsFramework/Public/BaseBehaviors/ClickDragBehavior.h
new file mode 100644
--- /dev/null
+++ b/Engine/Source/Runtime/Experimental/InteractiveToolsFramework/Public/BaseBehaviors/ClickDragBehavior.h
@@ -0,0 +1,76 @@
+// Copyright 1998-2019 Epic Games, Inc. All Rights Reserved.
+
+#pragma once
+
+#include "CoreMinimal.h"
+#include "BaseBehaviors/AnyButtonInputBehavior.h"
+#include "BaseBehaviors/BehaviorTargetInterfaces.h"
+#include "BaseBehaviors/InputBehaviorModifierStates.h"
+#include "ClickDragBehavior.generated.h"
+
+
+/**
+ * UClickDragInputBehavior implements a standard "button-click-drag"-style input behavior.
+ * An IClickDragBehaviorTarget instance must be provided which is manipulated by this behavior.
+ * 
+ * The state machine works as follows:
+ *    1) on input-device-button-press, call Target::CanBeginClickDragSequence to determine if capture should begin
+ *    2) on input-device-move, call Target::OnClickDrag
+ *    3) on input-device-button-release, call Target::OnClickRelease
+ *    
+ * If a ForceEndCapture occurs we call Target::OnTerminateDragSequence   
+ */
+UCLASS()
+class INTERACTIVETOOLSFRAMEWORK_API UClickDragInputBehavior : public UAnyButtonInputBehavior
+{
+	GENERATED_BODY()
+
+public:
+	UClickDragInputBehavior();
+
+	/**
+	 * The modifier set for this behavior
+	 */
+	FInputBehaviorModifierStates Modifiers;
+
+	/**
+	 * Initialize this behavior with the given Target
+	 * @param Target implementor of hit-test and on-clicked functions
+	 */
+	virtual void Initialize(IClickDragBehaviorTarget* Target);
+
+
+	/**
+	 * WantsCapture() will only return capture request if this function returns true (or is null)
+	 */
+	TFunction<bool(const FInputDeviceState&)> ModifierCheckFunc = nullptr;
+
+
+	// UInputBehavior implementation
+
+	virtual FInputCaptureRequest WantsCapture(const FInputDeviceState& Input) override;
+	virtual FInputCaptureUpdate BeginCapture(const FInputDeviceState& Input, EInputCaptureSide Side) override;
+	virtual FInputCaptureUpdate UpdateCapture(const FInputDeviceState& Input, const FInputCaptureData& Data) override;
+	virtual void ForceEndCapture(const FInputCaptureData& Data) override;
+
+
+protected:
+	/** Click Target object */
+	IClickDragBehaviorTarget* Target;
+
+	/**
+	 * Internal function that forwards click evens to Target::OnClickPress, you can customize behavior here
+	 */
+	virtual void OnClickPress(const FInputDeviceState& Input, EInputCaptureSide Side);
+
+	/**
+	 * Internal function that forwards click evens to Target::OnClickDrag, you can customize behavior here
+	 */
+	virtual void OnClickDrag(const FInputDeviceState& Input, const FInputCaptureData& Data);
+
+	/**
+	 * Internal function that forwards click evens to Target::OnClickRelease, you can customize behavior here
+	 */
+	virtual void OnClickRelease(const FInputDeviceState& Input, const FInputCaptureData& Data);
+};
+
diff --git a/Engine/Source/Runtime/Experimental/InteractiveToolsFramework/Public/BaseBehaviors/InputBehaviorModifierStates.h b/Engine/Source/Runtime/Experimental/InteractiveToolsFramework/Public/BaseBehaviors/InputBehaviorModifierStates.h
new file mode 100644
--- /dev/null
+++ b/Engine/Source/Runtime/Experimental/InteractiveToolsFramework/Public/BaseBehaviors/InputBehaviorModifierStates.h
@@ -0,0 +1,66 @@
+// Copyright 1998-2019 Epic Games, Inc. All Rights Reserved.
+
+#pragma once
+
+#include "CoreMinimal.h"
+#include "InputState.h"
+#include "BehaviorTargetInterfaces.h"
+
+/**
+ * FInputBehaviorModifierStates is an object that can be placed in an InputBehavior to allow
+ * users of the behavior to request that they be notified about modifier keys/buttons/etc state.
+ * 
+ * We don't know ahead of time what might be used to determine modifier state, which input
+ * devices might be used, etc. So the user has to register (ModifierID,ModifierTestFunction) pairs.
+ * The behavior then calls UpdateModifiers() which will query each of the test functions and
+ * notify the target object about the state of the modifier.
+ */
+class FInputBehaviorModifierStates
+{
+public:
+	typedef TFunction<bool(const FInputDeviceState&)> FModifierTestFunction;
+
+protected:
+	/** List of modifier IDs that have been registered */
+	TArray<int> ModifierIDs;
+	/** The modifier test function for each ID */
+	TMap<int, FModifierTestFunction> ModifierTests;
+
+public:
+	
+	/**
+	 * Register a modifier ID and an associated test function
+	 * @param ModifierID the modifier ID
+	 * @param ModifierTest the test function
+	 */
+	void RegisterModifier(int ModifierID, const FModifierTestFunction& ModifierTest)
+	{
+		check(ModifierIDs.Contains(ModifierID) == false);
+		ModifierIDs.Add(ModifierID);
+		ModifierTests.Add(ModifierID, ModifierTest);
+	}
+
+
+	/**
+	 * @return true if any modifiers have been registered
+	 */
+	bool HasModifiers() const
+	{
+		return (ModifierIDs.Num() > 0);
+	}
+
+
+	/**
+	 * Look up the current state of each registered modifier and pass to the target
+	 * @param Input the current input device state
+	 * @param ModifiersTarget the target that is interested in the modifier values
+	 */
+	void UpdateModifiers(const FInputDeviceState& Input, IModifierToggleBehaviorTarget* ModifiersTarget) const
+	{
+		for (int ModifierID : ModifierIDs)
+		{
+			bool bIsOn = ModifierTests[ModifierID](Input);
+			ModifiersTarget->OnUpdateModifierState(ModifierID, bIsOn);
+		}
+	}
+};
\ No newline at end of file
diff --git a/Engine/Source/Runtime/Experimental/InteractiveToolsFramework/Public/BaseBehaviors/KeyAsModifierInputBehavior.h b/Engine/Source/Runtime/Experimental/InteractiveToolsFramework/Public/BaseBehaviors/KeyAsModifierInputBehavior.h
new file mode 100644
--- /dev/null
+++ b/Engine/Source/Runtime/Experimental/InteractiveToolsFramework/Public/BaseBehaviors/KeyAsModifierInputBehavior.h
@@ -0,0 +1,61 @@
+// Copyright 1998-2019 Epic Games, Inc. All Rights Reserved.
+
+#pragma once
+
+#include "CoreMinimal.h"
+#include "InputBehavior.h"
+#include "InteractiveTool.h"
+#include "BaseBehaviors/BehaviorTargetInterfaces.h"
+#include "BaseBehaviors/InputBehaviorModifierStates.h"
+#include "KeyAsModifierInputBehavior.generated.h"
+
+
+/**
+ * UKeyAsModifierInputBehavior converts a specific key press/release into
+ * a "Modifier" toggle via the IModifierToggleBehaviorTarget interface
+ */
+UCLASS()
+class INTERACTIVETOOLSFRAMEWORK_API UKeyAsModifierInputBehavior : public UInputBehavior
+{
+	GENERATED_BODY()
+
+public:
+	UKeyAsModifierInputBehavior();
+
+	virtual EInputDevices GetSupportedDevices() override
+	{
+		return EInputDevices::Keyboard;
+	}
+
+	/**
+	 * Initialize this behavior with the given Target
+	 * @param Target implementor of modifier-toggle behavior
+	 * @param ModifierID integer ID that identifiers the modifier toggle
+	 * @param ModifierKey the key that will be used as the modifier toggle
+	 */
+	virtual void Initialize(IModifierToggleBehaviorTarget* Target, int ModifierID, const FKey& ModifierKey);
+
+	/**
+	 * WantsCapture() will only return capture request if this function returns true (or is null)
+	 * Intended to be used for alt/ctrl/cmd/shift modifiers on the main ModifierKey
+	 */
+	TFunction<bool(const FInputDeviceState&)> ModifierCheckFunc = nullptr;
+
+	// UInputBehavior implementation
+
+	virtual FInputCaptureRequest WantsCapture(const FInputDeviceState& Input) override;
+	virtual FInputCaptureUpdate BeginCapture(const FInputDeviceState& Input, EInputCaptureSide Side) override;
+	virtual FInputCaptureUpdate UpdateCapture(const FInputDeviceState& Input, const FInputCaptureData& Data) override;
+	virtual void ForceEndCapture(const FInputCaptureData& Data) override;
+
+
+protected:
+	/** Modifier Target object */
+	IModifierToggleBehaviorTarget* Target;
+
+	/** Key that is used as modifier */
+	FKey ModifierKey;
+
+	/** Modifier set for this behavior, internally initialized with check on ModifierKey */
+	FInputBehaviorModifierStates Modifiers;
+};
diff --git a/Engine/Source/Runtime/Experimental/InteractiveToolsFramework/Public/BaseBehaviors/MouseHoverBehavior.h b/Engine/Source/Runtime/Experimental/InteractiveToolsFramework/Public/BaseBehaviors/MouseHoverBehavior.h
new file mode 100644
--- /dev/null
+++ b/Engine/Source/Runtime/Experimental/InteractiveToolsFramework/Public/BaseBehaviors/MouseHoverBehavior.h
@@ -0,0 +1,45 @@
+// Copyright 1998-2019 Epic Games, Inc. All Rights Reserved.
+
+#pragma once
+
+#include "CoreMinimal.h"
+#include "InputBehavior.h"
+#include "InteractiveTool.h"
+#include "BaseBehaviors/BehaviorTargetInterfaces.h"
+#include "BaseBehaviors/InputBehaviorModifierStates.h"
+#include "MouseHoverBehavior.generated.h"
+
+
+
+/**
+ * Trivial InputBehavior that forwards InputBehavior hover events to a Target object via
+ * the IHoverBehaviorTarget interface.
+ */
+UCLASS()
+class INTERACTIVETOOLSFRAMEWORK_API UMouseHoverBehavior : public UInputBehavior
+{
+	GENERATED_BODY()
+
+public:
+	UMouseHoverBehavior();
+
+	/**
+	 * The modifier set for this behavior
+	 */
+	FInputBehaviorModifierStates Modifiers;
+
+	virtual void Initialize(IHoverBehaviorTarget* Target);
+
+	// UInputBehavior hover implementation
+
+	virtual EInputDevices GetSupportedDevices() override;
+
+	virtual bool WantsHoverEvents() override;
+	virtual void UpdateHover(const FInputDeviceState& input) override;
+	virtual void EndHover(const FInputDeviceState& input) override;
+
+protected:
+	IHoverBehaviorTarget* Target;
+};
+
+
diff --git a/Engine/Source/Runtime/Experimental/InteractiveToolsFramework/Public/BaseBehaviors/MultiClickSequenceInputBehavior.h b/Engine/Source/Runtime/Experimental/InteractiveToolsFramework/Public/BaseBehaviors/MultiClickSequenceInputBehavior.h
new file mode 100644
--- /dev/null
+++ b/Engine/Source/Runtime/Experimental/InteractiveToolsFramework/Public/BaseBehaviors/MultiClickSequenceInputBehavior.h
@@ -0,0 +1,74 @@
+// Copyright 1998-2019 Epic Games, Inc. All Rights Reserved.
+
+#pragma once
+
+#include "CoreMinimal.h"
+#include "BaseBehaviors/AnyButtonInputBehavior.h"
+#include "BehaviorTargetInterfaces.h"
+#include "BaseBehaviors/InputBehaviorModifierStates.h"
+#include "MultiClickSequenceInputBehavior.generated.h"
+
+
+/**
+ * UMultiClickSequenceInputBehavior implements a generic multi-click-sequence input behavior.
+ * For example this behavior could be used to implement a multi-click polygon-drawing interaction.
+ *
+ * The internal state machine works as follows:
+ *   1) on input-device-button-press, check if target wants to begin sequence. If so, begin capture.
+ *   2) on button *release*, check if target wants to continue or terminate sequence
+ *       a) if terminate, release capture
+ *       b) if continue, do nothing (capture continues between presses)
+ *
+ * The target will receive "preview" notifications (basically hover) during updates where there is
+ * not a release. This can be used to (eg) update a rubber-band selection end point
+ * 
+ * @todo it may be better to implement this as multiple captures, and use hover callbacks to 
+ * do the between-capture previews. holding capture across mouse release is not ideal.
+ */
+UCLASS()
+class INTERACTIVETOOLSFRAMEWORK_API UMultiClickSequenceInputBehavior : public UAnyButtonInputBehavior
+{
+	GENERATED_BODY()
+
+public:
+	UMultiClickSequenceInputBehavior();
+
+	/**
+	 * The modifier set for this behavior
+	 */
+	FInputBehaviorModifierStates Modifiers;
+
+	/**
+	 * Initialize this behavior with the given Target
+	 * @param Target implementor of hit-test and on-clicked functions
+	 */
+	virtual void Initialize(IClickSequenceBehaviorTarget* Target);
+
+
+	/**
+	 * WantsCapture() will only return capture request if this function returns true (or is null)
+	 */
+	TFunction<bool(const FInputDeviceState&)> ModifierCheckFunc = nullptr;
+
+
+	// UInputBehavior implementation
+
+	virtual FInputCaptureRequest WantsCapture(const FInputDeviceState& Input) override;
+	virtual FInputCaptureUpdate BeginCapture(const FInputDeviceState& Input, EInputCaptureSide eSide) override;
+	virtual FInputCaptureUpdate UpdateCapture(const FInputDeviceState& Input, const FInputCaptureData& Data) override;
+	virtual void ForceEndCapture(const FInputCaptureData& Data) override;
+
+	virtual bool WantsHoverEvents() override;
+	virtual void UpdateHover(const FInputDeviceState& input) override;
+	virtual void EndHover(const FInputDeviceState& input) override;
+
+
+public:
+
+protected:
+	/** Click Target object */
+	IClickSequenceBehaviorTarget* Target;
+
+	bool bInActiveSequence;
+};
+
diff --git a/Engine/Source/Runtime/Experimental/InteractiveToolsFramework/Public/BaseBehaviors/SingleClickBehavior.h b/Engine/Source/Runtime/Experimental/InteractiveToolsFramework/Public/BaseBehaviors/SingleClickBehavior.h
new file mode 100644
--- /dev/null
+++ b/Engine/Source/Runtime/Experimental/InteractiveToolsFramework/Public/BaseBehaviors/SingleClickBehavior.h
@@ -0,0 +1,67 @@
+// Copyright 1998-2019 Epic Games, Inc. All Rights Reserved.
+
+#pragma once
+
+#include "CoreMinimal.h"
+#include "BaseBehaviors/AnyButtonInputBehavior.h"
+#include "BehaviorTargetInterfaces.h"
+#include "SingleClickBehavior.generated.h"
+
+
+/**
+ * USingleClickInputBehavior implements a standard "button-click"-style input behavior.
+ * The state machine works as follows:
+ *    1) on input-device-button-press, hit-test the target. If hit, begin capture
+ *    2) on input-device-button-release, hit-test the target. If hit, call Target::OnClicked(). If not hit, ignore click.
+ *    
+ * The second hit-test is required to allow the click to be "cancelled" by moving away
+ * from the target. This is standard GUI behavior. You can disable this second hit test
+ * using the .HitTestOnRelease property. This is strongly discouraged.
+ *    
+ * The hit-test and on-clicked functions are provided by a IClickBehaviorTarget instance.
+ */
+UCLASS()
+class INTERACTIVETOOLSFRAMEWORK_API USingleClickInputBehavior : public UAnyButtonInputBehavior
+{
+	GENERATED_BODY()
+
+public:
+	USingleClickInputBehavior();
+
+	/**
+	 * Initialize this behavior with the given Target
+	 * @param Target implementor of hit-test and on-clicked functions
+	 */
+	virtual void Initialize(IClickBehaviorTarget* Target);
+
+
+	/**
+	 * WantsCapture() will only return capture request if this function returns true (or is null)
+	 */
+	TFunction<bool(const FInputDeviceState&)> ModifierCheckFunc = nullptr;
+
+
+	// UInputBehavior implementation
+
+	virtual FInputCaptureRequest WantsCapture(const FInputDeviceState& Input) override;
+	virtual FInputCaptureUpdate BeginCapture(const FInputDeviceState& Input, EInputCaptureSide eSide) override;
+	virtual FInputCaptureUpdate UpdateCapture(const FInputDeviceState& Input, const FInputCaptureData& Data) override;
+	virtual void ForceEndCapture(const FInputCaptureData& Data) override;
+
+
+public:
+	/** Hit-test is repeated on release (standard behavior). If false, */
+	UPROPERTY()
+	bool HitTestOnRelease;
+
+
+protected:
+	/** Click Target object */
+	IClickBehaviorTarget* Target;
+
+	/**
+	 * Internal function that forwards click evens to Target::OnClicked, you can customize behavior here
+	 */
+	virtual void Clicked(const FInputDeviceState& Input, const FInputCaptureData& Data);
+};
+
diff --git a/Engine/Source/Runtime/Experimental/InteractiveToolsFramework/Public/BaseTools/ClickDragTool.h b/Engine/Source/Runtime/Experimental/InteractiveToolsFramework/Public/BaseTools/ClickDragTool.h
new file mode 100644
--- /dev/null
+++ b/Engine/Source/Runtime/Experimental/InteractiveToolsFramework/Public/BaseTools/ClickDragTool.h
@@ -0,0 +1,81 @@
+// Copyright 1998-2019 Epic Games, Inc. All Rights Reserved.
+
+#pragma once
+
+#include "CoreMinimal.h"
+#include "InteractiveTool.h"
+#include "InteractiveToolBuilder.h"
+#include "BaseBehaviors/BehaviorTargetInterfaces.h"
+#include "ClickDragTool.generated.h"
+
+
+
+/**
+ * Builder for UClickDragTool
+ */
+UCLASS()
+class INTERACTIVETOOLSFRAMEWORK_API UClickDragToolBuilder : public UInteractiveToolBuilder
+{
+	GENERATED_BODY()
+
+public:
+	virtual bool CanBuildTool(const FToolBuilderState& SceneState) const override;
+	virtual UInteractiveTool* BuildTool(const FToolBuilderState& SceneState) const override;
+};
+
+
+/**
+ * UClickDragTool is a base tool that basically just implements IClickDragBehaviorTarget,
+ * and on setup registers a UClickDragInputBehavior. You can subclass this Tool to
+ * implement basic click-drag type Tools. If you want to do more advanced things, 
+ * like handle modifier buttons/keys, you will need to implement IClickDragBehaviorTarget yourself
+ */
+UCLASS()
+class INTERACTIVETOOLSFRAMEWORK_API UClickDragTool : public UInteractiveTool, public IClickDragBehaviorTarget
+{
+	GENERATED_BODY()
+
+public:
+
+	/**
+	 * Register default primary-button-click InputBehaviors
+	 */
+	virtual void Setup() override;
+
+
+	//
+	// IClickBehaviorTarget implementation
+	//
+
+	/**
+	 * Test if target can begin click-drag interaction at this point
+	 * @param ClickPos device position/ray at click point
+	 * @return true if target wants to begin sequence
+	 */
+	virtual bool CanBeginClickDragSequence(const FInputDeviceRay& PressPos) override;
+
+	/**
+	 * Notify Target that click press ocurred
+	 * @param ClickPos device position/ray at click point
+	 */
+	virtual void OnClickPress(const FInputDeviceRay& PressPos) override;
+
+	/**
+	 * Notify Target that input position has changed
+	 * @param ClickPos device position/ray at click point
+	 */
+	virtual void OnClickDrag(const FInputDeviceRay& DragPos) override;
+
+	/**
+	 * Notify Target that click release occurred
+	 * @param ClickPos device position/ray at click point
+	 */
+	virtual void OnClickRelease(const FInputDeviceRay& ReleasePos) override;
+
+	/**
+	 * Notify Target that click-drag sequence has been explicitly terminated (eg by escape key)
+	 */
+	virtual void OnTerminateDragSequence() override;
+
+};
+
diff --git a/Engine/Source/Runtime/Experimental/InteractiveToolsFramework/Public/BaseTools/MeshSurfacePointTool.h b/Engine/Source/Runtime/Experimental/InteractiveToolsFramework/Public/BaseTools/MeshSurfacePointTool.h
new file mode 100644
--- /dev/null
+++ b/Engine/Source/Runtime/Experimental/InteractiveToolsFramework/Public/BaseTools/MeshSurfacePointTool.h
@@ -0,0 +1,144 @@
+// Copyright 1998-2019 Epic Games, Inc. All Rights Reserved.
+
+#pragma once
+
+#include "CoreMinimal.h"
+#include "InteractiveTool.h"
+#include "InteractiveToolBuilder.h"
+#include "InputBehaviorSet.h"
+#include "BaseBehaviors/BehaviorTargetInterfaces.h"
+#include "BaseBehaviors/AnyButtonInputBehavior.h"
+#include "MeshSurfacePointTool.generated.h"
+
+
+
+class UMeshSurfacePointTool;
+
+/**
+ * 
+ */
+UCLASS()
+class INTERACTIVETOOLSFRAMEWORK_API UMeshSurfacePointToolBuilder : public UInteractiveToolBuilder
+{
+	GENERATED_BODY()
+
+public:
+	/** @return true if a single mesh source can be found in the active selection */
+	virtual bool CanBuildTool(const FToolBuilderState& SceneState) const override;
+
+	/** @return new Tool instance initialized with selected mesh source */
+	virtual UInteractiveTool* BuildTool(const FToolBuilderState& SceneState) const override;
+
+	/** @return new Tool instance. Override this in subclasses to build a different Tool class type */
+	virtual UMeshSurfacePointTool* CreateNewTool(const FToolBuilderState& SceneState) const;
+
+	/** Called by BuildTool to configure the Tool with the input MeshSource based on the SceneState */
+	virtual void InitializeNewTool(UMeshSurfacePointTool* Tool, const FToolBuilderState& SceneState) const;
+};
+
+
+
+/**
+ * UMeshSurfacePointTool is a base Tool implementation that can be used to implement various
+ * "point on surface" interactions. The tool acts on an input IMeshDescriptionSource object,
+ * which the standard Builder can extract from the current selection (eg Editor selection).
+ * 
+ * Subclasses override the OnBeginDrag/OnUpdateDrag/OnEndDrag and OnUpdateHover functions
+ * to implement custom behavior.
+ */
+UCLASS()
+class INTERACTIVETOOLSFRAMEWORK_API UMeshSurfacePointTool : public UInteractiveTool, public IHoverBehaviorTarget
+{
+	GENERATED_BODY()
+
+public:
+	// UInteractiveTool API Implementation
+
+	/** Register InputBehaviors, etc */
+	virtual void Setup() override;
+
+
+	// UMeshSurfacePointTool API
+
+	/**
+	 * Set the Target MeshDescription-providing Object for this tool
+	 * @param MeshSource object that can provide a MeshDescriptionSource. Note UMeshSurfacePointTool
+	 *        only calls MeshSource->HitTest, so this source technically does not have to actually provide
+	 *        a valid MeshDescription
+	 */
+	virtual void SetMeshSource(TUniquePtr<IMeshDescriptionSource> MeshSource);
+
+	/**
+	 * @return true if the target MeshSource is hit by the Ray
+	 */
+	virtual bool HitTest(const FRay& Ray, FHitResult& OutHit);
+
+	
+	/**
+	 * This function is called by registered InputBehaviors when the user begins a click-drag-release interaction
+	 */
+	virtual void OnBeginDrag(const FRay& Ray);
+
+	/**
+	 * This function is called by registered InputBehaviorseach frame that the user is in a click-drag-release interaction
+	 */
+	virtual void OnUpdateDrag(const FRay& Ray);
+
+	/**
+	 * This function is called by registered InputBehaviors when the user releases the button driving a click-drag-release interaction
+	 */
+	virtual void OnEndDrag(const FRay& Ray);
+
+	/**
+	 * This function is called each frame while the tool is active, to support hover updates
+	 */
+	virtual void OnUpdateHover(const FInputDeviceRay& DevicePos) {}
+
+
+
+	/** Called by registered InputBehaviors to set the state of the "shift" button (or device equivalent) */
+	virtual void SetShiftToggle(bool bShiftDown);
+
+	/** @return current state of the shift toggle */
+	virtual bool GetShiftToggle() { return bShiftToggle; }
+
+
+protected:
+	/** Target source object that provides a MeshDescription and various other query API calls */
+	TUniquePtr<IMeshDescriptionSource> MeshSource;
+
+	/** Current state of the shift modifier toggle */
+	bool bShiftToggle;
+};
+
+
+
+
+
+
+
+
+/**
+ * UMeshSurfacePointToolMouseBehavior implements mouse press-drag-release interaction behavior for Mouse devices.
+ * You can configure the base UAnyButtonInputBehavior to change the mouse button in use (default = left mouse)
+ */
+UCLASS()
+class INTERACTIVETOOLSFRAMEWORK_API UMeshSurfacePointToolMouseBehavior : public UAnyButtonInputBehavior
+{
+	GENERATED_BODY()
+
+public:
+	virtual void Initialize(UMeshSurfacePointTool* Tool);
+
+	virtual FInputCaptureRequest WantsCapture(const FInputDeviceState& input) override;
+	virtual FInputCaptureUpdate BeginCapture(const FInputDeviceState& input, EInputCaptureSide eSide) override;
+	virtual FInputCaptureUpdate UpdateCapture(const FInputDeviceState& input, const FInputCaptureData& data) override;
+	virtual void ForceEndCapture(const FInputCaptureData& data) override;
+
+protected:
+	UMeshSurfacePointTool* Tool;
+	FRay LastWorldRay;
+	bool bInDragCapture;
+};
+
+
diff --git a/Engine/Source/Runtime/Experimental/InteractiveToolsFramework/Public/BaseTools/SingleClickTool.h b/Engine/Source/Runtime/Experimental/InteractiveToolsFramework/Public/BaseTools/SingleClickTool.h
new file mode 100644
--- /dev/null
+++ b/Engine/Source/Runtime/Experimental/InteractiveToolsFramework/Public/BaseTools/SingleClickTool.h
@@ -0,0 +1,62 @@
+// Copyright 1998-2019 Epic Games, Inc. All Rights Reserved.
+
+#pragma once
+
+#include "CoreMinimal.h"
+#include "InteractiveTool.h"
+#include "InteractiveToolBuilder.h"
+#include "BaseBehaviors/SingleClickBehavior.h"
+#include "SingleClickTool.generated.h"
+
+
+
+/**
+ * Builder for USingleClickTool
+ */
+UCLASS()
+class INTERACTIVETOOLSFRAMEWORK_API USingleClickToolBuilder : public UInteractiveToolBuilder
+{
+	GENERATED_BODY()
+
+public:
+	virtual bool CanBuildTool(const FToolBuilderState& SceneState) const override;
+	virtual UInteractiveTool* BuildTool(const FToolBuilderState& SceneState) const override;
+};
+
+
+/**
+ * USingleClickTool is perhaps the simplest possible interactive tool. It simply
+ * reacts to default primary button clicks for the active device (eg left-mouse clicks).
+ * 
+ * The function ::IsHitByClick() determines what is clickable by this Tool. The default is
+ * to return true, which means the click will activate anywhere (the Tool itself has no
+ * notion of Actors, Components, etc). You can override this function to, for example,
+ * filter out clicks that don't hit a target object, etc.
+ * 
+ * The function ::OnClicked() implements the action that will occur when a click happens.
+ * You must override this to implement any kind of useful behavior.
+ */
+UCLASS()
+class INTERACTIVETOOLSFRAMEWORK_API USingleClickTool : public UInteractiveTool, public IClickBehaviorTarget
+{
+	GENERATED_BODY()
+
+public:
+
+	/**
+	 * Register default primary-button-click InputBehaviors
+	 */
+	virtual void Setup() override;
+
+	/**
+	 * Test if the Target is hit at this 2D position / 3D ray
+	 */
+	virtual bool IsHitByClick(const FInputDeviceRay& ClickPos);
+
+	/**
+	 * Click the Target at this 2D position / 3D ray. Default behavior is to print debug string.
+	 */
+	virtual void OnClicked(const FInputDeviceRay& ClickPos);
+
+};
+
diff --git a/Engine/Source/Runtime/Experimental/InteractiveToolsFramework/Public/ComponentSourceInterfaces.h b/Engine/Source/Runtime/Experimental/InteractiveToolsFramework/Public/ComponentSourceInterfaces.h
new file mode 100644
--- /dev/null
+++ b/Engine/Source/Runtime/Experimental/InteractiveToolsFramework/Public/ComponentSourceInterfaces.h
@@ -0,0 +1,100 @@
+// Copyright 1998-2019 Epic Games, Inc. All Rights Reserved.
+
+#pragma once
+
+#include "CoreMinimal.h"
+#include "Math/UnrealMath.h"
+#include "Engine/EngineTypes.h"    // FHitResult
+
+// predeclarations
+class AActor;
+class UActorComponent;
+struct FMeshDescription;
+class UMaterialInterface;
+
+
+
+/**
+ * Wrapper around a UObject Component that can provide a MeshDescription, and
+ * (optionally) bake a modified MeshDescription back to this Component.
+ * An example of a Source might be a StaticMeshComponent. How a modified
+ * MeshDescription is committed back is context-dependent (in Editor vs PIE vs Runtime, etc).
+ * 
+ * (Conceivably this doesn't have to be backed by a Component, but most usage will assume there is an Actor)
+ */
+class IMeshDescriptionSource
+{
+public:
+	virtual ~IMeshDescriptionSource() {}
+
+	
+	/** @return the Actor that owns this Component */
+	virtual AActor* GetOwnerActor() const = 0;
+
+	/** @return the Component this is a Source for */
+	virtual UActorComponent* GetOwnerComponent() const = 0;
+
+	/** @return Pointer to the MeshDescription this Source is providing */
+	virtual FMeshDescription* GetMeshDescription() const = 0;
+
+	/**
+	 * Get pointer to a Material provided by this Source
+	 * @param MaterialIndex index of the material
+	 * @return MaterialInterface pointer, or null if MaterialIndex is invalid
+	 */
+	virtual UMaterialInterface* GetMaterial(int32 MaterialIndex) const = 0;
+
+	/** 
+	 * @return the transform on this component
+	 * @todo Do we need to return a list of transforms here?
+	 */
+	virtual FTransform GetWorldTransform() const = 0;
+
+	/**
+	 * Compute ray intersection with the MeshDescription this Source is providing
+	 * @param WorldRay ray in world space
+	 * @param OutHit hit test data
+	 * @return true if ray intersected Component
+	 */
+	virtual bool HitTest(const FRay& WorldRay, FHitResult& OutHit) const = 0;
+
+	/**
+	 * Set the visibility of the Component associated with this Source (ie to hide during Tool usage)
+	 * @param bVisible desired visibility
+	 */
+	virtual void SetOwnerVisibility(bool bVisible) const = 0;
+
+
+	/** @return true if this Source is read-only, ie Commit functions cannot be called */
+	virtual bool IsReadOnly() const { return true; }
+
+	/**
+	 * Call this to modify the MeshDescription provided by the Source. You provide a callback
+	 * function that the Source will immediately call with a suitable MeshDescription instance.
+	 * The Source will then update the Component as necessary. The Source is responsible for
+	 * making sure that it is safe to modify this MeshDescription instance.
+	 * @param ModifyFunction callback function that updates/modifies a MeshDescription
+	 */
+	virtual void CommitInPlaceModification(const TFunction<void(FMeshDescription*)>& ModifyFunction) { check(false); }
+};
+
+
+
+
+
+/**
+ * Interface to a Factory that knows how to build Sources for UObjects. 
+ */
+class IComponentSourceFactory
+{
+public:
+	virtual ~IComponentSourceFactory() {}
+
+	/**
+	 * Create a MeshDescription source for the given Component
+	 * @param Component A UObject that can provide a MeshDescription. Assumption is this is a Component of an Actor.
+	 * @return A MeshDescriptionSource instance. Must not return null.
+	 */
+	virtual TUniquePtr<IMeshDescriptionSource> MakeMeshDescriptionSource(UActorComponent* Component) = 0;
+};
+
diff --git a/Engine/Source/Runtime/Experimental/InteractiveToolsFramework/Public/InputBehavior.h b/Engine/Source/Runtime/Experimental/InteractiveToolsFramework/Public/InputBehavior.h
new file mode 100644
--- /dev/null
+++ b/Engine/Source/Runtime/Experimental/InteractiveToolsFramework/Public/InputBehavior.h
@@ -0,0 +1,301 @@
+// Copyright 1998-2019 Epic Games, Inc. All Rights Reserved.
+
+#pragma once
+
+#include "CoreMinimal.h"
+#include "InputState.h"
+#include "Math/NumericLimits.h"
+#include "InputBehavior.generated.h"
+
+
+
+/**
+ * Input can be captured separately for Left and Right sides (eg for VR controllers)
+ * Currently mouse is Left.
+ */
+UENUM()
+enum class EInputCaptureSide
+{
+	None = 0,
+	Left = 1,
+	Right = 2,
+	Both = 3,
+	Any = 99
+};
+
+
+/**
+ * An active capturing behavior may need to keep track of additional data that
+ * cannot be stored within the behavior (for example if the same behavior instance
+ * is capturing for Left and Right separately). So FInputCaptureUpdate can optionally
+ * return this structure, and we will pass it to the next UpdateCapture() call
+ */
+struct INTERACTIVETOOLSFRAMEWORK_API FInputCaptureData
+{
+	/** Which side do we want to capture on */
+	EInputCaptureSide WhichSide;
+	/** pointer to data defined by the InputBehavior, which is also responsible for cleaning it up */
+	void* CustomData;
+
+	FInputCaptureData()
+	{
+		WhichSide = EInputCaptureSide::None;
+		CustomData = nullptr;
+	}
+};
+
+
+/**
+ * Used by FInputCaptureRequest to indicate whether the InputBehavior
+ * wants to capture or ignore an input event
+ */
+UENUM()
+enum class EInputCaptureRequestType
+{
+	Begin = 1,
+	Ignore = 2
+};
+
+
+// predeclaration
+class UInputBehavior;
+
+
+/**
+ * UInputBehavior returns an FInputCaptureRequest from WantsCapture() to indicate
+ * whether it wants to capture or ignore an input event
+ */
+struct INTERACTIVETOOLSFRAMEWORK_API FInputCaptureRequest
+{
+	/** Which input behavior generated this request */
+	UInputBehavior* Source;
+    /** What type of capture request is this (Begin or Ignore) */
+	EInputCaptureRequestType Type;
+	/** Which side does request want to capture on */
+	EInputCaptureSide Side;
+
+	/** Depth along hit-test ray */
+	float HitDepth;
+
+	/** Owner of the requesting behavior. Behavior doesn't know this, so this is initialized to null */
+	void* Owner;
+
+
+	FInputCaptureRequest(EInputCaptureRequestType type, UInputBehavior* behavior, EInputCaptureSide whichSide, float hitDepth = TNumericLimits<float>::Max() )
+	{
+		this->Type = type;
+		this->Source = behavior;
+		this->Side = whichSide;
+		this->Owner = nullptr;
+		this->HitDepth = hitDepth;
+	}
+
+	/** Create a Begin-capture request */
+	static FInputCaptureRequest Begin(UInputBehavior* behavior, EInputCaptureSide whichSide, float hitDepth = TNumericLimits<float>::Max() )
+	{
+		return FInputCaptureRequest(EInputCaptureRequestType::Begin, behavior, whichSide, hitDepth);
+	}
+
+	/** Create an ignore-capture request */
+	static FInputCaptureRequest Ignore() 
+	{
+		return FInputCaptureRequest(EInputCaptureRequestType::Ignore, nullptr, EInputCaptureSide::Any, TNumericLimits<float>::Max() );
+	}
+
+	friend bool operator<(const FInputCaptureRequest& l, const FInputCaptureRequest& r);
+};
+
+
+
+
+
+
+/**
+ * FInputCaptureUpdate uses this type to indicate what state the capturing Behavior
+ * would like to transition to, based on the input event
+ */
+UENUM()
+enum class EInputCaptureState
+{
+	Begin = 1,		// start capturing (which should always be the case if BeginCapture is called)
+	Continue = 2,   // Behavior wants to continue capturing
+	End = 3,        // Behavior wants to end capturing
+	Ignore = 4      // Behavior ignored this event
+};
+
+
+
+/**
+ * IInputBehavior returns an FInputCaptureUpdate from BeginCapture() and UpdateCapture(),
+ * which indicates to the InputRouter what the Behavior would like to have happen.
+ */
+struct INTERACTIVETOOLSFRAMEWORK_API FInputCaptureUpdate
+{
+	/** Indicates what capture state the Behavior wants to transition to */
+	EInputCaptureState State;
+	/** Which Behavior did this update come from */
+	UInputBehavior* Source;
+	/** custom data for the active capture that should be propagated to next UpdateCapture() call */
+	FInputCaptureData Data;
+
+	/** 
+	 * Create a begin-capturing instance of FInputCaptureUpdate
+	 * @param Source UInputBehavior that is returning this update
+	 * @param Which Which side we are capturing on
+	 * @param CustomData client-provided data that will be passed to UInputBehavior::UpdateCapture() calls. Client owns this memory!
+	 */
+	static FInputCaptureUpdate Begin(UInputBehavior* SourceBehavior, EInputCaptureSide WhichSide, void* CustomData = nullptr)
+	{
+		return FInputCaptureUpdate(EInputCaptureState::Begin, SourceBehavior, WhichSide, CustomData);
+	}
+
+	/** Create a default continue-capturing instance of FInputCaptureUpdate */
+	static FInputCaptureUpdate Continue()
+	{
+		return FInputCaptureUpdate(EInputCaptureState::Continue, nullptr, EInputCaptureSide::Any);
+	}
+	/** Create a default end-capturing instance of FInputCaptureUpdate */
+	static FInputCaptureUpdate End()
+	{
+		return FInputCaptureUpdate(EInputCaptureState::End, nullptr, EInputCaptureSide::Any);
+	}
+	/** Create a default ignore-capturing instance of FInputCaptureUpdate */
+	static FInputCaptureUpdate Ignore()
+	{
+		return FInputCaptureUpdate(EInputCaptureState::Ignore, nullptr, EInputCaptureSide::Any);
+	}
+
+	/**
+	 * @param StateIn desired capture state
+	 * @param Source UInputBehavior that is returning this update
+	 * @param Which Which side we are capturing on
+	 * @param CustomData client-provided data that will be passed to UInputBehavior::UpdateCapture() calls. Client owns this memory!
+	 */
+	FInputCaptureUpdate(EInputCaptureState StateIn, UInputBehavior* SourceBehaviorIn, EInputCaptureSide WhichSideIn, void* CustomData = nullptr)
+	{
+		State = StateIn;
+		Source = SourceBehaviorIn;
+		Data.WhichSide = WhichSideIn;
+		Data.CustomData = CustomData;
+	}
+};
+
+
+
+/**
+ * Each UInputBehavior provides a priority that is used to help resolve situations
+ * when multiple Behaviors want to capture based on the same input event
+ */
+struct INTERACTIVETOOLSFRAMEWORK_API FInputCapturePriority
+{
+	static constexpr int DEFAULT_GIZMO_PRIORITY = 50;
+	static constexpr int DEFAULT_TOOL_PRIORITY = 100;
+
+	/** Constant priority value */
+	int Priority;
+
+	FInputCapturePriority(int priority = DEFAULT_TOOL_PRIORITY)
+	{
+		Priority = priority;
+	}
+
+	/** @return a priority lower than this priority */
+	FInputCapturePriority MakeLower(int DeltaAmount = 1) const
+	{
+		return FInputCapturePriority(Priority + DeltaAmount);
+	}
+
+	/** @return a priority higher than this priority */
+	FInputCapturePriority MakeHigher(int DeltaAmount = 1) const
+	{
+		return FInputCapturePriority(Priority - DeltaAmount);
+	}
+
+	friend bool operator<(const FInputCapturePriority& l, const FInputCapturePriority& r)
+	{
+		return l.Priority < r.Priority;
+	}
+	friend bool operator==(const FInputCapturePriority& l, const FInputCapturePriority& r)
+	{
+		return l.Priority == r.Priority;
+	}
+};
+
+
+
+/**
+ * An InputBehavior implements a state machine for a user interaction. 
+ * The InputRouter maintains a set of active Behaviors, and when new input
+ * events occur, it calls WantsCapture() to check if the Behavior would like to
+ * begin capturing the applicable input event stream (eg for a mouse, one or both VR controllers, etc).
+ * If the Behavior acquires capture, UpdateCapture() is called until the Behavior
+ * indicates that it wants to release the device, or until the InputRouter force-terminates
+ * the capture via ForceEndCapture().
+ *
+ * For example, something like ButtonSetClickBehavior might work as follows:
+ *    - in WantsCapture(), if left mouse is pressed and a button is under cursor, return Begin, otherwise Ignore
+ *    - in BeginCapture(), save identifier for button that is under cursor
+ *    - in UpdateCapture()
+ *        - if left mouse is down, return Continue
+ *        - if left mouse is released:
+ *            - if saved button is still under cursor, call button.Clicked()
+ *            - return End
+ *
+ * Written sufficiently generically, the above Behavior doesn't need to know about buttons,
+ * it just needs to know how to hit-test the clickable object(s). Similarly separate 
+ * Behaviors can be written for mouse, VR, touch, gamepad, etc. 
+ *
+ * Implementing interactions in this way allows the input handling to be separated from functionality.
+ */
+UCLASS(Transient)
+class INTERACTIVETOOLSFRAMEWORK_API UInputBehavior : public UObject
+{
+	GENERATED_BODY()
+
+public:
+	UInputBehavior();
+
+	/** The priority is used to resolve situations where multiple behaviors want the same capture */
+	virtual FInputCapturePriority GetPriority();
+
+	/** Configure the default priority of an instance of this behavior */
+	virtual void SetDefaultPriority(const FInputCapturePriority& Priority);
+
+
+	/** Which device types does this Behavior support */
+	virtual EInputDevices GetSupportedDevices();
+
+	/** Given the input state, does this Behavior want to begin capturing some input devices? */
+	virtual FInputCaptureRequest WantsCapture(const FInputDeviceState& InputState);
+
+	/** Called after WantsCapture() returns a capture request that was accepted */
+	virtual FInputCaptureUpdate BeginCapture(const FInputDeviceState& InputState, EInputCaptureSide eSide);
+
+	/** 
+	 * Called for each new input event during a capture sequence. Return Continue to keep
+	 * capturing, or End to finish capturing.
+	 */
+	virtual FInputCaptureUpdate UpdateCapture(const FInputDeviceState& InputState, const FInputCaptureData& CaptureData);
+
+	/** If this is called, the Behavior has forcibly lost capture (eg due to app losing focus for example) and needs to clean up accordingly */
+	virtual void ForceEndCapture(const FInputCaptureData& CaptureData);
+
+
+	//
+	// hover support (optional)
+	//
+
+	/** return true if this Behavior supports hover (ie passive input events) */
+	virtual bool WantsHoverEvents();
+
+	/** called on each new hover input event, ie if no other behavior is actively capturing input */
+	virtual void UpdateHover(const FInputDeviceState& InputState);
+
+	/** if a capture begins or focus is lost, any active hover visualization needs to terminate */
+	virtual void EndHover(const FInputDeviceState& InputState);
+
+
+protected:
+	/** priority returned by GetPriority() */
+	FInputCapturePriority DefaultPriority;
+};
\ No newline at end of file
diff --git a/Engine/Source/Runtime/Experimental/InteractiveToolsFramework/Public/InputBehaviorSet.h b/Engine/Source/Runtime/Experimental/InteractiveToolsFramework/Public/InputBehaviorSet.h
new file mode 100644
--- /dev/null
+++ b/Engine/Source/Runtime/Experimental/InteractiveToolsFramework/Public/InputBehaviorSet.h
@@ -0,0 +1,179 @@
+// Copyright 1998-2019 Epic Games, Inc. All Rights Reserved.
+
+#pragma once
+
+#include "CoreMinimal.h"
+#include "UObject/Interface.h"
+#include "InputBehavior.h"
+#include "InputBehaviorSet.generated.h"
+
+
+
+/**
+ * This is an internal structure used by UInputBehaviorSet. 
+ */
+USTRUCT()
+struct INTERACTIVETOOLSFRAMEWORK_API FBehaviorInfo
+{
+	GENERATED_BODY()
+	
+	/** Reference to a Behavior */
+	UPROPERTY()
+	UInputBehavior* Behavior;
+
+	/** Source object that provided this Behavior */
+	void* Source;
+
+	/** Group identifier for this Behavior */
+	FString Group;
+
+	friend bool operator<(const FBehaviorInfo& l, const FBehaviorInfo& r)
+	{
+		return l.Behavior->GetPriority() < r.Behavior->GetPriority();
+	}
+};
+
+
+
+
+/**
+ * UInputBehaviorSet manages a set of UInputBehaviors, and provides various functions
+ * to query and forward events to the set. Tools and Widgets provide instances of this via
+ * IInputBehaviorSource, and UInputRouter collects and manages them (see comments there)
+ *
+ * Behaviors in the set each have a source pointer and group tag, which allows sets of
+ * behaviors to be managed together. For example one UInputBehaviorSet can be merged into 
+ * another and removed later.
+ */
+UCLASS(Transient)
+class INTERACTIVETOOLSFRAMEWORK_API UInputBehaviorSet : public UObject
+{
+	GENERATED_BODY()
+
+public:
+	UInputBehaviorSet();
+	virtual ~UInputBehaviorSet();
+
+	//
+	// Set Management
+	//
+
+
+	/** @return true if there are no Behaviors in set */
+	virtual bool IsEmpty() const;
+
+	/** 
+	 * Add a Behavior to the set
+	 * @param Behavior Behavior to add to set
+	 * @param Source pointer to owning object, used only to identify Behavior later
+	 * @param GroupName string identifier for this Behavior or group of Behaviors
+	 */
+	virtual void Add(UInputBehavior* Behavior, void* Source = nullptr, const FString& GroupName = "");
+
+	/**
+	 * Merge another BehaviorSet into this Set
+	 * @param OtherSet Set of Behaviors to add to this Set
+	 * @param Source pointer to owning object, used only to identify Behavior later. If nullptr, source is copied from other Set.
+	 * @param NewGroupName string identifier for this Behavior or group of Behaviors. If empty string, group is copied from other Set.
+	 */
+	virtual void Add(const UInputBehaviorSet* OtherSet, void* NewSource = nullptr, const FString& NewGroupName = "");
+
+	/**
+	 * Remove a Behavior from the Set
+	 * @param Behavior Behavior to remove
+	 * @return true if Behavior was found and removed
+	 */
+	virtual bool Remove(UInputBehavior* Behavior);
+	/**
+	 * Remove a group of Behaviors from the Set
+	 * @param GroupName name of group, all Behaviors that were added with this GroupName are removed
+	 * @return true if any Behaviors were found and removed
+	 */
+	virtual bool RemoveByGroup(const FString& GroupName);
+	/**
+	 * Remove a group of Behaviors from the Set
+	 * @param Source source object pointer, all Behaviors that were added with this Source pointer are removed.
+	 * @return true if any Behaviors were found and removed
+	 */
+	virtual bool RemoveBySource(void* Source);
+	/**
+	 * Remove all Behaviors from the set
+	 */
+	virtual void RemoveAll();
+
+
+	//
+	// Queries and Event Forwarding
+	//
+
+
+	/**
+	 * Call UInputBehavior::WantsCapture() on each valid Behavior and collect up the requests that indicated a Capture was desired.
+	 * @param InputState current input device state information
+	 * @param ResultOut returned set of non-Ignoring FInputCaptureRequests returned by Behaviors
+	 */
+	virtual void CollectWantsCapture(const FInputDeviceState& InputState, TArray<FInputCaptureRequest>& ResultOut);
+
+	/**
+	 * Call UInputBehavior::UpdateHover() on each valid Behavior that wants hover events
+	 * @param InputState current input device state information
+	 * @return true if any InputBehavior indicated that it wants hover events
+	 */
+	virtual bool UpdateHover(const FInputDeviceState&  InputState);
+
+	/**
+	 * Call UInputBehavior::EndHover() on each valid Behavior that wants hover events
+	 * @param InputState current input device state information
+	 * @return true if any InputBehavior indicated that it wants hover events
+	 */
+	virtual bool EndHover(const FInputDeviceState&  InputState);
+
+
+protected:
+
+	/** Current set of known Behaviors */
+	UPROPERTY()
+	TArray<FBehaviorInfo> Behaviors;
+
+	/**
+	 * called internally when Behaviors list is updated, to re-sort by Priority, etc (@todo: and emit change event)
+	 */
+	virtual void BehaviorsModified();
+
+
+	/** @return true if Behavior supports InputState.InputDevice */
+	bool SupportsInputType(UInputBehavior* Behavior, const FInputDeviceState& InputState)
+	{
+		return (Behavior->GetSupportedDevices() & InputState.InputDevice) != EInputDevices::None;
+	}
+
+};
+
+
+
+
+
+// UInterface for IInputBehavior
+UINTERFACE(MinimalAPI)
+class UInputBehaviorSource : public UInterface
+{
+	GENERATED_BODY()
+};
+
+
+/** 
+ * UObjects that implement IInputBehaviorSource have an UInputBehaviorSet
+ * that they can provide (to UInputRouter, primarily)
+ * @todo callback/delegate for when the provided InputBehaviorSet changes
+ */
+class IInputBehaviorSource
+{
+	GENERATED_BODY()
+
+public:
+	/**
+	 * @return The current UInputBehaviorSet for this Source
+	 */
+	virtual const UInputBehaviorSet* GetInputBehaviors() const = 0;
+};
+
diff --git a/Engine/Source/Runtime/Experimental/InteractiveToolsFramework/Public/InputRouter.h b/Engine/Source/Runtime/Experimental/InteractiveToolsFramework/Public/InputRouter.h
new file mode 100644
--- /dev/null
+++ b/Engine/Source/Runtime/Experimental/InteractiveToolsFramework/Public/InputRouter.h
@@ -0,0 +1,105 @@
+// Copyright 1998-2019 Epic Games, Inc. All Rights Reserved.
+
+#pragma once
+
+#include "CoreMinimal.h"
+#include "UObject/Object.h"
+#include "Misc/Change.h"
+#include "InputState.h"
+#include "InputBehaviorSet.h"
+#include "ToolContextInterfaces.h"
+#include "InputRouter.generated.h"
+
+
+/**
+ * UInputRouter mediates between a higher-level input event source (eg like an FEdMode)
+ * and a set of InputBehaviors that respond to those events. Sets of InputBehaviors are
+ * registered, and then PostInputEvent() is called for each event. 
+ *
+ * Internally one of the active Behaviors may "capture" the event stream.
+ * Separate "Left" and "Right" captures are supported, which means that (eg)
+ * an independent capture can be tracked for each VR controller.
+ *
+ * If the input device supports "hover",  PostHoverInputEvent() will forward
+ * hover events to InputBehaviors that also support it.
+ *
+ */
+UCLASS(Transient)
+class INTERACTIVETOOLSFRAMEWORK_API UInputRouter : public UObject
+{
+	GENERATED_BODY()
+
+protected:
+	friend class UInteractiveToolsContext;		// to call Initialize/Shutdown
+
+	UInputRouter();
+
+	/** Initialize the InputRouter with the necessary Context-level state. UInteractiveToolsContext calls this, you should not. */
+	virtual void Initialize(IToolsContextTransactionsAPI* TransactionsAPI);
+
+	/** Shutdown the InputRouter. Called by UInteractiveToolsContext. */
+	virtual void Shutdown();
+
+public:
+	/** Add a new behavior Source. Behaviors from this source will be added to the active behavior set. */
+	virtual void RegisterSource(IInputBehaviorSource* Source);
+
+	/** Remove Behaviors from this Source from the active set */
+	virtual void DeregisterSource(IInputBehaviorSource* Source);
+
+	/** Insert a new input event which is used to check for new captures, or forwarded to active capture */
+	virtual void PostInputEvent(const FInputDeviceState& Input);
+
+	/** Returns true if there is an active mouse capture */
+	virtual bool HasActiveMouseCapture() const;
+
+	// TODO: other capture queries
+
+	/** Insert a new hover input event which is forwarded to all hover-enabled Behaviors */
+	virtual void PostHoverInputEvent(const FInputDeviceState& Input);
+
+	/** If this Behavior is capturing, call ForceEndCapture() to notify that we are taking capture away */
+	virtual void ForceTerminateSource(IInputBehaviorSource* Source);
+
+	/** Terminate any active captures and end all hovers */
+	virtual void ForceTerminateAll();
+
+
+public:
+	/** If true, then we post an Invalidation (ie redraw) request if any active InputBehavior responds to Hover events (default false) */
+	UPROPERTY()
+	bool bAutoInvalidateOnHover;
+
+	/** If true, then we post an Invalidation (ie redraw) request on every captured input event (default false) */
+	UPROPERTY()
+	bool bAutoInvalidateOnCapture;
+
+
+protected:
+	IToolsContextTransactionsAPI* TransactionsAPI;
+
+	UPROPERTY()
+	UInputBehaviorSet* ActiveInputBehaviors;
+
+	UInputBehavior* ActiveKeyboardCapture;
+	void* ActiveKeyboardCaptureOwner;
+	FInputCaptureData ActiveKeyboardCaptureData;
+
+	UInputBehavior* ActiveLeftCapture;
+	void* ActiveLeftCaptureOwner;
+	FInputCaptureData ActiveLeftCaptureData;
+
+	UInputBehavior* ActiveRightCapture;
+	void* ActiveRightCaptureOwner;
+	FInputCaptureData ActiveRightCaptureData;
+
+	FInputDeviceState LastHoverInput;
+
+	virtual void PostInputEvent_Keyboard(const FInputDeviceState& Input);
+	void CheckForKeyboardCaptures(const FInputDeviceState& Input);
+	void HandleCapturedKeyboardInput(const FInputDeviceState& Input);
+
+	virtual void PostInputEvent_Mouse(const FInputDeviceState& Input);
+	void CheckForMouseCaptures(const FInputDeviceState& Input);
+	void HandleCapturedMouseInput(const FInputDeviceState& Input);
+};
\ No newline at end of file
diff --git a/Engine/Source/Runtime/Experimental/InteractiveToolsFramework/Public/InputState.h b/Engine/Source/Runtime/Experimental/InteractiveToolsFramework/Public/InputState.h
new file mode 100644
--- /dev/null
+++ b/Engine/Source/Runtime/Experimental/InteractiveToolsFramework/Public/InputState.h
@@ -0,0 +1,248 @@
+// Copyright 1998-2019 Epic Games, Inc. All Rights Reserved.
+
+#pragma once
+
+#include "CoreMinimal.h"
+#include "UObject/ObjectMacros.h"
+#include "InputCoreTypes.h"
+#include "Math/UnrealMath.h"
+
+
+/**
+ * Input event data can be applicable to many possible input devices.
+ * These flags are used to indicate specific or sets of device types.
+ */
+UENUM()
+enum class EInputDevices
+{
+	None = 0,
+	Keyboard = 1,
+	Mouse = 2,
+	Gamepad = 4,
+
+	OculusTouch = 8,
+	HTCViveWands = 16,
+	AnySpatialDevice = OculusTouch | HTCViveWands,
+
+	TabletFingers = 1024
+};
+ENUM_CLASS_FLAGS(EInputDevices);
+
+
+
+/**
+ * FInputDeviceRay represents a 3D ray created based on an input device.
+ * If the device is a 2D input device like a mouse, then the ray may
+ * have an associated 2D screen position.
+ */
+struct FInputDeviceRay
+{
+	/** 3D ray in 3D scene, in world coordinates */
+	FRay WorldRay;
+
+	/** If true, WorldRay has 2D device position coordinates */
+	bool bHas2D;
+
+	/** 2D device position coordinates associated with the ray */
+	FVector2D ScreenPosition;
+
+
+	FInputDeviceRay(const FRay& WorldRayIn)
+	{
+		WorldRay = WorldRayIn;
+		bHas2D = false;
+		ScreenPosition = FVector2D(0, 0);
+	}
+
+	FInputDeviceRay(const FRay& WorldRayIn, const FVector2D& ScreenPositionIn)
+	{
+		WorldRay = WorldRayIn;
+		bHas2D = true;
+		ScreenPosition = ScreenPositionIn;
+	}
+};
+
+
+/**
+ * Current State of a physical device button (mouse, key, etc) at a point in time.
+ * Each "click" of a button should involve at minimum two separate state
+ * events, one where bPressed=true and one where bReleased=true.
+ * Each of these states should occur only once.
+ * In addition there may be additional frames where the button is
+ * held down and bDown=true and bPressed=false.
+ */
+struct FDeviceButtonState
+{
+	/** Button identifier */
+	FKey Button;
+
+	/** Was the button pressed down this frame. This should happen once per "click" */
+	bool bPressed;
+	/** Is the button currently pressed down. This should be true every frame the button is pressed. */
+	bool bDown;
+	/** Was the button released this frame. This should happen once per "click" */
+	bool bReleased;
+
+	FDeviceButtonState()
+	{
+		Button = FKey();
+		bPressed = bDown = bReleased = false;
+	}
+
+	FDeviceButtonState(const FKey& ButtonIn)
+	{
+		Button = ButtonIn;
+		bPressed = bDown = bReleased = false;
+	}
+
+	/** Update the states of this button */
+	void SetStates(bool bPressedIn, bool bDownIn, bool bReleasedIn)
+	{
+		bPressed = bPressedIn;
+		bDown = bDownIn;
+		bReleased = bReleasedIn;
+	}
+};
+
+
+
+/**
+ * Current state of active keyboard key at a point in time
+ * @todo would be useful to track set of active keys
+ */
+struct FKeyboardInputDeviceState
+{
+	/** state of active key that was modified (ie press or release) */
+	FDeviceButtonState ActiveKey;
+};
+
+
+
+/**
+ * Current State of a physical Mouse device at a point in time.
+ */
+struct FMouseInputDeviceState
+{
+	/** State of the left mouse button */
+	FDeviceButtonState Left;
+	/** State of the middle mouse button */
+	FDeviceButtonState Middle;
+	/** State of the right mouse button */
+	FDeviceButtonState Right;
+
+	/** Change in 'ticks' of the mouse wheel since last state event */
+	float WheelDelta;
+
+	/** Current 2D position of the mouse, in application-defined coordinate system */
+	FVector2D Position2D;
+
+	/** Change in 2D mouse position from last state event */
+	FVector2D Delta2D;
+
+	/** Ray into current 3D scene at current 2D mouse position */
+	FRay WorldRay;
+
+
+	FMouseInputDeviceState()
+	{
+		Left = FDeviceButtonState(EKeys::LeftMouseButton);
+		Middle = FDeviceButtonState(EKeys::MiddleMouseButton);
+		Right = FDeviceButtonState(EKeys::RightMouseButton);
+		WheelDelta = false;
+		Position2D = FVector2D::ZeroVector;
+		Delta2D = FVector2D::ZeroVector;
+		WorldRay = FRay();
+	}
+};
+
+
+
+/**
+ * Current state of physical input devices at a point in time.
+ * Assumption is that the state refers to a single physical input device,
+ * ie InputDevice field is a single value of EInputDevices and not a combination.
+ */
+struct FInputDeviceState
+{
+	/** Which InputDevice member is valid in this state */
+	EInputDevices InputDevice;
+
+	//
+	// keyboard modifiers
+	//
+
+	/** Is they keyboard SHIFT modifier key currently pressed down */
+	bool bShiftKeyDown;
+	/** Is they keyboard ALT modifier key currently pressed down */
+	bool bAltKeyDown;
+	/** Is they keyboard CTRL modifier key currently pressed down */
+	bool bCtrlKeyDown;
+	/** Is they keyboard CMD modifier key currently pressed down (only on Apple devices) */
+	bool bCmdKeyDown;
+
+
+	/** Current state of Keyboard device, if InputDevice == EInputDevices::Keyboard */
+	FKeyboardInputDeviceState Keyboard;
+
+	/** Current state of Mouse device, if InputDevice == EInputDevices::Mouse */
+	FMouseInputDeviceState Mouse;
+
+
+	FInputDeviceState() 
+	{
+		InputDevice = EInputDevices::None;
+		bShiftKeyDown = bAltKeyDown = bCtrlKeyDown = bCmdKeyDown = false;
+		Keyboard = FKeyboardInputDeviceState();
+		Mouse = FMouseInputDeviceState();
+	}
+
+	/** Update keyboard modifier key states */
+	void SetModifierKeyStates(bool bShiftDown, bool bAltDown, bool bCtrlDown, bool bCmdDown) 
+	{
+		bShiftKeyDown = bShiftDown;
+		bAltKeyDown = bAltDown;
+		bCtrlKeyDown = bCtrlDown;
+		bCmdKeyDown = bCmdDown;
+	}
+
+	/**
+	 * @param DeviceType Combination of device-type flags
+	 * @return true if this input state is for an input device that matches the query flags 
+	 */
+	bool IsFromDevice(EInputDevices DeviceType) const
+	{
+		return ((InputDevice & DeviceType) != EInputDevices::None);
+	}
+
+
+
+
+	//
+	// utility functions to pass as lambdas
+	// 
+
+	/** @return true if shift key is down in input state */
+	static bool IsShiftKeyDown(const FInputDeviceState& InputState)
+	{
+		return InputState.bShiftKeyDown;
+	}
+
+	/** @return true if ctrl key is down in input state */
+	static bool IsCtrlKeyDown(const FInputDeviceState& InputState)
+	{
+		return InputState.bCtrlKeyDown;
+	}
+
+	/** @return true if alt key is down in input state */
+	static bool IsAltKeyDown(const FInputDeviceState& InputState)
+	{
+		return InputState.bAltKeyDown;
+	}
+
+	/** @return true if Apple Command key is down in input state */
+	static bool IsCmdKeyDown(const FInputDeviceState& InputState)
+	{
+		return InputState.bCmdKeyDown;
+	}
+};
+
diff --git a/Engine/Source/Runtime/Experimental/InteractiveToolsFramework/Public/InteractiveGizmo.h b/Engine/Source/Runtime/Experimental/InteractiveToolsFramework/Public/InteractiveGizmo.h
new file mode 100644
--- /dev/null
+++ b/Engine/Source/Runtime/Experimental/InteractiveToolsFramework/Public/InteractiveGizmo.h
@@ -0,0 +1,83 @@
+// Copyright 1998-2019 Epic Games, Inc. All Rights Reserved.
+
+#pragma once
+
+#include "CoreMinimal.h"
+#include "InputBehaviorSet.h"
+#include "ToolContextInterfaces.h"
+#include "InteractiveGizmo.generated.h"
+
+class UInteractiveGizmoManager;
+
+
+
+/**
+ * UInteractiveGizmo is the base class for all Gizmos in the InteractiveToolsFramework.
+ *
+ * @todo callback/delegate for if/when .InputBehaviors changes
+ * @todo callback/delegate for when Gizmo properties change
+ */
+UCLASS(Transient)
+class INTERACTIVETOOLSFRAMEWORK_API UInteractiveGizmo : public UObject, public IInputBehaviorSource
+{
+	GENERATED_BODY()
+
+public:
+	UInteractiveGizmo();
+
+	/**
+	 * Called by GizmoManager to initialize the Gizmo *after* GizmoBuilder::BuildGizmo() has been called
+	 */
+	virtual void Setup();
+
+	/**
+	 * Called by GizmoManager to shut down the Gizmo
+	 */
+	virtual void Shutdown();
+
+	/**
+	 * Allow the Gizmo to do any custom drawing (ie via PDI/RHI)
+	 * @param RenderAPI Abstraction that provides access to Rendering in the current ToolsContext
+	 */
+	virtual void Render(IToolsContextRenderAPI* RenderAPI);
+
+	/**
+	 * Allow the Gizmo to do any necessary processing on Tick 
+	 * @param DeltaTime the time delta since last tick
+	 */
+	virtual void Tick(float DeltaTime);
+
+
+
+	/**
+	 * @return GizmoManager that owns this Gizmo
+	 */
+	virtual UInteractiveGizmoManager* GetGizmoManager() const;
+
+
+
+	//
+	// Input Behaviors support
+	//
+
+	/**
+	 * Add an input behavior for this Gizmo
+	 * @param Behavior behavior to add
+	 */
+	virtual void AddInputBehavior(UInputBehavior* Behavior);
+
+	/**
+	 * @return Current input behavior set.
+	 */
+	virtual const UInputBehaviorSet* GetInputBehaviors() const;
+
+
+
+protected:
+
+	/** The current set of InputBehaviors provided by this Gizmo */
+	UPROPERTY()
+	UInputBehaviorSet* InputBehaviors;
+};
+
+
diff --git a/Engine/Source/Runtime/Experimental/InteractiveToolsFramework/Public/InteractiveGizmoBuilder.h b/Engine/Source/Runtime/Experimental/InteractiveToolsFramework/Public/InteractiveGizmoBuilder.h
new file mode 100644
--- /dev/null
+++ b/Engine/Source/Runtime/Experimental/InteractiveToolsFramework/Public/InteractiveGizmoBuilder.h
@@ -0,0 +1,34 @@
+// Copyright 1998-2019 Epic Games, Inc. All Rights Reserved.
+
+#pragma once
+
+#include "CoreMinimal.h"
+#include "UObject/Interface.h"
+#include "InteractiveGizmo.h"
+#include "InteractiveGizmoBuilder.generated.h"
+
+
+
+/**
+ * A UInteractiveGizmoBuilder creates a new instance of an InteractiveGizmo (basically this is a Factory).
+ * These are registered with the InteractiveGizmoManager, which calls BuildGizmo().
+ * This is an abstract base class, you must subclass it in order to create your particular Gizmo instance
+ */
+UCLASS(Transient, Abstract)
+class INTERACTIVETOOLSFRAMEWORK_API UInteractiveGizmoBuilder : public UObject
+{
+	GENERATED_BODY()
+
+public:
+
+	/** 
+	 * Create a new instance of this builder's Gizmo
+	 * @param SceneState the current scene selection state, etc
+	 * @return a new instance of the Gizmo, or nullptr on error/failure
+	 */
+	virtual UInteractiveGizmo* BuildGizmo(const FToolBuilderState& SceneState) const
+	{
+		unimplemented();
+		return nullptr;
+	}
+};
\ No newline at end of file
diff --git a/Engine/Source/Runtime/Experimental/InteractiveToolsFramework/Public/InteractiveGizmoManager.h b/Engine/Source/Runtime/Experimental/InteractiveToolsFramework/Public/InteractiveGizmoManager.h
new file mode 100644
--- /dev/null
+++ b/Engine/Source/Runtime/Experimental/InteractiveToolsFramework/Public/InteractiveGizmoManager.h
@@ -0,0 +1,179 @@
+// Copyright 1998-2019 Epic Games, Inc. All Rights Reserved.
+
+#pragma once
+
+#include "CoreMinimal.h"
+#include "UObject/Object.h"
+#include "Misc/Change.h"
+#include "InteractiveGizmo.h"
+#include "InteractiveGizmoBuilder.h"
+#include "InputRouter.h"
+#include "ToolContextInterfaces.h"
+#include "InteractiveGizmoManager.generated.h"
+
+
+
+USTRUCT()
+struct FActiveGizmo
+{
+	GENERATED_BODY();
+
+	UInteractiveGizmo* Gizmo;
+	FString BuilderIdentifier;
+	FString InstanceIdentifier;
+};
+
+
+/**
+ * UInteractiveGizmoManager allows users of the Tools framework to create and operate Gizmo instances.
+ * For each Gizmo, a (string,GizmoBuilder) pair is registered with the GizmoManager.
+ * Gizmos can then be activated via the string identifier.
+ * 
+ */
+UCLASS(Transient)
+class INTERACTIVETOOLSFRAMEWORK_API UInteractiveGizmoManager : public UObject
+{
+	GENERATED_BODY()
+
+protected:
+	friend class UInteractiveToolsContext;		// to call Initialize/Shutdown
+
+	UInteractiveGizmoManager();
+
+	/** Initialize the GizmoManager with the necessary Context-level state. UInteractiveToolsContext calls this, you should not. */
+	virtual void Initialize(IToolsContextQueriesAPI* QueriesAPI, IToolsContextTransactionsAPI* TransactionsAPI, UInputRouter* InputRouter);
+
+	/** Shutdown the GizmoManager. Called by UInteractiveToolsContext. */
+	virtual void Shutdown();
+
+public:
+
+	//
+	// GizmoBuilder Registration and Gizmo Creation/Shutdown
+	//
+
+	/**
+	 * Register a new GizmoBuilder
+	 * @param BuilderIdentifier string used to identify this Builder
+	 * @param Builder new GizmoBuilder instance
+	 */
+	virtual void RegisterGizmoType(const FString& BuilderIdentifier, UInteractiveGizmoBuilder* Builder);
+
+	/**
+	 * Remove a GizmoBuilder from the set of known GizmoBuilders
+	 * @param BuilderIdentifier identification string that was passed to RegisterGizmoType()
+	 * @return true if Builder was found and deregistered
+	 */
+	virtual bool DeregisterGizmoType(const FString& BuilderIdentifier);
+
+
+
+	/**
+	 * Try to activate a new Gizmo instance on the given Side
+	 * @param BuilderIdentifier string used to identify Builder that should be called
+	 * @param InstanceIdentifier client-defined string that can be used to locate this instance
+	 * @return new Gizmo instance that has been created and initialized
+	 */	
+	virtual UInteractiveGizmo* CreateGizmo(const FString& BuilderIdentifier, const FString& InstanceIdentifier );
+
+
+	/**
+	 * Shutdown and remove a Gizmo
+	 * @param Gizmo the Gizmo to shutdown and remove
+	 * @return true if the Gizmo was found and removed
+	 */
+	virtual bool DestroyGizmo(UInteractiveGizmo* Gizmo);
+
+	/**
+	 * Destroy all Gizmos that were created by the identified GizmoBuilder
+	 * @param BuilderIdentifier the Builder string registered with RegisterGizmoType
+	 */
+	virtual void DestroyAllGizmosOfType(const FString& BuilderIdentifier);
+
+	/**
+	 * Find all the existing Gizmo instances that were created by the identified GizmoBuilder
+	 * @param BuilderIdentifier the Builder string registered with RegisterGizmoType
+	 * @return list of found Gizmos
+	 */
+	virtual TArray<UInteractiveGizmo*> FindAllGizmosOfType(const FString& BuilderIdentifier);
+
+	/**
+	 * Find the Gizmo that was created with the given instance identifier
+	 * @param Identifier the InstanceIdentifier that was passed to CreateGizmo()
+	 * @return the found Gizmo, or null
+	 */
+	virtual UInteractiveGizmo* FindGizmoByInstanceIdentifier(const FString& Identifier);
+
+
+
+
+	//
+	// Functions that Gizmos can call to interact with Transactions API
+	//
+	
+	/** Post a message via the Transactions API */
+	virtual void PostMessage(const TCHAR* Message, EToolMessageLevel Level);
+
+	/** Post a message via the Transactions API */
+	virtual void PostMessage(const FString& Message, EToolMessageLevel Level);
+
+	/** Request an Invalidation via the Transactions API (ie to cause a repaint, etc) */
+	virtual void PostInvalidation();
+
+	/**
+	 * Request that the Context open a Transaction, whatever that means to the current Context
+	 * @param Description text description of this transaction (this is the string that appears on undo/redo in the UE Editor)
+	 */
+	virtual void BeginUndoTransaction(const FText& Description);
+
+	/** Request that the Context close and commit the open Transaction */
+	virtual void EndUndoTransaction();
+
+	/**
+	 * Forward an FChange object to the Context
+	 * @param TargetObject the object that the FChange applies to
+	 * @param Change the change object that the Context should insert into the transaction history
+	 * @param Description text description of this change (this is the string that appears on undo/redo in the UE Editor)
+	 */
+	virtual void EmitObjectChange(UObject* TargetObject, TUniquePtr<FChange> Change, const FText& Description );
+
+
+	//
+	// State control  (@todo: have the Context call these? not safe for anyone to call)
+	//
+
+	/** Tick any active Gizmos. Called by UInteractiveToolsContext */
+	virtual void Tick(float DeltaTime);
+
+	/** Render any active Gizmos. Called by UInteractiveToolsContext. */
+	virtual void Render(IToolsContextRenderAPI* RenderAPI);
+
+
+
+	//
+	// access to APIs, etc
+	//
+
+	/** @return current IToolsContextQueriesAPI */
+	virtual IToolsContextQueriesAPI* GetContextQueriesAPI() { return QueriesAPI; }
+
+
+public:
+	/** set of Currently-active Gizmos */
+	UPROPERTY()
+	TArray<FActiveGizmo> ActiveGizmos;
+
+
+protected:
+	/** Current Context-Queries implementation */
+	IToolsContextQueriesAPI* QueriesAPI;
+	/** Current Transactions implementation */
+	IToolsContextTransactionsAPI* TransactionsAPI;
+
+	/** Current InputRouter (Context owns this) */
+	UInputRouter* InputRouter;
+
+	/** Current set of named GizmoBuilders */
+	UPROPERTY()
+	TMap<FString, UInteractiveGizmoBuilder*> GizmoBuilders;
+};
\ No newline at end of file
diff --git a/Engine/Source/Runtime/Experimental/InteractiveToolsFramework/Public/InteractiveTool.h b/Engine/Source/Runtime/Experimental/InteractiveToolsFramework/Public/InteractiveTool.h
new file mode 100644
--- /dev/null
+++ b/Engine/Source/Runtime/Experimental/InteractiveToolsFramework/Public/InteractiveTool.h
@@ -0,0 +1,239 @@
+// Copyright 1998-2019 Epic Games, Inc. All Rights Reserved.
+
+#pragma once
+
+#include "CoreMinimal.h"
+#include "InputBehaviorSet.h"
+#include "InteractiveToolActionSet.h"
+#include "ToolContextInterfaces.h"
+#include "InteractiveTool.generated.h"
+
+class UInteractiveToolManager;
+
+
+/** Passed to UInteractiveTool::Shutdown to indicate how Tool should shut itself down*/
+enum class EToolShutdownType
+{
+	/** Tool cleans up and exits. Pass this to tools that do not have Accept/Cancel options. */
+	Completed = 0,
+	/** Tool commits current preview to scene */
+	Accept = 1,
+	/** Tool discards current preview without modifying scene */
+	Cancel = 2
+};
+
+
+/** This delegate is used by UInteractiveToolPropertySet */
+DECLARE_MULTICAST_DELEGATE_TwoParams(FInteractiveToolPropertySetModifiedSignature, UObject*, UProperty*);
+
+
+/**
+ * a UInteractiveTool contains a set of UObjects that contain "properties" of the Tool, ie
+ * the configuration flags, parameters, etc that control the Tool. Currently any UObject
+ * can be added as a property set, however there is no automatic mechanism for those child 
+ * UObjects to notify the Tool when a property changes.
+ * 
+ * If you make your property set UObjects subclasses of UInteractiveToolPropertySet, then
+ * when the properties are changed in the Editor, the parent Tool will be automatically notified.
+ * You can override UInteractiveTool::OnPropertyModified() to act on these notifications
+ */
+UCLASS(Transient)
+class INTERACTIVETOOLSFRAMEWORK_API UInteractiveToolPropertySet : public UObject
+{
+	GENERATED_BODY()
+
+protected:
+	FInteractiveToolPropertySetModifiedSignature OnModified;
+
+public:
+
+	/** @return the multicast delegate that is called when properties are modified */
+	FInteractiveToolPropertySetModifiedSignature& GetOnModified()
+	{
+		return OnModified;
+	}
+
+	/** posts a message to the OnModified delegate with the modified UProperty */
+	void PostEditChangeProperty(FPropertyChangedEvent& PropertyChangedEvent)
+	{
+		OnModified.Broadcast(this, PropertyChangedEvent.Property);
+	}
+};
+
+
+/**
+ * UInteractiveTool is the base class for all Tools in the InteractiveToolsFramework.
+ * A Tool is is a "lightweight mode" that may "own" one or more Actors/Components/etc in
+ * the current scene, may capture certain input devices or event streams, and so on.
+ * The base implementation essentially does nothing but provide sane default behaviors.
+ *
+ * The BaseTools/ subfolder contains implementations of various kinds of standard
+ * "tool behavior", like a tool that responds to a mouse click, etc, that can be
+ * extended to implement custom behaviors.
+ *
+ * In the framework, you do not create instances of UInteractiveTool yourself. 
+ * You provide a UInteractiveToolBuilder implementation that can properly construct
+ * an instance of your Tool, this is where for example default parameters would be set.
+ * The ToolBuilder is registered with the ToolManager, and then UInteractiveToolManager::ActivateTool()
+ * is used to kick things off.
+ *
+ * @todo callback/delegate for if/when .InputBehaviors changes
+ * @todo callback/delegate for when tool properties change
+ */
+UCLASS(Transient)
+class INTERACTIVETOOLSFRAMEWORK_API UInteractiveTool : public UObject, public IInputBehaviorSource
+{
+	GENERATED_BODY()
+
+public:
+	UInteractiveTool();
+
+	/**
+	 * Called by ToolManager to initialize the Tool *after* ToolBuilder::BuildTool() has been called
+	 */
+	virtual void Setup();
+
+	/**
+	 * Called by ToolManager to shut down the Tool
+	 * @param ShutdownType indicates how the tool should shutdown (ie Accept or Cancel current preview, etc)
+	 */
+	virtual void Shutdown(EToolShutdownType ShutdownType);
+
+	/**
+	 * Allow the Tool to do any custom drawing (ie via PDI/RHI)
+	 * @param RenderAPI Abstraction that provides access to Rendering in the current ToolsContext
+	 */
+	virtual void Render(IToolsContextRenderAPI* RenderAPI);
+
+	/**
+	 * Allow the Tool to do any necessary processing on Tick 
+	 * @param DeltaTime the time delta since last tick
+	 */
+	virtual void Tick(float DeltaTime);
+
+
+
+	/**
+	 * @return ToolManager that owns this Tool
+	 */
+	virtual UInteractiveToolManager* GetToolManager() const;
+
+
+	/** 
+	 * @return true if this Tool support being Cancelled, ie calling Shutdown(EToolShutdownType::Cancel)  
+	 */
+	virtual bool HasCancel() const;
+
+	/** 
+	 * @return true if this Tool support being Accepted, ie calling Shutdown(EToolShutdownType::Accept)  
+	 */
+	virtual bool HasAccept() const;
+
+	/** 
+	 * @return true if this Tool is currently in a state where it can be Accepted. This may be false if for example there was an error in the Tool. 
+	 */
+	virtual bool CanAccept() const;
+
+
+
+	//
+	// Input Behaviors support
+	//
+
+	/**
+	 * Add an input behavior for this Tool
+	 * @param Behavior behavior to add
+	 */
+	virtual void AddInputBehavior(UInputBehavior* Behavior);
+
+	/**
+	 * @return Current input behavior set.
+	 */
+	virtual const UInputBehaviorSet* GetInputBehaviors() const;
+
+
+	//
+	// Property support
+	//
+
+	/**
+	 * @return list of property UObjects for this tool (ie to add to a DetailsViewPanel, for example)
+	 */
+	virtual const TArray<UObject*>& GetToolProperties() const;
+
+
+	/**
+	 * Automatically called by UInteractiveToolPropertySet.OnModified delegate to notify Tool of child property set changes
+	 * @param PropertySet which UInteractiveToolPropertySet was modified
+	 * @param Property which UProperty in the set was modified
+	 */
+	virtual void OnPropertyModified(UObject* PropertySet, UProperty* Property)
+	{
+	}
+
+
+protected:
+
+	/** The current set of InputBehaviors provided by this Tool */
+	UPROPERTY()
+	UInputBehaviorSet* InputBehaviors;
+
+	/** The current set of Property UObjects provided by this Tool. May contain pointer to itself. */
+	UPROPERTY()
+	TArray<UObject*> ToolPropertyObjects;
+
+	/**
+	 * Add a Property object for this Tool
+	 * @param Property object to add
+	 */
+	virtual void AddToolPropertySource(UObject* PropertyObject);
+
+	/**
+	 * Add a PropertySet object for this Tool
+	 * @param PropertySet Property Set object to add
+	 */
+	virtual void AddToolPropertySource(UInteractiveToolPropertySet* PropertySet);
+
+
+
+
+	//
+	// Action support/system
+	// 
+	// Your Tool subclass can register a set of "Actions" it can execute
+	// by overloading RegisterActions(). Then external systems can use GetActionSet() to
+	// find out what Actions your Tool supports, and ExecuteAction() to run those actions.
+	//
+	
+public:
+	/**
+	 * Get the internal Action Set for this Tool. The action set is created and registered on-demand.
+	 * @return pointer to initialized Action set
+	 */
+	virtual FInteractiveToolActionSet* GetActionSet();
+
+	/**
+	 * Request that the Action identified by ActionID be executed.
+	 * Default implementation forwards these requests to internal ToolActionSet.
+	 */
+	virtual void ExecuteAction(int32 ActionID);
+
+
+protected:
+	/**
+	 * Override this function to register the set of Actions this Tool supports, using FInteractiveToolActionSet::RegisterAction.
+	 * Note that 
+	 */
+	virtual void RegisterActions(FInteractiveToolActionSet& ActionSet);
+
+
+private:
+	/** 
+	 * Set of actions this Tool can execute. This variable is allocated on-demand. 
+	 * Use GetActionSet() instead of accessing this pointer directly!
+	 */
+	FInteractiveToolActionSet* ToolActionSet = nullptr;
+
+};
+
+
diff --git a/Engine/Source/Runtime/Experimental/InteractiveToolsFramework/Public/InteractiveToolActionSet.h b/Engine/Source/Runtime/Experimental/InteractiveToolsFramework/Public/InteractiveToolActionSet.h
new file mode 100644
--- /dev/null
+++ b/Engine/Source/Runtime/Experimental/InteractiveToolsFramework/Public/InteractiveToolActionSet.h
@@ -0,0 +1,114 @@
+// Copyright 1998-2019 Epic Games, Inc. All Rights Reserved.
+
+#pragma once
+
+#include "CoreMinimal.h"
+#include "GenericPlatform/GenericApplication.h"
+#include "InputCoreTypes.h"
+
+class UInteractiveTool;
+
+
+/**
+ * Standard Actions that can be shared across multiple Tools.
+ * The enum values are used as FInteractiveToolAction.ActionID.
+ * If you need to define your own values, start at BaseClientDefinedActionID
+ */
+enum class EStandardToolActions
+{
+	IncreaseBrushSize = 100,
+	DecreaseBrushSize = 101,
+
+
+	BaseClientDefinedActionID = 10000
+};
+
+
+
+/**
+ * FInteractiveToolAction is returned by a UInteractiveTool to represent
+ * an "Action" the Tool can execute. 
+ */
+struct INTERACTIVETOOLSFRAMEWORK_API FInteractiveToolAction
+{
+	/** Which type of UInteractiveTool this Action can be applied to*/
+	const UClass* ClassType;
+
+	/** Identifier for this Action */
+	int32 ActionID;
+
+	/** Internal name for this Action */
+	FString ActionName;
+	/** Short name for this Action */
+	FString ShortName;
+	/** Descriptive name for this Action */
+	FString Description;
+
+	/** Suggested modifier keys for this Action */
+	EModifierKey::Type DefaultModifiers;
+	/** Suggested keybinding for this Action. */
+	FKey DefaultKey;
+
+	/** Call this function to execute the Action */
+	TFunction<void()> OnAction;
+
+
+	FInteractiveToolAction()
+	{
+		ClassType = nullptr;
+		ActionID = 0;
+	}
+
+	FInteractiveToolAction(const UClass* ClassTypeIn, int32 ActionIDIn,
+		const FString& ActionNameIn, const FString& ShortNameIn, const FString& DescriptionIn,
+		EModifierKey::Type DefaultModifiersIn, const FKey& DefaultKeyIn )
+	{
+		ClassType = ClassTypeIn;
+		ActionID = ActionIDIn;
+		ActionName = ActionNameIn;
+		ShortName = ShortNameIn;
+		Description = DescriptionIn;
+		DefaultModifiers = DefaultModifiersIn;
+		DefaultKey = DefaultKeyIn;
+	}
+
+};
+
+
+
+/**
+ * FInteractiveToolActionSet maintains a list of FInteractiveToolAction.
+ * Each UInteractiveTool contains an instance of this class.
+ */
+class INTERACTIVETOOLSFRAMEWORK_API FInteractiveToolActionSet
+{
+public:
+
+	/**
+	 * Register an Action with the ActionSet. This function is intended to be called by
+	 * UInteractiveTool::RegisterActions() implementations
+	 */
+	void RegisterAction(UInteractiveTool* Tool, int32 ActionID,
+		const FString& ActionName, const FString& ShortUIName, const FString& DescriptionText,
+		EModifierKey::Type Modifiers, const FKey& ShortcutKey,
+		TFunction<void()> ActionFunction );
+
+	/**
+	 * Find an existing Action by ID
+	 * @return located Action, or nullptr if not found
+	 */
+	const FInteractiveToolAction* FindActionByID(int32 ActionID) const;
+
+	/**
+	 * Return the internal list of registered Actions by adding to the OutActions array
+	 */
+	void CollectActions(TArray<FInteractiveToolAction>& OutActions) const;
+
+	/**
+	 * Execute the action identified by ActionID
+	 */
+	void ExecuteAction(int32 ActionID) const;
+
+protected:
+	TArray<FInteractiveToolAction> Actions;
+};
\ No newline at end of file
diff --git a/Engine/Source/Runtime/Experimental/InteractiveToolsFramework/Public/InteractiveToolBuilder.h b/Engine/Source/Runtime/Experimental/InteractiveToolsFramework/Public/InteractiveToolBuilder.h
new file mode 100644
--- /dev/null
+++ b/Engine/Source/Runtime/Experimental/InteractiveToolsFramework/Public/InteractiveToolBuilder.h
@@ -0,0 +1,46 @@
+// Copyright 1998-2019 Epic Games, Inc. All Rights Reserved.
+
+#pragma once
+
+#include "CoreMinimal.h"
+#include "UObject/Interface.h"
+#include "InteractiveTool.h"
+#include "InteractiveToolBuilder.generated.h"
+
+
+
+/**
+ * A UInteractiveToolBuilder creates a new instance of an InteractiveTool (basically this is a Factory).
+ * These are registered with the InteractiveToolManager, which calls BuildTool() if CanBuildTool() returns true.
+ * In addition CanBuildTool() will be queried to (for example) enable/disable UI buttons, etc.
+ * This is an abstract base class, you must subclass it in order to create your particular Tool instance
+ */
+UCLASS(Transient, Abstract)
+class INTERACTIVETOOLSFRAMEWORK_API UInteractiveToolBuilder : public UObject
+{
+	GENERATED_BODY()
+
+public:
+
+	/** 
+	 * Check if, given the current scene state, a new instance of this builder's Tool can be created
+	 * @param SceneState the current scene selection state, etc
+	 * @return true if a new Tool instance can be created
+	 */
+	virtual bool CanBuildTool(const FToolBuilderState& SceneState) const
+	{
+		check(false);
+		return false;
+	}
+
+	/** 
+	 * Create a new instance of this builder's Tool
+	 * @param SceneState the current scene selection state, etc
+	 * @return a new instance of the Tool, or nullptr on error/failure
+	 */
+	virtual UInteractiveTool* BuildTool(const FToolBuilderState& SceneState) const
+	{
+		check(false);
+		return nullptr;
+	}
+};
\ No newline at end of file
diff --git a/Engine/Source/Runtime/Experimental/InteractiveToolsFramework/Public/InteractiveToolManager.h b/Engine/Source/Runtime/Experimental/InteractiveToolsFramework/Public/InteractiveToolManager.h
new file mode 100644
--- /dev/null
+++ b/Engine/Source/Runtime/Experimental/InteractiveToolsFramework/Public/InteractiveToolManager.h
@@ -0,0 +1,235 @@
+// Copyright 1998-2019 Epic Games, Inc. All Rights Reserved.
+
+#pragma once
+
+#include "CoreMinimal.h"
+#include "UObject/Object.h"
+#include "Misc/Change.h"
+#include "InteractiveTool.h"
+#include "InteractiveToolBuilder.h"
+#include "InputRouter.h"
+#include "ToolContextInterfaces.h"
+#include "InteractiveToolManager.generated.h"
+
+
+
+/** A Tool can be activated on a particular input device, currently identified by a "side" */
+UENUM()
+enum class EToolSide
+{
+	/** Left-hand Tool, also used for Mouse */
+	Left = 1,
+	Mouse = 1,
+	/** Right-hand Tool*/
+	Right = 2,
+};
+
+
+/**
+ * UInteractiveToolManager allows users of the tools framework to create and operate Tool instances.
+ * For each Tool, a (string,ToolBuilder) pair is registered with the ToolManager.
+ * Tools can then be activated via the string identifier.
+ * 
+ * Currently a single Tool can be active for each input device. So for mouse input a single
+ * Tool is available and effectively a lightweight mode. The mouse uses the "Left" tool slot.
+ * 
+ * For VR controllers and touch input, a "Left" and "Right" tool can be active at the same time.  
+ * @todo this is not fully supported yet
+ * 
+ * Tools are not directly created. Use SelectActiveToolType(side,string) to set the active ToolBuilder
+ * on a given side, and then use ActivateTool() to create the new Tool instance.
+ *
+ */
+UCLASS(Transient)
+class INTERACTIVETOOLSFRAMEWORK_API UInteractiveToolManager : public UObject
+{
+	GENERATED_BODY()
+
+protected:
+	friend class UInteractiveToolsContext;		// to call Initialize/Shutdown
+
+	UInteractiveToolManager();
+
+	/** Initialize the ToolManager with the necessary Context-level state. UInteractiveToolsContext calls this, you should not. */
+	virtual void Initialize(IToolsContextQueriesAPI* QueriesAPI, IToolsContextTransactionsAPI* TransactionsAPI, UInputRouter* InputRouter);
+
+	/** Shutdown the ToolManager. Called by UInteractiveToolsContext. */
+	virtual void Shutdown();
+
+public:
+
+	//
+	// Tool registration and Current Tool state
+	//
+
+	/**
+	 * Register a new ToolBuilder
+	 * @param Identifier string used to identify this Builder
+	 * @param Builder new ToolBuilder instance
+	 */
+	virtual void RegisterToolType(const FString& Identifier, UInteractiveToolBuilder* Builder);
+
+	/**
+	 * Set active ToolBuilder for a ToolSide via string identifier
+	 * @param Side which "side" should we set this Builder on
+	 * @param Identifier name of ToolBuilder that was passed to RegisterToolType()
+	 */
+	virtual bool SelectActiveToolType(EToolSide Side, const FString& Identifier);
+
+	/**
+	 * Check if a named Tool type can currently be activated on the given ToolSide
+	 * @param Side which "side" you would like to active the tool on
+	 * @param Identifier string name of the Tool type
+	 * @return true if the Tool type could be activated
+	 */
+	virtual bool CanActivateTool(EToolSide eSide, const FString& Identifier);
+
+	/**
+	 * Try to activate a new Tool instance on the given Side
+	 * @param Side which "side" you would like to active the tool on
+	 * @return true if a new Tool instance was created and initialized
+	 */	
+	virtual bool ActivateTool(EToolSide Side);
+
+	/**
+	 * Check if there is an active Tool on the given Side
+	 * @param Side which Side to check
+	 * @return true if there is an active Tool on that side
+	 */
+	virtual bool HasActiveTool(EToolSide Side) const;
+
+
+	/**
+	 * @return true if there are any active tools
+	 */
+	virtual bool HasAnyActiveTool() const;
+
+
+	/**
+	 * Get pointer to active Tool on a given side
+	 * @param Side which Side is being requested
+	 * @return pointer to Tool instance active on that Side, or nullptr if no such Tool exists
+	 */
+	virtual UInteractiveTool* GetActiveTool(EToolSide Side);
+
+
+	/**
+	 * Check if an active Tool on the given Side can be Accepted in its current state
+	 * @param Side which Side to check
+	 * @return true if there is an active Tool and it returns true from HasAccept() and CanAccept()
+	 */
+	virtual bool CanAcceptActiveTool(EToolSide Side);
+
+	/**
+	 * Check if an active Tool on the given Side can be Canceled
+	 * @param Side which Side to check
+	 * @return true if there is an active Tool and it returns true from HasCancel()
+	 */
+	virtual bool CanCancelActiveTool(EToolSide Side);
+
+	/**
+	 * Shut down an active Tool on the given side
+	 * @param Side which "side" you would like to shut down
+	 * @param ShutdownType how should the tool be terminated (eg Accept/Cancel)
+	 */
+	virtual void DeactivateTool(EToolSide Side, EToolShutdownType ShutdownType);
+
+
+
+	//
+	// Functions that Tools can call to interact with Transactions API
+	//
+	
+	/** Post a message via the Transactions API */
+	virtual void PostMessage(const TCHAR* Message, EToolMessageLevel Level);
+
+	/** Post a message via the Transactions API */
+	virtual void PostMessage(const FString& Message, EToolMessageLevel Level);
+
+	/** Request an Invalidation via the Transactions API (ie to cause a repaint, etc) */
+	virtual void PostInvalidation();
+
+	/**
+	 * Request that the Context open a Transaction, whatever that means to the current Context
+	 * @param Description text description of this transaction (this is the string that appears on undo/redo in the UE Editor)
+	 */
+	virtual void BeginUndoTransaction(const FText& Description);
+
+	/** Request that the Context close and commit the open Transaction */
+	virtual void EndUndoTransaction();
+
+	/**
+	 * Forward an FChange object to the Context
+	 * @param TargetObject the object that the FChange applies to
+	 * @param Change the change object that the Context should insert into the transaction history
+	 * @param Description text description of this change (this is the string that appears on undo/redo in the UE Editor)
+	 */
+	virtual void EmitObjectChange(UObject* TargetObject, TUniquePtr<FChange> Change, const FText& Description );
+
+
+	/**
+	 * Forward an FChange object to the Context
+	 */
+	virtual bool RequestSelectionChange(const FSelectedOjectsChangeList& SelectionChange);
+
+
+
+	//
+	// State control  (@todo: have the Context call these? not safe for anyone to call)
+	//
+
+	/** Tick any active Tools. Called by UInteractiveToolsContext */
+	virtual void Tick(float DeltaTime);
+
+	/** Render any active Tools. Called by UInteractiveToolsContext. */
+	virtual void Render(IToolsContextRenderAPI* RenderAPI);
+
+
+	//
+	// access to APIs, etc
+	//
+	
+	/** @return current IToolsContextQueriesAPI */
+	virtual IToolsContextQueriesAPI* GetContextQueriesAPI() { return QueriesAPI; }
+
+
+
+public:
+	/** Currently-active Left Tool, or null if no Tool is active */
+	UPROPERTY()
+	UInteractiveTool* ActiveLeftTool;
+
+	/** Currently-active Right Tool, or null if no Tool is active */
+	UPROPERTY()
+	UInteractiveTool* ActiveRightTool;
+
+
+
+
+public:
+	DECLARE_MULTICAST_DELEGATE_TwoParams(FToolManagerToolStartedSignature, UInteractiveToolManager*, UInteractiveTool*);
+	FToolManagerToolStartedSignature OnToolStarted;
+
+	DECLARE_MULTICAST_DELEGATE_TwoParams(FToolManagerToolEndedSignature, UInteractiveToolManager*, UInteractiveTool*);
+	FToolManagerToolEndedSignature OnToolEnded;
+
+
+protected:
+	/** Pointer to current Context-Queries implementation */
+	IToolsContextQueriesAPI* QueriesAPI;
+	/** Pointer to current Transactions implementation */
+	IToolsContextTransactionsAPI* TransactionsAPI;
+
+	/** Pointer to current InputRouter (Context owns this) */
+	UInputRouter* InputRouter;
+
+	/** Current set of named ToolBuilders */
+	UPROPERTY()
+	TMap<FString, UInteractiveToolBuilder*> ToolBuilders;
+
+	/** Currently-active Left ToolBuilder */
+	UInteractiveToolBuilder* ActiveLeftBuilder;
+	/** Currently-active Right ToolBuilder */
+	UInteractiveToolBuilder* ActiveRightBuilder;
+
+};
\ No newline at end of file
diff --git a/Engine/Source/Runtime/Experimental/InteractiveToolsFramework/Public/InteractiveToolsContext.h b/Engine/Source/Runtime/Experimental/InteractiveToolsFramework/Public/InteractiveToolsContext.h
new file mode 100644
--- /dev/null
+++ b/Engine/Source/Runtime/Experimental/InteractiveToolsFramework/Public/InteractiveToolsContext.h
@@ -0,0 +1,45 @@
+// Copyright 1998-2019 Epic Games, Inc. All Rights Reserved.
+
+#pragma once
+
+#include "CoreMinimal.h"
+#include "InteractiveToolManager.h"
+#include "InteractiveGizmoManager.h"
+#include "InteractiveToolsContext.generated.h"
+
+/**
+ * InteractiveToolsContext owns a ToolManager and an InputRouter. This is just a top-level 
+ * UObject container, however implementations like UEdModeInteractiveToolsContext extend
+ * this class to make it easier to connect external systems (like an FEdMode) to the ToolsFramework.
+ */
+UCLASS(Transient)
+class INTERACTIVETOOLSFRAMEWORK_API UInteractiveToolsContext : public UObject
+{
+	GENERATED_BODY()
+	
+public:
+	UInteractiveToolsContext();
+
+	/** 
+	 * Initialize the Context. This creates the InputRouter and ToolManager 
+	 * @param QueriesAPI client-provided implementation of the API for querying the higher-evel scene state
+	 * @param TransactionsAPI client-provided implementation of the API for publishing events and transactions
+	 */
+	virtual void Initialize(IToolsContextQueriesAPI* QueriesAPI, IToolsContextTransactionsAPI* TransactionsAPI);
+
+	/** Shutdown Context by destroying InputRouter and ToolManager */
+	virtual void Shutdown();
+
+public:
+	/** current UInputRouter for this Context */
+	UPROPERTY()
+	UInputRouter* InputRouter;	
+
+	/** current UInteractiveToolManager for this Context */
+	UPROPERTY()
+	UInteractiveToolManager* ToolManager;	
+
+	/** current UInteractiveGizmoManager for this Context */
+	UPROPERTY()
+	UInteractiveGizmoManager* GizmoManager;
+};
\ No newline at end of file
diff --git a/Engine/Source/Runtime/Experimental/InteractiveToolsFramework/Public/InteractiveToolsFramework.h b/Engine/Source/Runtime/Experimental/InteractiveToolsFramework/Public/InteractiveToolsFramework.h
new file mode 100644
--- /dev/null
+++ b/Engine/Source/Runtime/Experimental/InteractiveToolsFramework/Public/InteractiveToolsFramework.h
@@ -0,0 +1,15 @@
+// Copyright 1998-2019 Epic Games, Inc. All Rights Reserved.
+
+#pragma once
+
+#include "CoreMinimal.h"
+#include "Modules/ModuleManager.h"
+
+class FInteractiveToolsFrameworkModule : public IModuleInterface
+{
+public:
+
+	/** IModuleInterface implementation */
+	virtual void StartupModule() override;
+	virtual void ShutdownModule() override;
+};
diff --git a/Engine/Source/Runtime/Experimental/InteractiveToolsFramework/Public/ToolBuilderUtil.h b/Engine/Source/Runtime/Experimental/InteractiveToolsFramework/Public/ToolBuilderUtil.h
new file mode 100644
--- /dev/null
+++ b/Engine/Source/Runtime/Experimental/InteractiveToolsFramework/Public/ToolBuilderUtil.h
@@ -0,0 +1,117 @@
+// Copyright 1998-2019 Epic Games, Inc. All Rights Reserved.
+
+#pragma once
+
+#include "CoreMinimal.h"
+#include "GameFramework/Actor.h"
+#include "Engine/Selection.h"
+#include "InteractiveToolBuilder.h"
+
+/**
+* Helper functions that can be used in InteractiveToolBuilder implementations
+*/
+namespace ToolBuilderUtil
+{
+	/** Returns true if this UObject can provide a FMeshDescription */
+	INTERACTIVETOOLSFRAMEWORK_API
+	bool IsMeshDescriptionSourceComponent(UActorComponent* ComponentObject);
+	
+	/** Count number of selected components that pass predicate. If Component selection is not empty, returns that count, otherwise counts in all selected Actors */
+	INTERACTIVETOOLSFRAMEWORK_API
+	int CountComponents(const FToolBuilderState& InputState, const TFunction<bool(UActorComponent*)>& Predicate);
+
+	/** First first available component that passes predicate. Searches Components selection list first, then all Actors */
+	INTERACTIVETOOLSFRAMEWORK_API
+	UActorComponent* FindFirstComponent(const FToolBuilderState& InputState, const TFunction<bool(UActorComponent*)>& Predicate);
+
+	/** First all components that passes predicate. Searches Components selection list first, then all Actors */
+	INTERACTIVETOOLSFRAMEWORK_API
+	TArray<UActorComponent*> FindAllComponents(const FToolBuilderState& InputState, const TFunction<bool(UActorComponent*)>& Predicate);
+
+	// @todo not sure that actors with multiple components are handled properly...
+	/** Count number of components of given type. If Component selection is not empty, returns that count, otherwise counts in all selected Actors */
+	template<typename ComponentType>
+	int CountSelectedComponentsOfType(const FToolBuilderState& InputState);
+
+	/** First first available component of given type. Searches Components selection list first, then all Actors */
+	template<typename ComponentType>
+	ComponentType* FindFirstComponentOfType(const FToolBuilderState& InputState);
+
+}
+
+
+
+
+/*
+ * Template Implementations
+ */
+
+
+template<typename ComponentType>
+int ToolBuilderUtil::CountSelectedComponentsOfType(const FToolBuilderState& InputState)
+{
+	int nTypedComponents = 0;
+
+	if (InputState.SelectedComponents != nullptr && InputState.SelectedComponents->Num() > 0)
+	{
+		for (FSelectionIterator Iter(*InputState.SelectedComponents); Iter; ++Iter)
+		{
+			if (ComponentType* TypedComponent = Cast<ComponentType>(*Iter))
+			{
+				nTypedComponents++;
+			}
+		}
+	}
+	else
+	{
+		for (FSelectionIterator Iter(*InputState.SelectedActors); Iter; ++Iter)
+		{
+			if (AActor* Actor = Cast<AActor>(*Iter))
+			{
+				// [RMS] is there a more efficient way to do this?
+				TArray<UActorComponent*> TypedComponents = Actor->GetComponentsByClass(ComponentType::StaticClass());
+				nTypedComponents += TypedComponents.Num();
+			}
+		}
+	}
+
+	return nTypedComponents;
+}
+
+
+
+
+
+
+template<typename ComponentType>
+ComponentType* ToolBuilderUtil::FindFirstComponentOfType(const FToolBuilderState& InputState)
+{
+	if (InputState.SelectedComponents != nullptr && InputState.SelectedComponents->Num() > 0)
+	{
+		for (FSelectionIterator Iter(*InputState.SelectedComponents); Iter; ++Iter)
+		{
+			if (ComponentType* TypedComponent = Cast<ComponentType>(*Iter))
+			{
+				return TypedComponent;
+			}
+		}
+	}
+	else
+	{
+		for (FSelectionIterator Iter(*InputState.SelectedActors); Iter; ++Iter)
+		{
+			if (AActor* Actor = Cast<AActor>(*Iter))
+			{
+				// [RMS] is there a more efficient way to do this?
+				TArray<UActorComponent*> TypedComponents = Actor->GetComponentsByClass(ComponentType::StaticClass());
+				if (TypedComponents.Num() > 0)
+				{
+					return (ComponentType*)TypedComponents[0];
+				}
+			}
+		}
+	}
+
+	return nullptr;
+}
+
diff --git a/Engine/Source/Runtime/Experimental/InteractiveToolsFramework/Public/ToolContextInterfaces.h b/Engine/Source/Runtime/Experimental/InteractiveToolsFramework/Public/ToolContextInterfaces.h
new file mode 100644
--- /dev/null
+++ b/Engine/Source/Runtime/Experimental/InteractiveToolsFramework/Public/ToolContextInterfaces.h
@@ -0,0 +1,339 @@
+// Copyright 1998-2019 Epic Games, Inc. All Rights Reserved.
+
+#pragma once
+
+#include "CoreMinimal.h"
+#include "ComponentSourceInterfaces.h"
+
+
+// predeclarations so we don't have to include these in all tools
+class FChange;
+class UPackage;
+class FPrimitiveDrawInterface;
+class USelection;
+class UInteractiveToolManager;
+class UInteractiveGizmoManager;
+
+/**
+ * FToolBuilderState is a bucket of state information that a ToolBuilder might need
+ * to construct a Tool. This information comes from a level above the Tools framework,
+ * and depends on the context we are in (Editor vs Runtime, for example).
+ */
+struct INTERACTIVETOOLSFRAMEWORK_API FToolBuilderState
+{
+	/** The current UWorld */
+	UWorld* World = nullptr;
+	/** The current ToolManager */
+	UInteractiveToolManager* ToolManager = nullptr;
+	/** The current GizmoManager */
+	UInteractiveGizmoManager* GizmoManager = nullptr;
+
+	/** Current selected Actors. May be empty or nullptr. */
+	USelection* SelectedActors = nullptr;
+	/** Current selected Components. May be empty or nullptr. */
+	USelection* SelectedComponents = nullptr;
+
+	/** Implementation that can build Sources (like MeshDescriptionSource) for Components */
+	IComponentSourceFactory* SourceBuilder = nullptr;
+
+	FToolBuilderState()
+	{
+	}
+};
+
+
+/**
+ * FViewCameraState is a bucket of state information that a Tool might
+ * need to implement interactions that depend on the current scene view.
+ */
+struct INTERACTIVETOOLSFRAMEWORK_API FViewCameraState
+{
+	/** Current camera/head position */
+	FVector Position;
+	/** Current camera/head orientation */
+	FQuat Orientation;
+	/** Is current view an orthographic view */
+	bool bIsOrthographic;
+	/** Is current view a VR view */
+	bool bIsVR;
+
+	/** @return "right"/horizontal direction in camera plane */
+	FVector Right() const { return Orientation.GetAxisY(); }
+	/** @return "up"/vertical direction in camera plane */
+	FVector Up() const { return Orientation.GetAxisZ(); }
+	/** @return forward camera direction */
+	FVector Forward() const { return Orientation.GetAxisX(); }
+};
+
+
+
+/** Types of Snap Queries that a ToolsContext parent may support, that Tools may request */
+UENUM()
+enum class ESceneSnapQueryType
+{
+	/** snapping a position */
+	Position = 1		
+};
+
+/** Types of snap targets that a Tool may want to run snap queries against. */
+UENUM()
+enum class ESceneSnapQueryTargetType
+{
+	None = 0,
+	/** Consider any mesh vertex */
+	MeshVertex = 1,
+	/** Consider any mesh edge */
+	MeshEdge = 2,
+
+	All = MeshVertex | MeshEdge
+};
+ENUM_CLASS_FLAGS(ESceneSnapQueryTargetType);
+
+/**
+ * Configuration variables for a IToolsContextQueriesAPI snap query request.
+ */
+struct INTERACTIVETOOLSFRAMEWORK_API FSceneSnapQueryRequest
+{
+	/** What type of snap query geometry is this */
+	ESceneSnapQueryType RequestType;
+	/** What does caller want to try to snap to */
+	ESceneSnapQueryTargetType TargetTypes;
+
+	/** Snap input position */
+	FVector Position;
+	/** Another position must deviate less than this number of degrees (in visual angle) to be considered an acceptable snap position */
+	float VisualAngleThresholdDegrees;
+
+	/** Snap input direction */
+	FVector Direction;
+	/** Another direction must deviate less than this number of degrees from Direction to be considered an acceptable snap direction */
+	float DirectionAngleThresholdDegrees;
+};
+
+
+/**
+ * Computed result of a IToolsContextQueriesAPI snap query request
+ */
+struct INTERACTIVETOOLSFRAMEWORK_API FSceneSnapQueryResult
+{
+	/** Actor that owns snap target */
+	AActor* TargetActor = nullptr;
+	/** Component that owns snap target */
+	UActorComponent* TargetComponent = nullptr;
+	/** What kind of geometric element was snapped to */
+	ESceneSnapQueryTargetType TargetType = ESceneSnapQueryTargetType::None;
+
+	/** Snap position (may not be set depending on query types) */
+	FVector Position;
+	/** Snap normal (may not be set depending on query types) */
+	FVector Normal;
+	/** Snap direction (may not be set depending on query types) */
+	FVector Direction;
+
+	/** Vertices of triangle that contains result (for debugging, may not be set) */
+	FVector TriVertices[3];
+	/** Vertex/Edge index we snapped to in triangle */
+	int TriSnapIndex;
+
+};
+
+
+
+/** Types of standard materials that Tools may request from Context */
+UENUM()
+enum class EStandardToolContextMaterials
+{
+	/** White material that displays vertex colors set on mesh */
+	VertexColorMaterial = 1
+};
+
+
+
+/**
+ * Users of the Tools Framework need to implement IToolsContextQueriesAPI to provide
+ * access to scene state information like the current UWorld, active USelections, etc.
+ */
+class IToolsContextQueriesAPI
+{
+public:
+	virtual ~IToolsContextQueriesAPI() {}
+
+	/**
+	 * Collect up current-selection information for the current scene state (ie what is selected in Editor, etc)
+	 * @param StateOut this structure is populated with available state information
+	 */
+	virtual void GetCurrentSelectionState(FToolBuilderState& StateOut) const = 0;
+
+
+	/**
+	 * Request information about current view state
+	 * @param StateOut this structure is populated with available state information
+	 */
+	virtual void GetCurrentViewState(FViewCameraState& StateOut) const = 0;
+
+
+	/**
+	 * Try to find Snap Targets in the scene that satisfy the Snap Query.
+	 * @param Request snap query configuration
+	 * @param Results list of potential snap results
+	 * @return true if any valid snap target was found
+	 * @warning implementations are not required (and may not be able) to support snapping
+	 */
+	virtual bool ExecuteSceneSnapQuery(const FSceneSnapQueryRequest& Request, TArray<FSceneSnapQueryResult>& Results ) const = 0;
+
+
+	/**
+	 * Many tools need standard types of materials that the user should provide (eg a vertex-color material, etc)
+	 * @param MaterialType the type of material being requested
+	 * @return Instance of material to use for this purpose
+	 */
+	virtual UMaterialInterface* GetStandardMaterial(EStandardToolContextMaterials MaterialType) const = 0;
+};
+
+
+
+
+/** Level of severity of messages emitted by Tool framework */
+UENUM()
+enum class EToolMessageLevel
+{
+	/** Development message goes into development log */
+	Internal = 0,
+	/** User message should appear in user-facing log */
+	UserMessage = 1,
+	/** Notification message should be shown in a non-modal notification window */
+	UserNotification = 2,
+	/** Error message should be shown in a modal notification window */
+	UserError = 3
+};
+
+/** Type of change we want to apply to a selection */
+UENUM()
+enum class ESelectedObjectsModificationType
+{
+	Replace = 0,
+	Add = 1,
+	Remove = 2,
+	Clear = 3
+};
+
+
+/** Represents a change to a set of selected Actors and Components */
+struct FSelectedOjectsChangeList
+{
+	/** How should this list be interpreted in the context of a larger selection set */
+	ESelectedObjectsModificationType ModificationType;
+	/** List of Actors */
+	TArray<AActor*> Actors;
+	/** List of Componets */
+	TArray<UActorComponent*> Components;
+};
+
+
+/**
+ * Users of the Tools Framework need to implement IToolsContextTransactionsAPI so that
+ * the Tools have the ability to create Transactions and emit Changes. Note that this is
+ * technically optional, but that undo/redo won't be supported without it.
+ */
+class IToolsContextTransactionsAPI
+{
+public:
+	virtual ~IToolsContextTransactionsAPI() {}
+
+	/**
+	 * Request that context display message information.
+	 * @param Message text of message
+	 * @param Level severity level of message
+	 */
+	virtual void PostMessage(const TCHAR* Message, EToolMessageLevel Level) = 0;
+
+	/** 
+	 * Forward an invalidation request from Tools framework, to cause repaint/etc. 
+	 * This is not always necessary but in some situations (eg in Non-Realtime mode in Editor)
+	 * a redraw will not happen every frame. 
+	 * See UInputRouter for options to enable auto-invalidation.
+	 */
+	virtual void PostInvalidation() = 0;
+	
+	/**
+	 * Begin a Transaction, whatever this means in the current Context. For example in the
+	 * Editor it means open a GEditor Transaction. You must call EndUndoTransaction() after calling this.
+	 * @param Description text description of the transaction that could be shown to user
+	 */
+	virtual void BeginUndoTransaction(const FText& Description) = 0;
+
+	/**
+	 * Complete the Transaction. Assumption is that Begin/End are called in pairs.
+	 */
+	virtual void EndUndoTransaction() = 0;
+
+	/**
+	 * Insert an FChange into the transaction history in the current Context. 
+	 * This cannot be called between Begin/EndUndoTransaction, the FChange should be 
+	 * automatically inserted into a Transaction.
+	 * @param TargetObject The UObject this Change is applied to
+	 * @param Change The Change implementation
+	 * @param Description text description of the transaction that could be shown to user
+	 */
+	virtual void AppendChange(UObject* TargetObject, TUniquePtr<FChange> Change, const FText& Description) = 0;
+
+
+
+	/**
+	 * Request a modification to the current selected objects
+	 * @param SelectionChange desired modification to current selection
+	 * @return true if the selection change could be applied
+	 */
+	virtual bool RequestSelectionChange(const FSelectedOjectsChangeList& SelectionChange) = 0;
+
+};
+
+
+/**
+ * Users of the Tools Framework need to implement IToolsContextRenderAPI to allow
+ * Tools, Indicators, and Gizmos to make low-level rendering calls for things like line drawing.
+ * This API will be passed to eg UInteractiveTool::Render(), so access is only provided when it
+ * makes sense to call the functions
+ */
+class IToolsContextRenderAPI
+{
+public:
+	virtual ~IToolsContextRenderAPI() {}
+
+	/** @return Current PDI */
+	virtual FPrimitiveDrawInterface* GetPrimitiveDrawInterface() = 0;
+};
+
+
+
+
+/**
+ * Users of the Tools Framework need to provide an IToolsContextAssetAPI implementation
+ * that allows Packages and Assets to be created/saved. Note that this is not strictly
+ * necessary, for example a trivial implementation could just store things in the Transient
+ * package and not do any saving.
+ */
+class IToolsContextAssetAPI
+{
+public:
+	virtual ~IToolsContextAssetAPI() {}
+
+	/** Get default path to save assets in. For example the currently-visible path in the Editor. */
+	virtual FString GetActiveAssetFolderPath() = 0;
+
+	/** Combines folder and asset names */
+	virtual FString MakePackageName(const FString& AssetName, const FString& FolderPath) = 0;
+
+	/** return "unique" version of AssetName, created by appending _1, _2, ... if AssetName already exists */
+	virtual FString MakeUniqueAssetName(const FString& AssetName, const FString& FolderPath) = 0;
+
+	/** create a new package at the given path (calls MakePackageName) */
+	virtual UPackage* CreateNewPackage(const FString& AssetName, const FString& FolderPath) = 0;
+
+	/** Request saving of asset to persistent storage via something like an interactive popup dialog */
+	virtual void InteractiveSaveGeneratedAsset(UObject* Asset, UPackage* AssetPackage) = 0;
+
+	/** Autosave asset to persistent storage */
+	virtual void AutoSaveGeneratedAsset(UObject* Asset, UPackage* AssetPackage) = 0;
+};
+
diff --git a/Engine/Source/Runtime/InteractiveToolsFramework/InteractiveToolsFramework.Build.cs b/Engine/Source/Runtime/InteractiveToolsFramework/InteractiveToolsFramework.Build.cs
deleted file mode 100644
diff --git a/Engine/Source/Runtime/InteractiveToolsFramework/Private/BaseBehaviors/AnyButtonInputBehavior.cpp b/Engine/Source/Runtime/InteractiveToolsFramework/Private/BaseBehaviors/AnyButtonInputBehavior.cpp
deleted file mode 100644
diff --git a/Engine/Source/Runtime/InteractiveToolsFramework/Private/BaseBehaviors/ClickDragBehavior.cpp b/Engine/Source/Runtime/InteractiveToolsFramework/Private/BaseBehaviors/ClickDragBehavior.cpp
deleted file mode 100644
diff --git a/Engine/Source/Runtime/InteractiveToolsFramework/Private/BaseBehaviors/KeyAsModifierInputBehavior.cpp b/Engine/Source/Runtime/InteractiveToolsFramework/Private/BaseBehaviors/KeyAsModifierInputBehavior.cpp
deleted file mode 100644
diff --git a/Engine/Source/Runtime/InteractiveToolsFramework/Private/BaseBehaviors/MouseHoverBehavior.cpp b/Engine/Source/Runtime/InteractiveToolsFramework/Private/BaseBehaviors/MouseHoverBehavior.cpp
deleted file mode 100644
diff --git a/Engine/Source/Runtime/InteractiveToolsFramework/Private/BaseBehaviors/MultiClickSequenceInputBehavior.cpp b/Engine/Source/Runtime/InteractiveToolsFramework/Private/BaseBehaviors/MultiClickSequenceInputBehavior.cpp
deleted file mode 100644
diff --git a/Engine/Source/Runtime/InteractiveToolsFramework/Private/BaseBehaviors/SingleClickBehavior.cpp b/Engine/Source/Runtime/InteractiveToolsFramework/Private/BaseBehaviors/SingleClickBehavior.cpp
deleted file mode 100644
diff --git a/Engine/Source/Runtime/InteractiveToolsFramework/Private/BaseTools/ClickDragTool.cpp b/Engine/Source/Runtime/InteractiveToolsFramework/Private/BaseTools/ClickDragTool.cpp
deleted file mode 100644
diff --git a/Engine/Source/Runtime/InteractiveToolsFramework/Private/BaseTools/MeshSurfacePointTool.cpp b/Engine/Source/Runtime/InteractiveToolsFramework/Private/BaseTools/MeshSurfacePointTool.cpp
deleted file mode 100644
diff --git a/Engine/Source/Runtime/InteractiveToolsFramework/Private/BaseTools/SingleClickTool.cpp b/Engine/Source/Runtime/InteractiveToolsFramework/Private/BaseTools/SingleClickTool.cpp
deleted file mode 100644
diff --git a/Engine/Source/Runtime/InteractiveToolsFramework/Private/InputBehavior.cpp b/Engine/Source/Runtime/InteractiveToolsFramework/Private/InputBehavior.cpp
deleted file mode 100644
diff --git a/Engine/Source/Runtime/InteractiveToolsFramework/Private/InputBehaviorSet.cpp b/Engine/Source/Runtime/InteractiveToolsFramework/Private/InputBehaviorSet.cpp
deleted file mode 100644
diff --git a/Engine/Source/Runtime/InteractiveToolsFramework/Private/InputRouter.cpp b/Engine/Source/Runtime/InteractiveToolsFramework/Private/InputRouter.cpp
deleted file mode 100644
diff --git a/Engine/Source/Runtime/InteractiveToolsFramework/Private/InputState.cpp b/Engine/Source/Runtime/InteractiveToolsFramework/Private/InputState.cpp
deleted file mode 100644
diff --git a/Engine/Source/Runtime/InteractiveToolsFramework/Private/InteractiveGizmo.cpp b/Engine/Source/Runtime/InteractiveToolsFramework/Private/InteractiveGizmo.cpp
deleted file mode 100644
diff --git a/Engine/Source/Runtime/InteractiveToolsFramework/Private/InteractiveGizmoManager.cpp b/Engine/Source/Runtime/InteractiveToolsFramework/Private/InteractiveGizmoManager.cpp
deleted file mode 100644
diff --git a/Engine/Source/Runtime/InteractiveToolsFramework/Private/InteractiveTool.cpp b/Engine/Source/Runtime/InteractiveToolsFramework/Private/InteractiveTool.cpp
deleted file mode 100644
diff --git a/Engine/Source/Runtime/InteractiveToolsFramework/Private/InteractiveToolActionSet.cpp b/Engine/Source/Runtime/InteractiveToolsFramework/Private/InteractiveToolActionSet.cpp
deleted file mode 100644
diff --git a/Engine/Source/Runtime/InteractiveToolsFramework/Private/InteractiveToolBuilder.cpp b/Engine/Source/Runtime/InteractiveToolsFramework/Private/InteractiveToolBuilder.cpp
deleted file mode 100644
diff --git a/Engine/Source/Runtime/InteractiveToolsFramework/Private/InteractiveToolManager.cpp b/Engine/Source/Runtime/InteractiveToolsFramework/Private/InteractiveToolManager.cpp
deleted file mode 100644
diff --git a/Engine/Source/Runtime/InteractiveToolsFramework/Private/InteractiveToolsContext.cpp b/Engine/Source/Runtime/InteractiveToolsFramework/Private/InteractiveToolsContext.cpp
deleted file mode 100644
diff --git a/Engine/Source/Runtime/InteractiveToolsFramework/Private/InteractiveToolsFramework.cpp b/Engine/Source/Runtime/InteractiveToolsFramework/Private/InteractiveToolsFramework.cpp
deleted file mode 100644
diff --git a/Engine/Source/Runtime/InteractiveToolsFramework/Private/ToolBuilderUtil.cpp b/Engine/Source/Runtime/InteractiveToolsFramework/Private/ToolBuilderUtil.cpp
deleted file mode 100644
diff --git a/Engine/Source/Runtime/InteractiveToolsFramework/Public/BaseBehaviors/AnyButtonInputBehavior.h b/Engine/Source/Runtime/InteractiveToolsFramework/Public/BaseBehaviors/AnyButtonInputBehavior.h
deleted file mode 100644
diff --git a/Engine/Source/Runtime/InteractiveToolsFramework/Public/BaseBehaviors/BehaviorTargetInterfaces.h b/Engine/Source/Runtime/InteractiveToolsFramework/Public/BaseBehaviors/BehaviorTargetInterfaces.h
deleted file mode 100644
diff --git a/Engine/Source/Runtime/InteractiveToolsFramework/Public/BaseBehaviors/ClickDragBehavior.h b/Engine/Source/Runtime/InteractiveToolsFramework/Public/BaseBehaviors/ClickDragBehavior.h
deleted file mode 100644
diff --git a/Engine/Source/Runtime/InteractiveToolsFramework/Public/BaseBehaviors/InputBehaviorModifierStates.h b/Engine/Source/Runtime/InteractiveToolsFramework/Public/BaseBehaviors/InputBehaviorModifierStates.h
deleted file mode 100644
diff --git a/Engine/Source/Runtime/InteractiveToolsFramework/Public/BaseBehaviors/KeyAsModifierInputBehavior.h b/Engine/Source/Runtime/InteractiveToolsFramework/Public/BaseBehaviors/KeyAsModifierInputBehavior.h
deleted file mode 100644
diff --git a/Engine/Source/Runtime/InteractiveToolsFramework/Public/BaseBehaviors/MouseHoverBehavior.h b/Engine/Source/Runtime/InteractiveToolsFramework/Public/BaseBehaviors/MouseHoverBehavior.h
deleted file mode 100644
diff --git a/Engine/Source/Runtime/InteractiveToolsFramework/Public/BaseBehaviors/MultiClickSequenceInputBehavior.h b/Engine/Source/Runtime/InteractiveToolsFramework/Public/BaseBehaviors/MultiClickSequenceInputBehavior.h
deleted file mode 100644
diff --git a/Engine/Source/Runtime/InteractiveToolsFramework/Public/BaseBehaviors/SingleClickBehavior.h b/Engine/Source/Runtime/InteractiveToolsFramework/Public/BaseBehaviors/SingleClickBehavior.h
deleted file mode 100644
diff --git a/Engine/Source/Runtime/InteractiveToolsFramework/Public/BaseTools/ClickDragTool.h b/Engine/Source/Runtime/InteractiveToolsFramework/Public/BaseTools/ClickDragTool.h
deleted file mode 100644
diff --git a/Engine/Source/Runtime/InteractiveToolsFramework/Public/BaseTools/MeshSurfacePointTool.h b/Engine/Source/Runtime/InteractiveToolsFramework/Public/BaseTools/MeshSurfacePointTool.h
deleted file mode 100644
diff --git a/Engine/Source/Runtime/InteractiveToolsFramework/Public/BaseTools/SingleClickTool.h b/Engine/Source/Runtime/InteractiveToolsFramework/Public/BaseTools/SingleClickTool.h
deleted file mode 100644
diff --git a/Engine/Source/Runtime/InteractiveToolsFramework/Public/ComponentSourceInterfaces.h b/Engine/Source/Runtime/InteractiveToolsFramework/Public/ComponentSourceInterfaces.h
deleted file mode 100644
diff --git a/Engine/Source/Runtime/InteractiveToolsFramework/Public/InputBehavior.h b/Engine/Source/Runtime/InteractiveToolsFramework/Public/InputBehavior.h
deleted file mode 100644
diff --git a/Engine/Source/Runtime/InteractiveToolsFramework/Public/InputBehaviorSet.h b/Engine/Source/Runtime/InteractiveToolsFramework/Public/InputBehaviorSet.h
deleted file mode 100644
diff --git a/Engine/Source/Runtime/InteractiveToolsFramework/Public/InputRouter.h b/Engine/Source/Runtime/InteractiveToolsFramework/Public/InputRouter.h
deleted file mode 100644
diff --git a/Engine/Source/Runtime/InteractiveToolsFramework/Public/InputState.h b/Engine/Source/Runtime/InteractiveToolsFramework/Public/InputState.h
deleted file mode 100644
diff --git a/Engine/Source/Runtime/InteractiveToolsFramework/Public/InteractiveGizmo.h b/Engine/Source/Runtime/InteractiveToolsFramework/Public/InteractiveGizmo.h
deleted file mode 100644
diff --git a/Engine/Source/Runtime/InteractiveToolsFramework/Public/InteractiveGizmoBuilder.h b/Engine/Source/Runtime/InteractiveToolsFramework/Public/InteractiveGizmoBuilder.h
deleted file mode 100644
diff --git a/Engine/Source/Runtime/InteractiveToolsFramework/Public/InteractiveGizmoManager.h b/Engine/Source/Runtime/InteractiveToolsFramework/Public/InteractiveGizmoManager.h
deleted file mode 100644
diff --git a/Engine/Source/Runtime/InteractiveToolsFramework/Public/InteractiveTool.h b/Engine/Source/Runtime/InteractiveToolsFramework/Public/InteractiveTool.h
deleted file mode 100644
diff --git a/Engine/Source/Runtime/InteractiveToolsFramework/Public/InteractiveToolActionSet.h b/Engine/Source/Runtime/InteractiveToolsFramework/Public/InteractiveToolActionSet.h
deleted file mode 100644
diff --git a/Engine/Source/Runtime/InteractiveToolsFramework/Public/InteractiveToolBuilder.h b/Engine/Source/Runtime/InteractiveToolsFramework/Public/InteractiveToolBuilder.h
deleted file mode 100644
diff --git a/Engine/Source/Runtime/InteractiveToolsFramework/Public/InteractiveToolManager.h b/Engine/Source/Runtime/InteractiveToolsFramework/Public/InteractiveToolManager.h
deleted file mode 100644
diff --git a/Engine/Source/Runtime/InteractiveToolsFramework/Public/InteractiveToolsContext.h b/Engine/Source/Runtime/InteractiveToolsFramework/Public/InteractiveToolsContext.h
deleted file mode 100644
diff --git a/Engine/Source/Runtime/InteractiveToolsFramework/Public/InteractiveToolsFramework.h b/Engine/Source/Runtime/InteractiveToolsFramework/Public/InteractiveToolsFramework.h
deleted file mode 100644
diff --git a/Engine/Source/Runtime/InteractiveToolsFramework/Public/ToolBuilderUtil.h b/Engine/Source/Runtime/InteractiveToolsFramework/Public/ToolBuilderUtil.h
deleted file mode 100644
diff --git a/Engine/Source/Runtime/InteractiveToolsFramework/Public/ToolContextInterfaces.h b/Engine/Source/Runtime/InteractiveToolsFramework/Public/ToolContextInterfaces.h
deleted file mode 100644
