Fix Jira UE-82061
Add signal handling to accept and tear down current tool prior to World destruction

diff --git a/Engine/Source/Editor/Experimental/EditorInteractiveToolsFramework/Private/EdModeInteractiveToolsContext.cpp b/Engine/Source/Editor/Experimental/EditorInteractiveToolsFramework/Private/EdModeInteractiveToolsContext.cpp
--- a/Engine/Source/Editor/Experimental/EditorInteractiveToolsFramework/Private/EdModeInteractiveToolsContext.cpp
+++ b/Engine/Source/Editor/Experimental/EditorInteractiveToolsFramework/Private/EdModeInteractiveToolsContext.cpp
@@ -5,7 +5,9 @@
 #include "Editor.h"
 #include "EditorViewportClient.h"
 #include "EditorModeManager.h"
+#include "LevelEditor.h"
 #include "LevelEditorViewport.h"   // for GCurrentLevelEditingViewportClient
+#include "Modules/ModuleManager.h"
 #include "ShowFlags.h"				// for EngineShowFlags
 #include "Engine/Selection.h"
 #include "Misc/ITransaction.h"
@@ -343,6 +345,15 @@ void UEdModeInteractiveToolsContext::Initialize(IToolsContextQueriesAPI* Queries
 		TerminateActiveToolsOnSaveWorld();
 	});
 
+	FLevelEditorModule& LevelEditor = FModuleManager::LoadModuleChecked<FLevelEditorModule>("LevelEditor");
+	WorldTearDownDelegateHandle = LevelEditor.OnMapChanged().AddLambda([this](UWorld* World, EMapChangeType ChangeType)
+	{
+		if (ChangeType == EMapChangeType::TearDownWorld)
+		{
+			TerminateActiveToolsOnWorldTearDown();
+		}
+	});
+
 	ToolManager->OnToolEnded.AddLambda([this](UInteractiveToolManager*, UInteractiveTool*)
 	{
 		RestoreEditorState();
@@ -354,6 +365,9 @@ void UEdModeInteractiveToolsContext::Initialize(IToolsContextQueriesAPI* Queries
 
 void UEdModeInteractiveToolsContext::Shutdown()
 {
+
+	FLevelEditorModule& LevelEditor = FModuleManager::GetModuleChecked<FLevelEditorModule>("LevelEditor");
+	LevelEditor.OnMapChanged().Remove(WorldTearDownDelegateHandle);
 	FEditorDelegates::BeginPIE.Remove(BeginPIEDelegateHandle);
 	FEditorDelegates::PreSaveWorld.Remove(PreSaveWorldDelegateHandle);
 
@@ -434,6 +448,10 @@ void UEdModeInteractiveToolsContext::TerminateActiveToolsOnSaveWorld()
 {
 	DeactivateAllActiveTools();
 }
+void UEdModeInteractiveToolsContext::TerminateActiveToolsOnWorldTearDown()
+{
+	DeactivateAllActiveTools();
+}
 
 
 
diff --git a/Engine/Source/Editor/Experimental/EditorInteractiveToolsFramework/Public/EdModeInteractiveToolsContext.h b/Engine/Source/Editor/Experimental/EditorInteractiveToolsFramework/Public/EdModeInteractiveToolsContext.h
--- a/Engine/Source/Editor/Experimental/EditorInteractiveToolsFramework/Public/EdModeInteractiveToolsContext.h
+++ b/Engine/Source/Editor/Experimental/EditorInteractiveToolsFramework/Public/EdModeInteractiveToolsContext.h
@@ -32,6 +32,9 @@ public:
 	// default behavior is to accept active tool
 	virtual void TerminateActiveToolsOnSaveWorld();
 
+	// default behavior is to accept active tool
+	virtual void TerminateActiveToolsOnWorldTearDown();
+
 	IToolsContextQueriesAPI* GetQueriesAPI() const { return QueriesAPI; }
 	IToolsContextTransactionsAPI* GetTransactionAPI() const { return TransactionAPI; }
 	IToolsContextAssetAPI* GetAssetAPI() const { return AssetAPI; }
@@ -99,6 +102,8 @@ protected:
 	FDelegateHandle BeginPIEDelegateHandle;
 	// called before a Save starts. This currently shuts down active tools.
 	FDelegateHandle PreSaveWorldDelegateHandle;
+	// called when a map is changed
+	FDelegateHandle WorldTearDownDelegateHandle;
 
 	// EdMode implementation of InteractiveToolFramework APIs - see ToolContextInterfaces.h
 	IToolsContextQueriesAPI* QueriesAPI;
