Don't allow Actors being actively sculpted to be deleted in Editor (copy CL 9944430 from Dev-Editor)
#rnx

diff --git a/Engine/Plugins/Experimental/MeshModelingToolset/Source/ModelingComponents/Public/BaseDynamicMeshComponent.h b/Engine/Plugins/Experimental/MeshModelingToolset/Source/ModelingComponents/Public/BaseDynamicMeshComponent.h
--- a/Engine/Plugins/Experimental/MeshModelingToolset/Source/ModelingComponents/Public/BaseDynamicMeshComponent.h
+++ b/Engine/Plugins/Experimental/MeshModelingToolset/Source/ModelingComponents/Public/BaseDynamicMeshComponent.h
@@ -5,6 +5,7 @@
 #include "CoreMinimal.h"
 #include "UObject/ObjectMacros.h"
 #include "Components/MeshComponent.h"
+#include "InteractiveToolObjects.h"
 #include "Changes/MeshVertexChange.h"
 #include "Changes/MeshChange.h"
 
@@ -40,7 +41,7 @@ enum class EDynamicMeshTangentCalcType : uint8
  * Currently no functionality lives here, only some interface functions are defined that various subclasses implement.
  */
 UCLASS(hidecategories = (LOD, Physics, Collision), editinlinenew, ClassGroup = Rendering)
-class MODELINGCOMPONENTS_API UBaseDynamicMeshComponent : public UMeshComponent, public IMeshVertexCommandChangeTarget, public IMeshCommandChangeTarget
+class MODELINGCOMPONENTS_API UBaseDynamicMeshComponent : public UMeshComponent, public IToolFrameworkComponent, public IMeshVertexCommandChangeTarget, public IMeshCommandChangeTarget
 {
 	GENERATED_UCLASS_BODY()
 
diff --git a/Engine/Source/Editor/Experimental/EditorInteractiveToolsFramework/Private/EdModeInteractiveToolsContext.cpp b/Engine/Source/Editor/Experimental/EditorInteractiveToolsFramework/Private/EdModeInteractiveToolsContext.cpp
--- a/Engine/Source/Editor/Experimental/EditorInteractiveToolsFramework/Private/EdModeInteractiveToolsContext.cpp
+++ b/Engine/Source/Editor/Experimental/EditorInteractiveToolsFramework/Private/EdModeInteractiveToolsContext.cpp
@@ -516,17 +516,33 @@ bool UEdModeInteractiveToolsContext::ProcessEditDelete()
 
 	bool bSkipDelete = false;
 
-	// Test if any of the selected actors are AInternalToolFrameworkActor
-	// subclasses. In this case we do not want to allow them to be deleted,
-	// as generally this will cause problems for the Tool.
 	USelection* SelectedActors = GEditor->GetSelectedActors();
-	for (int i = 0; i < SelectedActors->Num(); ++i)
+	for (int i = 0; i < SelectedActors->Num() && bSkipDelete == false; ++i)
 	{
 		UObject* SelectedActor = SelectedActors->GetSelectedObject(i);
-		if ( Cast<AInternalToolFrameworkActor>(SelectedActor) != nullptr) 
+
+		// If any of the selected actors are AInternalToolFrameworkActor, we do not want to allow them to be deleted,
+		// as generally this will cause problems for the Tool.
+		if ( Cast<AInternalToolFrameworkActor>(SelectedActor) != nullptr)
 		{
 			bSkipDelete = true;
 		}
+
+		// If any Components of selected Actors implement UToolFrameworkComponent, we disable delete (for now).
+		// (Currently Sculpt and a few other Modeling Tools attach their preview mesh components to the selected Actor)
+		AActor* Actor = Cast<AActor>(SelectedActor);
+		if (Actor != nullptr)
+		{
+			const TSet<UActorComponent*>& Components = Actor->GetComponents();
+			for (const UActorComponent* Component : Components)
+			{
+				if ( Component->Implements<UToolFrameworkComponent>() )
+				{
+					bSkipDelete = true;
+				}
+			}
+		}
+
 	}
 
 	return bSkipDelete;
diff --git a/Engine/Source/Runtime/Experimental/InteractiveToolsFramework/Public/InteractiveToolObjects.h b/Engine/Source/Runtime/Experimental/InteractiveToolsFramework/Public/InteractiveToolObjects.h
--- a/Engine/Source/Runtime/Experimental/InteractiveToolsFramework/Public/InteractiveToolObjects.h
+++ b/Engine/Source/Runtime/Experimental/InteractiveToolsFramework/Public/InteractiveToolObjects.h
@@ -3,6 +3,7 @@
 #pragma once
 
 #include "CoreMinimal.h"
+#include "UObject/Interface.h"
 #include "GameFramework/Actor.h"
 #include "InteractiveToolObjects.generated.h"
 
@@ -24,3 +25,22 @@ public:
 };
 
 
+
+
+
+// UInterface for IToolFrameworkComponent
+UINTERFACE(MinimalAPI)
+class UToolFrameworkComponent : public UInterface
+{
+	GENERATED_BODY()
+};
+
+// IToolFrameworkComponent is an interface we can attach to custom Components used in the Tools framework.
+// Currently this interface is only used to identify such Components (to prevent certain Actors from being deleted)
+class IToolFrameworkComponent
+{
+	GENERATED_BODY()
+
+public:
+};
+
