Automatically disable viewport anti-aliasing while a Tool is active, and re-enable previous AA state on tool exit.
#rnx

diff --git a/Engine/Source/Editor/Experimental/EditorInteractiveToolsFramework/Private/EdModeInteractiveToolsContext.cpp b/Engine/Source/Editor/Experimental/EditorInteractiveToolsFramework/Private/EdModeInteractiveToolsContext.cpp
--- a/Engine/Source/Editor/Experimental/EditorInteractiveToolsFramework/Private/EdModeInteractiveToolsContext.cpp
+++ b/Engine/Source/Editor/Experimental/EditorInteractiveToolsFramework/Private/EdModeInteractiveToolsContext.cpp
@@ -6,6 +6,7 @@
 #include "EditorViewportClient.h"
 #include "EditorModeManager.h"
 #include "LevelEditorViewport.h"   // for GCurrentLevelEditingViewportClient
+#include "ShowFlags.h"				// for EngineShowFlags
 #include "Engine/Selection.h"
 #include "Misc/ITransaction.h"
 #include "ScopedTransaction.h"
@@ -333,6 +334,9 @@ void UEdModeInteractiveToolsContext::Shutdown()
 	FEditorDelegates::BeginPIE.Remove(BeginPIEDelegateHandle);
 	FEditorDelegates::PreSaveWorld.Remove(PreSaveWorldDelegateHandle);
 
+	// auto-accept any in-progress tools
+	DeactivateAllActiveTools();
+
 	UInteractiveToolsContext::Shutdown();
 }
 
@@ -396,33 +400,11 @@ void UEdModeInteractiveToolsContext::ShutdownContext()
 
 void UEdModeInteractiveToolsContext::TerminateActiveToolsOnPIEStart()
 {
-	if (ToolManager->HasActiveTool(EToolSide::Left))
-	{
-		EToolShutdownType ShutdownType = ToolManager->CanAcceptActiveTool(EToolSide::Left) ?
-			EToolShutdownType::Accept : EToolShutdownType::Cancel;
-		ToolManager->DeactivateTool(EToolSide::Left, ShutdownType);
-	}
-	if (ToolManager->HasActiveTool(EToolSide::Right))
-	{
-		EToolShutdownType ShutdownType = ToolManager->CanAcceptActiveTool(EToolSide::Right) ?
-			EToolShutdownType::Accept : EToolShutdownType::Cancel;
-		ToolManager->DeactivateTool(EToolSide::Right, ShutdownType);
-	}
+	DeactivateAllActiveTools();
 }
 void UEdModeInteractiveToolsContext::TerminateActiveToolsOnSaveWorld()
 {
-	if (ToolManager->HasActiveTool(EToolSide::Left))
-	{
-		EToolShutdownType ShutdownType = ToolManager->CanAcceptActiveTool(EToolSide::Left) ?
-			EToolShutdownType::Accept : EToolShutdownType::Cancel;
-		ToolManager->DeactivateTool(EToolSide::Left, ShutdownType);
-	}
-	if (ToolManager->HasActiveTool(EToolSide::Right))
-	{
-		EToolShutdownType ShutdownType = ToolManager->CanAcceptActiveTool(EToolSide::Right) ?
-			EToolShutdownType::Accept : EToolShutdownType::Cancel;
-		ToolManager->DeactivateTool(EToolSide::Right, ShutdownType);
-	}
+	DeactivateAllActiveTools();
 }
 
 
@@ -500,7 +482,7 @@ bool UEdModeInteractiveToolsContext::InputKey(FEditorViewportClient* ViewportCli
 		{
 			if (ToolManager->HasActiveTool(EToolSide::Mouse))
 			{
-				ToolManager->DeactivateTool(EToolSide::Mouse, EToolShutdownType::Cancel);
+				DeactivateActiveTool(EToolSide::Mouse, EToolShutdownType::Cancel);
 			}
 			return true;
 		}
@@ -515,13 +497,13 @@ bool UEdModeInteractiveToolsContext::InputKey(FEditorViewportClient* ViewportCli
 			{
 				if (ToolManager->CanAcceptActiveTool(EToolSide::Mouse))
 				{
-					ToolManager->DeactivateTool(EToolSide::Mouse, EToolShutdownType::Accept);
+					DeactivateActiveTool(EToolSide::Mouse, EToolShutdownType::Accept);
 					return true;
 				}
 			}
 			else
 			{
-				ToolManager->DeactivateTool(EToolSide::Mouse, EToolShutdownType::Completed);
+				DeactivateActiveTool(EToolSide::Mouse, EToolShutdownType::Completed);
 				return true;
 			}
 		}
@@ -753,15 +735,18 @@ bool UEdModeInteractiveToolsContext::CanStartTool(const FString& ToolTypeIdentif
 	return (ToolManager->HasActiveTool(EToolSide::Mouse) == false) &&
 		(ToolManager->CanActivateTool(EToolSide::Mouse, ToolTypeIdentifier) == true);
 }
+
 bool UEdModeInteractiveToolsContext::ActiveToolHasAccept() const
 {
 	return ToolManager->HasActiveTool(EToolSide::Mouse) &&
 		ToolManager->GetActiveTool(EToolSide::Mouse)->HasAccept();
 }
+
 bool UEdModeInteractiveToolsContext::CanAcceptActiveTool() const
 {
 	return ToolManager->CanAcceptActiveTool(EToolSide::Mouse);
 }
+
 bool UEdModeInteractiveToolsContext::CanCancelActiveTool() const
 {
 	return ToolManager->CanCancelActiveTool(EToolSide::Mouse);
@@ -771,6 +756,7 @@ bool UEdModeInteractiveToolsContext::CanCompleteActiveTool() const
 {
 	return ToolManager->HasActiveTool(EToolSide::Mouse) && CanCancelActiveTool() == false;
 }
+
 void UEdModeInteractiveToolsContext::StartTool(const FString& ToolTypeIdentifier)
 {
 	if (ToolManager->SelectActiveToolType(EToolSide::Mouse, ToolTypeIdentifier) == false)
@@ -780,9 +766,58 @@ void UEdModeInteractiveToolsContext::StartTool(const FString& ToolTypeIdentifier
 	else
 	{
 		ToolManager->ActivateTool(EToolSide::Mouse);
+		SaveEditorStateAndSetForTool();
 	}
 }
+
 void UEdModeInteractiveToolsContext::EndTool(EToolShutdownType ShutdownType)
 {
-	ToolManager->DeactivateTool(EToolSide::Mouse, ShutdownType);
+	DeactivateActiveTool(EToolSide::Mouse, ShutdownType);
+}
+
+
+
+void UEdModeInteractiveToolsContext::DeactivateActiveTool(EToolSide WhichSide, EToolShutdownType ShutdownType)
+{
+	ToolManager->DeactivateTool(WhichSide, ShutdownType);
+	RestoreEditorState();
+}
+
+void UEdModeInteractiveToolsContext::DeactivateAllActiveTools()
+{
+	if (ToolManager->HasActiveTool(EToolSide::Left))
+	{
+		EToolShutdownType ShutdownType = ToolManager->CanAcceptActiveTool(EToolSide::Left) ?
+			EToolShutdownType::Accept : EToolShutdownType::Cancel;
+		ToolManager->DeactivateTool(EToolSide::Left, ShutdownType);
+	}
+	if (ToolManager->HasActiveTool(EToolSide::Right))
+	{
+		EToolShutdownType ShutdownType = ToolManager->CanAcceptActiveTool(EToolSide::Right) ?
+			EToolShutdownType::Accept : EToolShutdownType::Cancel;
+		ToolManager->DeactivateTool(EToolSide::Right, ShutdownType);
+	}
+
+	RestoreEditorState();
+}
+
+
+
+
+void UEdModeInteractiveToolsContext::SaveEditorStateAndSetForTool()
+{
+	check(bHaveSavedEditorState == false);
+	bSavedAntiAliasingState = GCurrentLevelEditingViewportClient->EngineShowFlags.AntiAliasing;
+	bHaveSavedEditorState = true;
+
+	GCurrentLevelEditingViewportClient->EngineShowFlags.SetAntiAliasing(false);
+}
+
+void UEdModeInteractiveToolsContext::RestoreEditorState()
+{
+	if (bHaveSavedEditorState)
+	{
+		GCurrentLevelEditingViewportClient->EngineShowFlags.SetAntiAliasing(bSavedAntiAliasingState);
+		bHaveSavedEditorState = false;
+	}
 }
\ No newline at end of file
diff --git a/Engine/Source/Editor/Experimental/EditorInteractiveToolsFramework/Public/EdModeInteractiveToolsContext.h b/Engine/Source/Editor/Experimental/EditorInteractiveToolsFramework/Public/EdModeInteractiveToolsContext.h
--- a/Engine/Source/Editor/Experimental/EditorInteractiveToolsFramework/Public/EdModeInteractiveToolsContext.h
+++ b/Engine/Source/Editor/Experimental/EditorInteractiveToolsFramework/Public/EdModeInteractiveToolsContext.h
@@ -73,6 +73,8 @@ protected:
 	virtual void Initialize(IToolsContextQueriesAPI* QueriesAPI, IToolsContextTransactionsAPI* TransactionsAPI) override;
 	virtual void Shutdown() override;
 
+	virtual void DeactivateActiveTool(EToolSide WhichSide, EToolShutdownType ShutdownType);
+	virtual void DeactivateAllActiveTools();
 
 public:
 	UPROPERTY()
@@ -81,19 +83,33 @@ public:
 protected:
 	FEdMode* EditorMode;
 
+	// called when PIE is about to start, shuts down active tools
 	FDelegateHandle BeginPIEDelegateHandle;
+	// called before a Save starts. This currently shuts down active tools.
 	FDelegateHandle PreSaveWorldDelegateHandle;
 
+	// EdMode implementation of InteractiveToolFramework APIs - see ToolContextInterfaces.h
 	IToolsContextQueriesAPI* QueriesAPI;
 	IToolsContextTransactionsAPI* TransactionAPI;
 	IToolsContextAssetAPI* AssetAPI;
 	IComponentSourceFactory* SourceFactory;
 
+	// if true, we invalidate the ViewportClient on next tick
 	bool bInvalidationPending;
 
 	/** Input event instance used to keep track of various button states, etc, that we cannot directly query on-demand */
 	FInputDeviceState CurrentMouseState;
 
+	// Utility function to convert viewport x/y from mouse events (and others?) into scene ray.
+	// Copy-pasted from other Editor code, seems kind of expensive?
 	static FRay GetRayFromMousePos(FEditorViewportClient* ViewportClient, FViewport* Viewport, int MouseX, int MouseY);
 
+
+	// editor UI state that we set before starting tool and when exiting tool
+	// Currently disabling anti-aliasing during active Tools because it causes PDI flickering
+	bool bHaveSavedEditorState = false;
+	bool bSavedAntiAliasingState = false;
+	void SaveEditorStateAndSetForTool();
+	void RestoreEditorState();
+
 };
