don't try to restore editor state if program is exiting

diff --git a/Engine/Source/Editor/Experimental/EditorInteractiveToolsFramework/Private/EdModeInteractiveToolsContext.cpp b/Engine/Source/Editor/Experimental/EditorInteractiveToolsFramework/Private/EdModeInteractiveToolsContext.cpp
--- a/Engine/Source/Editor/Experimental/EditorInteractiveToolsFramework/Private/EdModeInteractiveToolsContext.cpp
+++ b/Engine/Source/Editor/Experimental/EditorInteractiveToolsFramework/Private/EdModeInteractiveToolsContext.cpp
@@ -815,7 +815,7 @@ void UEdModeInteractiveToolsContext::SaveEditorStateAndSetForTool()
 
 void UEdModeInteractiveToolsContext::RestoreEditorState()
 {
-	if (bHaveSavedEditorState)
+	if (bHaveSavedEditorState && !GIsRequestingExit)
 	{
 		GCurrentLevelEditingViewportClient->EngineShowFlags.SetAntiAliasing(bSavedAntiAliasingState);
 		bHaveSavedEditorState = false;
