Add AInternalToolFrameworkActor base AActor subclass for Actors that Tools create internally. EdModeInteractiveToolsContext::ProcessEditDelete() checks if any to-be-deleted Actors can be cast to this type, and if so, ignores delete operation.
#rnx

diff --git a/Engine/Source/Editor/Experimental/EditorInteractiveToolsFramework/Private/EdModeInteractiveToolsContext.cpp b/Engine/Source/Editor/Experimental/EditorInteractiveToolsFramework/Private/EdModeInteractiveToolsContext.cpp
--- a/Engine/Source/Editor/Experimental/EditorInteractiveToolsFramework/Private/EdModeInteractiveToolsContext.cpp
+++ b/Engine/Source/Editor/Experimental/EditorInteractiveToolsFramework/Private/EdModeInteractiveToolsContext.cpp
@@ -16,6 +16,7 @@
 
 #include "EditorToolAssetAPI.h"
 #include "EditorComponentSourceFactory.h"
+#include "InteractiveToolObjects.h"
 
 //#include "PhysicsEngine/BodySetup.h"
 //#include "Interfaces/Interface_CollisionDataProvider.h"
@@ -461,6 +462,32 @@ void UEdModeInteractiveToolsContext::Render(const FSceneView* View, FViewport* V
 
 
 
+bool UEdModeInteractiveToolsContext::ProcessEditDelete()
+{
+	if (ToolManager->HasAnyActiveTool() == false)
+	{
+		return false;
+	}
+
+	bool bSkipDelete = false;
+
+	// Test if any of the selected actors are AInternalToolFrameworkActor
+	// subclasses. In this case we do not want to allow them to be deleted,
+	// as generally this will cause problems for the Tool.
+	USelection* SelectedActors = GEditor->GetSelectedActors();
+	for (int i = 0; i < SelectedActors->Num(); ++i)
+	{
+		UObject* SelectedActor = SelectedActors->GetSelectedObject(i);
+		if ( Cast<AInternalToolFrameworkActor>(SelectedActor) != nullptr) 
+		{
+			bSkipDelete = true;
+		}
+	}
+
+	return bSkipDelete;
+}
+
+
 
 bool UEdModeInteractiveToolsContext::InputKey(FEditorViewportClient* ViewportClient, FViewport* Viewport, FKey Key, EInputEvent Event)
 {
diff --git a/Engine/Source/Editor/Experimental/EditorInteractiveToolsFramework/Public/EdModeInteractiveToolsContext.h b/Engine/Source/Editor/Experimental/EditorInteractiveToolsFramework/Public/EdModeInteractiveToolsContext.h
--- a/Engine/Source/Editor/Experimental/EditorInteractiveToolsFramework/Public/EdModeInteractiveToolsContext.h
+++ b/Engine/Source/Editor/Experimental/EditorInteractiveToolsFramework/Public/EdModeInteractiveToolsContext.h
@@ -44,6 +44,8 @@ public:
 	virtual void Tick(FEditorViewportClient* ViewportClient, float DeltaTime);
 	virtual void Render(const FSceneView* View, FViewport* Viewport, FPrimitiveDrawInterface* PDI);
 
+	virtual bool ProcessEditDelete();
+
 	virtual bool InputKey(FEditorViewportClient* ViewportClient, FViewport* Viewport, FKey Key, EInputEvent Event);
 
 	virtual bool MouseEnter(FEditorViewportClient* ViewportClient, FViewport* Viewport, int32 x, int32 y);
diff --git a/Engine/Source/Runtime/Experimental/InteractiveToolsFramework/Public/InteractiveToolObjects.h b/Engine/Source/Runtime/Experimental/InteractiveToolsFramework/Public/InteractiveToolObjects.h
new file mode 100644
--- /dev/null
+++ b/Engine/Source/Runtime/Experimental/InteractiveToolsFramework/Public/InteractiveToolObjects.h
@@ -0,0 +1,25 @@
+// Copyright 1998-2019 Epic Games, Inc. All Rights Reserved.
+
+#pragma once
+
+#include "CoreMinimal.h"
+#include "InteractiveToolObjects.generated.h"
+
+
+/**
+ * AInternalToolFrameworkActor is a base class for internal Actors that the
+ * ToolsFramework may spawn (for gizmos, mesh previews, etc). These Actors
+ * may need special-case handling, for example to prevent the user from
+ * selecting and deleting them. 
+ */
+UCLASS(Transient)
+class INTERACTIVETOOLSFRAMEWORK_API AInternalToolFrameworkActor : public AActor
+{
+	GENERATED_BODY()
+private:
+
+public:
+
+};
+
+
