ToolsFramework: restore editor state if a Tool ends via delegate, instead of assuming it ends via command (eg due to undo-cancel)
#rnx

diff --git a/Engine/Source/Editor/Experimental/EditorInteractiveToolsFramework/Private/EdModeInteractiveToolsContext.cpp b/Engine/Source/Editor/Experimental/EditorInteractiveToolsFramework/Private/EdModeInteractiveToolsContext.cpp
--- a/Engine/Source/Editor/Experimental/EditorInteractiveToolsFramework/Private/EdModeInteractiveToolsContext.cpp
+++ b/Engine/Source/Editor/Experimental/EditorInteractiveToolsFramework/Private/EdModeInteractiveToolsContext.cpp
@@ -325,6 +325,12 @@ void UEdModeInteractiveToolsContext::Initialize(IToolsContextQueriesAPI* Queries
 		TerminateActiveToolsOnSaveWorld();
 	});
 
+	ToolManager->OnToolEnded.AddLambda([this](UInteractiveToolManager*, UInteractiveTool*)
+	{
+		RestoreEditorState();
+	});
+
+
 	bInvalidationPending = false;
 }
 
