ToolsFramework: in Orthographic views, ViewportTransform.Rotation is not initialized, so construct rotation based on similar code in FEditorViewportClient
#rnx

diff --git a/Engine/Plugins/Experimental/MeshModelingToolset/Source/ModelingComponents/Private/ToolSceneQueriesUtil.cpp b/Engine/Plugins/Experimental/MeshModelingToolset/Source/ModelingComponents/Private/ToolSceneQueriesUtil.cpp
--- a/Engine/Plugins/Experimental/MeshModelingToolset/Source/ModelingComponents/Private/ToolSceneQueriesUtil.cpp
+++ b/Engine/Plugins/Experimental/MeshModelingToolset/Source/ModelingComponents/Private/ToolSceneQueriesUtil.cpp
@@ -66,20 +66,26 @@ double ToolSceneQueriesUtil::CalculateDimensionFromVisualAngleD(const FViewCamer
 
 
 
-
 bool ToolSceneQueriesUtil::IsPointVisible(const FViewCameraState& CameraState, const FVector3d& Point)
 {
-	FVector3d PointDir = (Point - CameraState.Position);
-	//@todo should use view frustum here!
-	if (PointDir.Dot(CameraState.Forward()) < 0.25)		// ballpark estimate
+	if (CameraState.bIsOrthographic == false)
 	{
-		return false;
+		FVector3d PointDir = (Point - CameraState.Position);
+		//@todo should use view frustum here!
+		if (PointDir.Dot(CameraState.Forward()) < 0.25)		// ballpark estimate
+		{
+			return false;
+		}
+	}
+	else
+	{
+		// @todo probably not always true but it's not exactly clear how ortho camera is configured...
+		return true;
 	}
 	return true;
 }
 
 
-
 bool ToolSceneQueriesUtil::FindSceneSnapPoint(const UInteractiveTool* Tool, const FVector3d& Point, FVector3d& SnapPointOut,
 	bool bVertices, bool bEdges, double VisualAngleThreshold, 
 	FSnapGeometry* SnapGeometry, FVector* DebugTriangleOut)
diff --git a/Engine/Source/Editor/Experimental/EditorInteractiveToolsFramework/Private/EdModeInteractiveToolsContext.cpp b/Engine/Source/Editor/Experimental/EditorInteractiveToolsFramework/Private/EdModeInteractiveToolsContext.cpp
--- a/Engine/Source/Editor/Experimental/EditorInteractiveToolsFramework/Private/EdModeInteractiveToolsContext.cpp
+++ b/Engine/Source/Editor/Experimental/EditorInteractiveToolsFramework/Private/EdModeInteractiveToolsContext.cpp
@@ -44,9 +44,42 @@ public:
 	void CacheCurrentViewState(FEditorViewportClient* ViewportClient)
 	{
 		FViewportCameraTransform ViewTransform = ViewportClient->GetViewTransform();
-		CachedViewState.Position = ViewTransform.GetLocation();
-		CachedViewState.Orientation = ViewTransform.GetRotation().Quaternion();
 		CachedViewState.bIsOrthographic = ViewportClient->IsOrtho();
+		CachedViewState.Position = ViewTransform.GetLocation();
+
+		// ViewTransform rotation is only initialized for perspective!
+		if (CachedViewState.bIsOrthographic == false)
+		{
+			CachedViewState.Orientation = ViewTransform.GetRotation().Quaternion();
+		}
+		else
+		{
+			// These rotations are based on hardcoded values in EditorViewportClient.cpp, see switches in FEditorViewportClient::CalcSceneView and FEditorViewportClient::Draw
+			switch (ViewportClient->ViewportType)
+			{
+			case LVT_OrthoXY:
+				CachedViewState.Orientation = FQuat(FRotator(-90.0f, -90.0f, 0.0f));
+				break;
+			case LVT_OrthoNegativeXY:
+				CachedViewState.Orientation = FQuat(FRotator(90.0f, 90.0f, 0.0f));
+				break;
+			case LVT_OrthoXZ:
+				CachedViewState.Orientation = FQuat(FRotator(0.0f, -90.0f, 0.0f));
+				break;
+			case LVT_OrthoNegativeXZ:
+				CachedViewState.Orientation = FQuat(FRotator(0.0f, 90.0f, 0.0f));
+				break;
+			case LVT_OrthoYZ:
+				CachedViewState.Orientation = FQuat(FRotator(0.0f, 0.0f, 0.0f));
+				break;
+			case LVT_OrthoNegativeYZ:
+				CachedViewState.Orientation = FQuat(FRotator(0.0f, 180.0f, 0.0f));
+				break;
+			default:
+				CachedViewState.Orientation = FQuat::Identity;
+			}
+		}
+
 		CachedViewState.bIsVR = false;
 	}
 
diff --git a/Engine/Source/Editor/UnrealEd/Private/Tools/UEdMode.cpp b/Engine/Source/Editor/UnrealEd/Private/Tools/UEdMode.cpp
--- a/Engine/Source/Editor/UnrealEd/Private/Tools/UEdMode.cpp
+++ b/Engine/Source/Editor/UnrealEd/Private/Tools/UEdMode.cpp
@@ -41,9 +41,42 @@ public:
 	void CacheCurrentViewState(FEditorViewportClient* ViewportClient)
 	{
 		FViewportCameraTransform ViewTransform = ViewportClient->GetViewTransform();
-		CachedViewState.Position = ViewTransform.GetLocation();
-		CachedViewState.Orientation = ViewTransform.GetRotation().Quaternion();
 		CachedViewState.bIsOrthographic = ViewportClient->IsOrtho();
+		CachedViewState.Position = ViewTransform.GetLocation();
+
+		// ViewTransform rotation is only initialized for perspective!
+		if (CachedViewState.bIsOrthographic == false)
+		{
+			CachedViewState.Orientation = ViewTransform.GetRotation().Quaternion();
+		}
+		else
+		{
+			// These rotations are based on hardcoded values in EditorViewportClient.cpp, see switches in FEditorViewportClient::CalcSceneView and FEditorViewportClient::Draw
+			switch (ViewportClient->ViewportType)
+			{
+			case LVT_OrthoXY:
+				CachedViewState.Orientation = FQuat(FRotator(-90.0f, -90.0f, 0.0f));
+				break;
+			case LVT_OrthoNegativeXY:
+				CachedViewState.Orientation = FQuat(FRotator(90.0f, 90.0f, 0.0f));
+				break;
+			case LVT_OrthoXZ:
+				CachedViewState.Orientation = FQuat(FRotator(0.0f, -90.0f, 0.0f));
+				break;
+			case LVT_OrthoNegativeXZ:
+				CachedViewState.Orientation = FQuat(FRotator(0.0f, 90.0f, 0.0f));
+				break;
+			case LVT_OrthoYZ:
+				CachedViewState.Orientation = FQuat(FRotator(0.0f, 0.0f, 0.0f));
+				break;
+			case LVT_OrthoNegativeYZ:
+				CachedViewState.Orientation = FQuat(FRotator(0.0f, 180.0f, 0.0f));
+				break;
+			default:
+				CachedViewState.Orientation = FQuat::Identity;
+			}
+		}
+
 		CachedViewState.bIsVR = false;
 	}
 
