ToolsFramework: add world-grid snap query that is usable by Tools
#rnx

diff --git a/Engine/Source/Editor/Experimental/EditorInteractiveToolsFramework/Private/EdModeInteractiveToolsContext.cpp b/Engine/Source/Editor/Experimental/EditorInteractiveToolsFramework/Private/EdModeInteractiveToolsContext.cpp
--- a/Engine/Source/Editor/Experimental/EditorInteractiveToolsFramework/Private/EdModeInteractiveToolsContext.cpp
+++ b/Engine/Source/Editor/Experimental/EditorInteractiveToolsFramework/Private/EdModeInteractiveToolsContext.cpp
@@ -27,6 +27,26 @@
 
 
 
+static float SnapToIncrement(float fValue, float fIncrement, float offset = 0)
+{
+	if (!FMath::IsFinite(fValue))
+	{
+		return 0;
+	}
+	fValue -= offset;
+	float sign = FMath::Sign(fValue);
+	fValue = FMath::Abs(fValue);
+	int nInc = (int)(fValue / fIncrement);
+	float fRem = (float)fmod(fValue, fIncrement);
+	if (fRem > fIncrement / 2)
+	{
+		++nInc;
+	}
+	return sign * (float)nInc * fIncrement + offset;
+}
+
+
+
 class FEdModeToolsContextQueriesImpl : public IToolsContextQueriesAPI
 {
 public:
@@ -115,6 +135,18 @@ public:
 
 		int FoundResultCount = 0;
 
+		if ((Request.TargetTypes & ESceneSnapQueryTargetType::Grid) != ESceneSnapQueryTargetType::None)
+		{
+			FSceneSnapQueryResult SnapResult;
+			SnapResult.TargetType = ESceneSnapQueryTargetType::Grid;
+			float SnapSize = GEditor->GetGridSize();
+			SnapResult.Position.X = SnapToIncrement(Request.Position.X, SnapSize);
+			SnapResult.Position.Y = SnapToIncrement(Request.Position.Y, SnapSize);
+			SnapResult.Position.Z = SnapToIncrement(Request.Position.Z, SnapSize);
+			Results.Add(SnapResult);
+			FoundResultCount++;
+		}
+
 		//
 		// Run a snap query by casting ray into the world.
 		// If a hit is found, we look up what triangle was hit, and then test its vertices and edges
diff --git a/Engine/Source/Runtime/Experimental/InteractiveToolsFramework/Public/ToolContextInterfaces.h b/Engine/Source/Runtime/Experimental/InteractiveToolsFramework/Public/ToolContextInterfaces.h
--- a/Engine/Source/Runtime/Experimental/InteractiveToolsFramework/Public/ToolContextInterfaces.h
+++ b/Engine/Source/Runtime/Experimental/InteractiveToolsFramework/Public/ToolContextInterfaces.h
@@ -81,8 +81,10 @@ enum class ESceneSnapQueryTargetType
 	MeshVertex = 1,
 	/** Consider any mesh edge */
 	MeshEdge = 2,
+	/** Grid Snapping */
+	Grid = 4,
 
-	All = MeshVertex | MeshEdge
+	All = MeshVertex | MeshEdge | Grid
 };
 ENUM_CLASS_FLAGS(ESceneSnapQueryTargetType);
 
@@ -92,9 +94,9 @@ ENUM_CLASS_FLAGS(ESceneSnapQueryTargetType);
 struct INTERACTIVETOOLSFRAMEWORK_API FSceneSnapQueryRequest
 {
 	/** What type of snap query geometry is this */
-	ESceneSnapQueryType RequestType;
+	ESceneSnapQueryType RequestType = ESceneSnapQueryType::Position;
 	/** What does caller want to try to snap to */
-	ESceneSnapQueryTargetType TargetTypes;
+	ESceneSnapQueryTargetType TargetTypes = ESceneSnapQueryTargetType::Grid;
 
 	/** Snap input position */
 	FVector Position;
