ToolsFramework: these changes (copied from Paint Mode) allow fly camera to capture mouse and prevent us from enabling tools during flycam movements
#rnx

diff --git a/Engine/Source/Editor/Experimental/EditorInteractiveToolsFramework/Private/EdModeInteractiveToolsContext.cpp b/Engine/Source/Editor/Experimental/EditorInteractiveToolsFramework/Private/EdModeInteractiveToolsContext.cpp
--- a/Engine/Source/Editor/Experimental/EditorInteractiveToolsFramework/Private/EdModeInteractiveToolsContext.cpp
+++ b/Engine/Source/Editor/Experimental/EditorInteractiveToolsFramework/Private/EdModeInteractiveToolsContext.cpp
@@ -651,6 +651,10 @@ bool UEdModeInteractiveToolsContext::InputKey(FEditorViewportClient* ViewportCli
 	{
 		return false;
 	}
+	if (ViewportClient->IsMovingCamera())
+	{
+		return false;
+	}
 
 	if (Event == IE_Pressed || Event == IE_Released)
 	{
@@ -800,7 +804,7 @@ bool UEdModeInteractiveToolsContext::MouseLeave(FEditorViewportClient* ViewportC
 
 bool UEdModeInteractiveToolsContext::StartTracking(FEditorViewportClient* InViewportClient, FViewport* InViewport)
 {
-	return false;
+	return true;
 }
 
 bool UEdModeInteractiveToolsContext::CapturedMouseMove(FEditorViewportClient* InViewportClient, FViewport* InViewport, int32 InMouseX, int32 InMouseY)
@@ -841,6 +845,9 @@ bool UEdModeInteractiveToolsContext::EndTracking(FEditorViewportClient* InViewpo
 	UE_LOG(LogTemp, Warning, TEXT("END TRACKING"));
 #endif
 
+	// unlock flight camera
+	InViewportClient->bLockFlightCamera = false;
+
 	return true;
 }
 
