ToolsFramework: workaround to ignore command hotkeys when in right-mouse fly+WASD mode (hence allowing those keys to be used as hotkeys when not in fly mode). Add an internal InputBehavior to UEdModeInteractiveToolsContext that captures right-mouse with highest priority. During that capture allow fly-mode to occur (ie don't set bLockFlightCamera=true) but set flag that we are in this mode. ModelingToolsEditorMode checks that flag before calling ProcessCommandBindings() on command list.

Note that this prevents right-mouse from being captured by any tool/gizmo (which seems like correct main-viewport behavior since fly mode needs to always work)

#rnx

diff --git a/Engine/Plugins/Experimental/ModelingToolsEditorMode/Source/ModelingToolsEditorMode/Private/ModelingToolsEditorMode.cpp b/Engine/Plugins/Experimental/ModelingToolsEditorMode/Source/ModelingToolsEditorMode/Private/ModelingToolsEditorMode.cpp
--- a/Engine/Plugins/Experimental/ModelingToolsEditorMode/Source/ModelingToolsEditorMode/Private/ModelingToolsEditorMode.cpp
+++ b/Engine/Plugins/Experimental/ModelingToolsEditorMode/Source/ModelingToolsEditorMode/Private/ModelingToolsEditorMode.cpp
@@ -154,9 +154,12 @@ bool FModelingToolsEditorMode::InputKey(FEditorViewportClient* ViewportClient, F
 	// try hotkeys
 	if (Event != IE_Released)
 	{
-		if (UICommandList->ProcessCommandBindings(Key, FSlateApplication::Get().GetModifierKeys(), false/*Event == IE_Repeat*/))
+		if (ToolsContext->ShouldIgnoreHotkeys() == false)		// allow the context to capture keyboard input if necessary
 		{
-			return true;
+			if (UICommandList->ProcessCommandBindings(Key, FSlateApplication::Get().GetModifierKeys(), false/*Event == IE_Repeat*/))
+			{
+				return true;
+			}
 		}
 	}
 
diff --git a/Engine/Source/Editor/Experimental/EditorInteractiveToolsFramework/Private/EdModeInteractiveToolsContext.cpp b/Engine/Source/Editor/Experimental/EditorInteractiveToolsFramework/Private/EdModeInteractiveToolsContext.cpp
--- a/Engine/Source/Editor/Experimental/EditorInteractiveToolsFramework/Private/EdModeInteractiveToolsContext.cpp
+++ b/Engine/Source/Editor/Experimental/EditorInteractiveToolsFramework/Private/EdModeInteractiveToolsContext.cpp
@@ -19,6 +19,7 @@
 #include "Tools/EditorToolAssetAPI.h"
 #include "Tools/EditorComponentSourceFactory.h"
 #include "InteractiveToolObjects.h"
+#include "BaseBehaviors/ClickDragBehavior.h"
 
 //#include "PhysicsEngine/BodySetup.h"
 //#include "Interfaces/Interface_CollisionDataProvider.h"
@@ -386,7 +387,6 @@ public:
 
 
 
-
 UEdModeInteractiveToolsContext::UEdModeInteractiveToolsContext()
 {
 	EditorMode = nullptr;
@@ -425,6 +425,20 @@ void UEdModeInteractiveToolsContext::Initialize(IToolsContextQueriesAPI* Queries
 	});
 
 
+	// If user right-press-drags, this enables "fly mode" in the main viewport, and in that mode the QEWASD keys should
+	// be used for flying control. However the EdMode InputKey/etc system doesn't enforce any of this, we can still also
+	// get that mouse input and hotkeys. So we register a dummy behavior that captures all right-mouse dragging, and 
+	// in that mode we set bInFlyMode=true, so that Modes based on this Context will know to skip hotkey processing
+	ULocalClickDragInputBehavior* RightMouseBehavior = NewObject<ULocalClickDragInputBehavior>(this);
+	RightMouseBehavior->CanBeginClickDragFunc = [](const FInputDeviceRay& PressPos) { return  FInputRayHit(0); };
+	RightMouseBehavior->OnClickPressFunc = [this](const FInputDeviceRay&) { bInFlyMode = true; };
+	RightMouseBehavior->OnClickReleaseFunc = [this](const FInputDeviceRay&) { bInFlyMode = false; };
+	RightMouseBehavior->OnTerminateFunc = [this]() { bInFlyMode = false; };
+	RightMouseBehavior->SetDefaultPriority(FInputCapturePriority(0));
+	RightMouseBehavior->SetUseRightMouseButton();
+	RightMouseBehavior->Initialize();
+	InputRouter->RegisterBehavior(RightMouseBehavior, this);
+
 	bInvalidationPending = false;
 }
 
@@ -700,7 +714,7 @@ bool UEdModeInteractiveToolsContext::InputKey(FEditorViewportClient* ViewportCli
 
 				InputRouter->PostInputEvent(InputState);
 
-				if (InputRouter->HasActiveMouseCapture())
+				if (InputRouter->HasActiveMouseCapture() && bInFlyMode == false)
 				{
 					// what is this about? MeshPaintMode has it...
 					ViewportClient->bLockFlightCamera = true;
diff --git a/Engine/Source/Editor/Experimental/EditorInteractiveToolsFramework/Public/EdModeInteractiveToolsContext.h b/Engine/Source/Editor/Experimental/EditorInteractiveToolsFramework/Public/EdModeInteractiveToolsContext.h
--- a/Engine/Source/Editor/Experimental/EditorInteractiveToolsFramework/Public/EdModeInteractiveToolsContext.h
+++ b/Engine/Source/Editor/Experimental/EditorInteractiveToolsFramework/Public/EdModeInteractiveToolsContext.h
@@ -72,6 +72,8 @@ public:
 	virtual void StartTool(const FString& ToolTypeIdentifier);
 	virtual void EndTool(EToolShutdownType ShutdownType);
 
+	virtual bool ShouldIgnoreHotkeys() const { return bInFlyMode; }
+
 protected:
 	// we hide these 
 	virtual void Initialize(IToolsContextQueriesAPI* QueriesAPI, IToolsContextTransactionsAPI* TransactionsAPI) override;
@@ -109,6 +111,8 @@ protected:
 	// Copy-pasted from other Editor code, seems kind of expensive?
 	static FRay GetRayFromMousePos(FEditorViewportClient* ViewportClient, FViewport* Viewport, int MouseX, int MouseY);
 
+	/** This will be set to true if user is in right-mouse "fly mode", which requires special handling to intercept hotkeys/etc */
+	bool bInFlyMode = false;
 
 	// editor UI state that we set before starting tool and when exiting tool
 	// Currently disabling anti-aliasing during active Tools because it causes PDI flickering
