ToolsFramework: use new EditorViewportClient ShowFlags Override functionality to Disable TAA, MotionBlur, and EyeAdaptation during modeling mode
#rnx

diff --git a/Engine/Source/Editor/Experimental/EditorInteractiveToolsFramework/Private/EdModeInteractiveToolsContext.cpp b/Engine/Source/Editor/Experimental/EditorInteractiveToolsFramework/Private/EdModeInteractiveToolsContext.cpp
--- a/Engine/Source/Editor/Experimental/EditorInteractiveToolsFramework/Private/EdModeInteractiveToolsContext.cpp
+++ b/Engine/Source/Editor/Experimental/EditorInteractiveToolsFramework/Private/EdModeInteractiveToolsContext.cpp
@@ -562,6 +562,12 @@ public:
 
 void UEdModeInteractiveToolsContext::Render(const FSceneView* View, FViewport* Viewport, FPrimitiveDrawInterface* PDI)
 {
+	// tools framework cannot use HitProxy so skip these calls
+	if (PDI->IsHitTesting())
+	{
+		return;
+	}
+
 	TempRenderContext RenderContext;
 	RenderContext.PDI = PDI;
 	ToolManager->Render(&RenderContext);
@@ -942,21 +948,26 @@ void UEdModeInteractiveToolsContext::DeactivateAllActiveTools()
 
 
 
-
 void UEdModeInteractiveToolsContext::SaveEditorStateAndSetForTool()
 {
 	check(bHaveSavedEditorState == false);
-	bSavedAntiAliasingState = GCurrentLevelEditingViewportClient->EngineShowFlags.AntiAliasing;
+	SavedViewportClient = GCurrentLevelEditingViewportClient;
 	bHaveSavedEditorState = true;
 
-	GCurrentLevelEditingViewportClient->EngineShowFlags.SetAntiAliasing(false);
+	GCurrentLevelEditingViewportClient->EnableOverrideEngineShowFlags([](FEngineShowFlags& Flags)
+	{
+		Flags.SetTemporalAA(false);
+		Flags.SetMotionBlur(false);
+		Flags.SetEyeAdaptation(false);
+	});
 }
 
+
 void UEdModeInteractiveToolsContext::RestoreEditorState()
 {
-	if (bHaveSavedEditorState && !IsEngineExitRequested())
+	if (bHaveSavedEditorState)
 	{
-		GCurrentLevelEditingViewportClient->EngineShowFlags.SetAntiAliasing(bSavedAntiAliasingState);
+		SavedViewportClient->DisableOverrideEngineShowFlags();
 		bHaveSavedEditorState = false;
 	}
 }
diff --git a/Engine/Source/Editor/Experimental/EditorInteractiveToolsFramework/Public/EdModeInteractiveToolsContext.h b/Engine/Source/Editor/Experimental/EditorInteractiveToolsFramework/Public/EdModeInteractiveToolsContext.h
--- a/Engine/Source/Editor/Experimental/EditorInteractiveToolsFramework/Public/EdModeInteractiveToolsContext.h
+++ b/Engine/Source/Editor/Experimental/EditorInteractiveToolsFramework/Public/EdModeInteractiveToolsContext.h
@@ -117,7 +117,7 @@ protected:
 	// editor UI state that we set before starting tool and when exiting tool
 	// Currently disabling anti-aliasing during active Tools because it causes PDI flickering
 	bool bHaveSavedEditorState = false;
-	bool bSavedAntiAliasingState = false;
+	FLevelEditorViewportClient* SavedViewportClient;
 	void SaveEditorStateAndSetForTool();
 	void RestoreEditorState();
 
