EdModeInteractiveToolsContext override of ViewportClient now iterates through existing viewports instead of using global GCurrentLevelEditingViewportClient, which is not safe to hold onto as it could be deleted at any time.
#rnx

diff --git a/Engine/Source/Editor/Experimental/EditorInteractiveToolsFramework/Private/EdModeInteractiveToolsContext.cpp b/Engine/Source/Editor/Experimental/EditorInteractiveToolsFramework/Private/EdModeInteractiveToolsContext.cpp
--- a/Engine/Source/Editor/Experimental/EditorInteractiveToolsFramework/Private/EdModeInteractiveToolsContext.cpp
+++ b/Engine/Source/Editor/Experimental/EditorInteractiveToolsFramework/Private/EdModeInteractiveToolsContext.cpp
@@ -6,7 +6,10 @@
 #include "EditorViewportClient.h"
 #include "EditorModeManager.h"
 #include "LevelEditor.h"
-#include "LevelEditorViewport.h"   // for GCurrentLevelEditingViewportClient
+#include "LevelEditorViewport.h"
+#include "IAssetViewport.h"
+#include "SLevelViewport.h"
+
 #include "Modules/ModuleManager.h"
 #include "ShowFlags.h"				// for EngineShowFlags
 #include "Engine/Selection.h"
@@ -424,6 +427,12 @@ void UEdModeInteractiveToolsContext::Initialize(IToolsContextQueriesAPI* Queries
 		RestoreEditorState();
 	});
 
+	// if viewport clients change we will discard our overrides as we aren't sure what happened
+	ViewportClientListChangedHandle = GEditor->OnViewportClientListChanged().AddLambda([this]()
+	{
+		RestoreEditorState();
+	});
+
 
 	// If user right-press-drags, this enables "fly mode" in the main viewport, and in that mode the QEWASD keys should
 	// be used for flying control. However the EdMode InputKey/etc system doesn't enforce any of this, we can still also
@@ -449,6 +458,7 @@ void UEdModeInteractiveToolsContext::Shutdown()
 	LevelEditor.OnMapChanged().Remove(WorldTearDownDelegateHandle);
 	FEditorDelegates::BeginPIE.Remove(BeginPIEDelegateHandle);
 	FEditorDelegates::PreSaveWorld.Remove(PreSaveWorldDelegateHandle);
+	GEditor->OnViewportClientListChanged().Remove(ViewportClientListChangedHandle);
 
 	// auto-accept any in-progress tools
 	DeactivateAllActiveTools();
@@ -962,15 +972,27 @@ void UEdModeInteractiveToolsContext::DeactivateAllActiveTools()
 void UEdModeInteractiveToolsContext::SaveEditorStateAndSetForTool()
 {
 	check(bHaveSavedEditorState == false);
-	SavedViewportClient = GCurrentLevelEditingViewportClient;
 	bHaveSavedEditorState = true;
 
-	GCurrentLevelEditingViewportClient->EnableOverrideEngineShowFlags([](FEngineShowFlags& Flags)
+	FLevelEditorModule& LevelEditorModule = FModuleManager::GetModuleChecked<FLevelEditorModule>("LevelEditor");
+	TSharedPtr<ILevelEditor> LevelEditor = LevelEditorModule.GetFirstLevelEditor();
+	if (LevelEditor.IsValid())
 	{
-		Flags.SetTemporalAA(false);
-		Flags.SetMotionBlur(false);
-		Flags.SetEyeAdaptation(false);
-	});
+		TArray<TSharedPtr<IAssetViewport>> Viewports = LevelEditor->GetViewports();
+		for (const TSharedPtr<IAssetViewport>& ViewportWindow : Viewports)
+		{
+			if (ViewportWindow.IsValid())
+			{
+				FEditorViewportClient& Viewport = ViewportWindow->GetAssetViewportClient();
+				Viewport.EnableOverrideEngineShowFlags([](FEngineShowFlags& Flags)
+				{
+					Flags.SetTemporalAA(false);
+					Flags.SetMotionBlur(false);
+					Flags.SetEyeAdaptation(false);
+				});
+			}
+		}
+	}
 }
 
 
@@ -978,7 +1000,23 @@ void UEdModeInteractiveToolsContext::RestoreEditorState()
 {
 	if (bHaveSavedEditorState)
 	{
-		SavedViewportClient->DisableOverrideEngineShowFlags();
 		bHaveSavedEditorState = false;
+
+		FLevelEditorModule& LevelEditorModule = FModuleManager::GetModuleChecked<FLevelEditorModule>("LevelEditor");
+		TSharedPtr<ILevelEditor> LevelEditor = LevelEditorModule.GetFirstLevelEditor();
+		if (LevelEditor.IsValid())
+		{
+			TArray<TSharedPtr<IAssetViewport>> Viewports = LevelEditor->GetViewports();
+			for (const TSharedPtr<IAssetViewport>& ViewportWindow : Viewports)
+			{
+				if (ViewportWindow.IsValid())
+				{
+					FEditorViewportClient& Viewport = ViewportWindow->GetAssetViewportClient();
+					Viewport.DisableOverrideEngineShowFlags();
+				}
+			}
+		}
+
 	}
 }
+
diff --git a/Engine/Source/Editor/Experimental/EditorInteractiveToolsFramework/Public/EdModeInteractiveToolsContext.h b/Engine/Source/Editor/Experimental/EditorInteractiveToolsFramework/Public/EdModeInteractiveToolsContext.h
--- a/Engine/Source/Editor/Experimental/EditorInteractiveToolsFramework/Public/EdModeInteractiveToolsContext.h
+++ b/Engine/Source/Editor/Experimental/EditorInteractiveToolsFramework/Public/EdModeInteractiveToolsContext.h
@@ -98,6 +98,8 @@ protected:
 	FDelegateHandle PreSaveWorldDelegateHandle;
 	// called when a map is changed
 	FDelegateHandle WorldTearDownDelegateHandle;
+	// called when viewport clients change
+	FDelegateHandle ViewportClientListChangedHandle;
 
 	// EdMode implementation of InteractiveToolFramework APIs - see ToolContextInterfaces.h
 	IToolsContextQueriesAPI* QueriesAPI;
@@ -120,7 +122,6 @@ protected:
 	// editor UI state that we set before starting tool and when exiting tool
 	// Currently disabling anti-aliasing during active Tools because it causes PDI flickering
 	bool bHaveSavedEditorState = false;
-	FLevelEditorViewportClient* SavedViewportClient;
 	void SaveEditorStateAndSetForTool();
 	void RestoreEditorState();
 
