Fix mouse capture issue in UEdModeInteractiveToolsContext, where user can get stuck in capture state (eg stuck in sculpting brush) by depressing alt key while doing action and then releasing left-mouse with alt-key still down, and then releasing alt. Also copied this fix to UEdMode where the logic had been duplicated.
#rnx

diff --git a/Engine/Source/Editor/Experimental/EditorInteractiveToolsFramework/Private/EdModeInteractiveToolsContext.cpp b/Engine/Source/Editor/Experimental/EditorInteractiveToolsFramework/Private/EdModeInteractiveToolsContext.cpp
--- a/Engine/Source/Editor/Experimental/EditorInteractiveToolsFramework/Private/EdModeInteractiveToolsContext.cpp
+++ b/Engine/Source/Editor/Experimental/EditorInteractiveToolsFramework/Private/EdModeInteractiveToolsContext.cpp
@@ -677,11 +677,6 @@ bool UEdModeInteractiveToolsContext::InputKey(FEditorViewportClient* ViewportCli
 		}
 	}
 
-	// if alt is down we do not process mouse event
-	if (ViewportClient->IsAltPressed())
-	{
-		return false;
-	}
 	if (ViewportClient->IsMovingCamera())
 	{
 		return false;
@@ -703,11 +698,11 @@ bool UEdModeInteractiveToolsContext::InputKey(FEditorViewportClient* ViewportCli
 
 			if (bIsLeftMouse || bIsMiddleMouse || bIsRightMouse)
 			{
-
-				// early-out here if we are going to do camera manipulation
-				if (ViewportClient->IsAltPressed())
+				// if alt is down and we are not capturing, somewhere higher in the ViewportClient/EdMode stack 
+				// is going to start doing alt+mouse camera manipulation. So we should ignore this mouse event.
+				if (ViewportClient->IsAltPressed() && InputRouter->HasActiveMouseCapture() == false)
 				{
-					return bHandled;
+					return false;
 				}
 
 				FInputDeviceState InputState = CurrentMouseState;
diff --git a/Engine/Source/Editor/UnrealEd/Private/Tools/UEdMode.cpp b/Engine/Source/Editor/UnrealEd/Private/Tools/UEdMode.cpp
--- a/Engine/Source/Editor/UnrealEd/Private/Tools/UEdMode.cpp
+++ b/Engine/Source/Editor/UnrealEd/Private/Tools/UEdMode.cpp
@@ -549,12 +549,6 @@ bool UEdMode::InputKey(FEditorViewportClient* ViewportClient, FViewport* Viewpor
 		}
 	}
 
-	// if alt is down we do not process mouse event
-	if (ViewportClient->IsAltPressed())
-	{
-		return false;
-	}
-
 	if (Event == IE_Pressed || Event == IE_Released)
 	{
 		if (Key.IsMouseButton())
@@ -565,11 +559,11 @@ bool UEdMode::InputKey(FEditorViewportClient* ViewportClient, FViewport* Viewpor
 
 			if (bIsLeftMouse || bIsMiddleMouse || bIsRightMouse)
 			{
-
-				// early-out here if we are going to do camera manipulation
-				if (ViewportClient->IsAltPressed())
+				// if alt is down and we are not capturing, somewhere higher in the ViewportClient/EdMode stack 
+				// is going to start doing alt+mouse camera manipulation. So we should ignore this mouse event.
+				if (ViewportClient->IsAltPressed() && ToolsContext->InputRouter->HasActiveMouseCapture() == false)
 				{
-					return bHandled;
+					return false;
 				}
 
 				FInputDeviceState InputState = CurrentMouseState;
