ModelingTools: Disable override of dynamic exposure because it then falls back to fixed exposure which may result in scene basically becoming all black
#rnx

diff --git a/Engine/Source/Editor/Experimental/EditorInteractiveToolsFramework/Private/EdModeInteractiveToolsContext.cpp b/Engine/Source/Editor/Experimental/EditorInteractiveToolsFramework/Private/EdModeInteractiveToolsContext.cpp
--- a/Engine/Source/Editor/Experimental/EditorInteractiveToolsFramework/Private/EdModeInteractiveToolsContext.cpp
+++ b/Engine/Source/Editor/Experimental/EditorInteractiveToolsFramework/Private/EdModeInteractiveToolsContext.cpp
@@ -983,7 +983,8 @@ void UEdModeInteractiveToolsContext::SaveEditorStateAndSetForTool()
 				{
 					Flags.SetTemporalAA(false);
 					Flags.SetMotionBlur(false);
-					Flags.SetEyeAdaptation(false);
+					// disable this as depending on fixed exposure settings the entire scene may turn black
+					//Flags.SetEyeAdaptation(false);
 				});
 			}
 		}
