ModelingTools: add some special-case code that allows alt+mouse navigation to work during draw-polygon tool
#rnx

diff --git a/Engine/Source/Editor/Experimental/EditorInteractiveToolsFramework/Private/EdModeInteractiveToolsContext.cpp b/Engine/Source/Editor/Experimental/EditorInteractiveToolsFramework/Private/EdModeInteractiveToolsContext.cpp
--- a/Engine/Source/Editor/Experimental/EditorInteractiveToolsFramework/Private/EdModeInteractiveToolsContext.cpp
+++ b/Engine/Source/Editor/Experimental/EditorInteractiveToolsFramework/Private/EdModeInteractiveToolsContext.cpp
@@ -704,6 +704,13 @@ bool UEdModeInteractiveToolsContext::InputKey(FEditorViewportClient* ViewportCli
 				{
 					return false;
 				}
+				// This is a special-case hack for UMultiClickSequenceInputBehavior, because it holds capture across multiple
+				// mouse clicks, which prevents alt+mouse navigation from working between clicks (very annoying in draw polygon).
+				// Remove this special-case once that tool is fixed to use CollectSurfacePathMechanic instead
+				if (Event == IE_Pressed && bIsLeftMouse && ViewportClient->IsAltPressed() && InputRouter->HasActiveMouseCapture())
+				{
+					return false;
+				}
 
 				FInputDeviceState InputState = CurrentMouseState;
 				InputState.InputDevice = EInputDevices::Mouse;
diff --git a/Engine/Source/Runtime/Experimental/InteractiveToolsFramework/Private/BaseBehaviors/MultiClickSequenceInputBehavior.cpp b/Engine/Source/Runtime/Experimental/InteractiveToolsFramework/Private/BaseBehaviors/MultiClickSequenceInputBehavior.cpp
--- a/Engine/Source/Runtime/Experimental/InteractiveToolsFramework/Private/BaseBehaviors/MultiClickSequenceInputBehavior.cpp
+++ b/Engine/Source/Runtime/Experimental/InteractiveToolsFramework/Private/BaseBehaviors/MultiClickSequenceInputBehavior.cpp
@@ -48,6 +48,13 @@ FInputCaptureUpdate UMultiClickSequenceInputBehavior::UpdateCapture(const FInput
 {
 	check(bInActiveSequence == true);   // should always be the case!
 
+	// This is a hack to avoid terminating multi-click sequences if the user does alt+mouse camera navigation.
+	// This entire class should be deprecated and removed, in which case this hack won't be relevant...
+	if (input.bAltKeyDown)
+	{
+		return FInputCaptureUpdate::Continue();
+	}
+
 	Modifiers.UpdateModifiers(input, Target);
 
 	// allow target to abort click sequence
