ModelingMode: add mechanism to schedule actions to process on next editor Tick(). This allows us to defer execution of complex operations like staticmesh rebuild to outside of slate button-click handlers, which (hopefully) will work around some mystery crashes that occur during very deep call stacks (eg when the staticmesh build launches a slowtask progress dialog, that then explicitly processes windows messages, which results in mouse events being processed, which kicks off a hitproxy rendering call, that in turn causes the RHI to try to rebuild rendering data, on the mesh that is being rebuilt - whew!)
#rnx

diff --git a/Engine/Source/Editor/Experimental/EditorInteractiveToolsFramework/Private/EdModeInteractiveToolsContext.cpp b/Engine/Source/Editor/Experimental/EditorInteractiveToolsFramework/Private/EdModeInteractiveToolsContext.cpp
--- a/Engine/Source/Editor/Experimental/EditorInteractiveToolsFramework/Private/EdModeInteractiveToolsContext.cpp
+++ b/Engine/Source/Editor/Experimental/EditorInteractiveToolsFramework/Private/EdModeInteractiveToolsContext.cpp
@@ -539,6 +539,16 @@ void UEdModeInteractiveToolsContext::PostInvalidation()
 
 void UEdModeInteractiveToolsContext::Tick(FEditorViewportClient* ViewportClient, float DeltaTime)
 {
+	// process any actions that were scheduled to execute on the next tick
+	if (NextTickExecuteActions.Num() > 0)
+	{
+		for (TUniqueFunction<void()>& Action : NextTickExecuteActions)
+		{
+			Action();
+		}
+		NextTickExecuteActions.Reset();
+	}
+
 	ToolManager->Tick(DeltaTime);
 	GizmoManager->Tick(DeltaTime);
 
@@ -944,15 +954,26 @@ bool UEdModeInteractiveToolsContext::CanCompleteActiveTool() const
 
 void UEdModeInteractiveToolsContext::StartTool(const FString& ToolTypeIdentifier)
 {
-	if (UInteractiveToolsContext::StartTool(EToolSide::Mouse, ToolTypeIdentifier))
+	FString LocalIdentifier(ToolTypeIdentifier);
+	ScheduleExecuteAction([this, LocalIdentifier]()
 	{
-		SaveEditorStateAndSetForTool();
-	}
+		if (UInteractiveToolsContext::StartTool(EToolSide::Mouse, LocalIdentifier))
+		{
+			SaveEditorStateAndSetForTool();
+		}
+	});
+
+	PostInvalidation();
 }
 
 void UEdModeInteractiveToolsContext::EndTool(EToolShutdownType ShutdownType)
 {
-	UInteractiveToolsContext::EndTool(EToolSide::Mouse, ShutdownType);
+	ScheduleExecuteAction([this, ShutdownType]()
+	{
+		UInteractiveToolsContext::EndTool(EToolSide::Mouse, ShutdownType);
+	});
+
+	PostInvalidation();
 }
 
 
@@ -1023,3 +1044,8 @@ void UEdModeInteractiveToolsContext::RestoreEditorState()
 	}
 }
 
+
+void UEdModeInteractiveToolsContext::ScheduleExecuteAction(TUniqueFunction<void()> Action)
+{
+	NextTickExecuteActions.Add(MoveTemp(Action));
+}
\ No newline at end of file
diff --git a/Engine/Source/Editor/Experimental/EditorInteractiveToolsFramework/Public/EdModeInteractiveToolsContext.h b/Engine/Source/Editor/Experimental/EditorInteractiveToolsFramework/Public/EdModeInteractiveToolsContext.h
--- a/Engine/Source/Editor/Experimental/EditorInteractiveToolsFramework/Public/EdModeInteractiveToolsContext.h
+++ b/Engine/Source/Editor/Experimental/EditorInteractiveToolsFramework/Public/EdModeInteractiveToolsContext.h
@@ -125,4 +125,9 @@ protected:
 	void SaveEditorStateAndSetForTool();
 	void RestoreEditorState();
 
+	// any actions in this array are executed at the start of the next ::Tick()
+	TArray<TUniqueFunction<void()>> NextTickExecuteActions;
+	// push the input action onto the NextTickExecuteActions list. Use this to defer execution to a "simpler" time, 
+	// eg so we are not kicking off a complicated static mesh rebuild process inside a slate button handler
+	void ScheduleExecuteAction(TUniqueFunction<void()> Action);
 };
