InteractiveToolsFramework: Make SceneView available from  IToolsContextRenderAPI

#rnx

diff --git a/Engine/Source/Editor/Experimental/EditorInteractiveToolsFramework/Private/EdModeInteractiveToolsContext.cpp b/Engine/Source/Editor/Experimental/EditorInteractiveToolsFramework/Private/EdModeInteractiveToolsContext.cpp
--- a/Engine/Source/Editor/Experimental/EditorInteractiveToolsFramework/Private/EdModeInteractiveToolsContext.cpp
+++ b/Engine/Source/Editor/Experimental/EditorInteractiveToolsFramework/Private/EdModeInteractiveToolsContext.cpp
@@ -626,12 +626,14 @@ class FEdModeTempRenderContext : public IToolsContextRenderAPI
 {
 public:
 	FPrimitiveDrawInterface* PDI;
+	const FSceneView* SceneView;
 	FViewCameraState ViewCameraState;
 
 	FEdModeTempRenderContext(const FSceneView* View, FViewport* Viewport, FEditorViewportClient* ViewportClient, FPrimitiveDrawInterface* DrawInterface)
 	{
 		PDI = DrawInterface;
-		CacheCurrentViewState(View, Viewport, ViewportClient);
+		SceneView = View;
+		CacheCurrentViewState(Viewport, ViewportClient);
 	}
 
 	virtual FPrimitiveDrawInterface* GetPrimitiveDrawInterface() override
@@ -639,12 +641,17 @@ public:
 		return PDI;
 	}
 
+	virtual const FSceneView* GetSceneView() override
+	{
+		return SceneView;
+	}
+
 	virtual FViewCameraState GetCameraState() override
 	{
 		return ViewCameraState;
 	}
 
-	void CacheCurrentViewState(const FSceneView* View, FViewport* Viewport, FEditorViewportClient* ViewportClient)
+	void CacheCurrentViewState(FViewport* Viewport, FEditorViewportClient* ViewportClient)
 	{
 		FViewportCameraTransform ViewTransform = ViewportClient->GetViewTransform();
 		ViewCameraState.bIsOrthographic = ViewportClient->IsOrtho();
diff --git a/Engine/Source/Runtime/Experimental/InteractiveToolsFramework/Public/ToolContextInterfaces.h b/Engine/Source/Runtime/Experimental/InteractiveToolsFramework/Public/ToolContextInterfaces.h
--- a/Engine/Source/Runtime/Experimental/InteractiveToolsFramework/Public/ToolContextInterfaces.h
+++ b/Engine/Source/Runtime/Experimental/InteractiveToolsFramework/Public/ToolContextInterfaces.h
@@ -14,6 +14,7 @@ class UActorComponent;
 class FToolCommandChange;
 class UPackage;
 class FPrimitiveDrawInterface;
+class FSceneView;
 class UInteractiveToolManager;
 class UInteractiveGizmoManager;
 struct FMeshDescription;
@@ -355,6 +356,9 @@ public:
 	/** @return Current PDI */
 	virtual FPrimitiveDrawInterface* GetPrimitiveDrawInterface() = 0;
 
+	/** @return Current FSceneView */
+	virtual const FSceneView* GetSceneView() = 0;
+
 	/** @return Current Camera State for this Render API */
 	virtual FViewCameraState GetCameraState() = 0;
 };
