Fixing a shutdown crash with the static global editor mode tools lifecycle when cooking.

#rnx

diff --git a/Engine/Source/Editor/Experimental/EditorInteractiveToolsFramework/Private/EdModeInteractiveToolsContext.cpp b/Engine/Source/Editor/Experimental/EditorInteractiveToolsFramework/Private/EdModeInteractiveToolsContext.cpp
--- a/Engine/Source/Editor/Experimental/EditorInteractiveToolsFramework/Private/EdModeInteractiveToolsContext.cpp
+++ b/Engine/Source/Editor/Experimental/EditorInteractiveToolsFramework/Private/EdModeInteractiveToolsContext.cpp
@@ -495,16 +495,18 @@ void UEdModeInteractiveToolsContext::Initialize(IToolsContextQueriesAPI* Queries
 
 void UEdModeInteractiveToolsContext::Shutdown()
 {
-	FLevelEditorModule& LevelEditor = FModuleManager::GetModuleChecked<FLevelEditorModule>("LevelEditor");
-	LevelEditor.OnMapChanged().Remove(WorldTearDownDelegateHandle);
-	FEditorDelegates::BeginPIE.Remove(BeginPIEDelegateHandle);
-	FEditorDelegates::PreSaveWorld.Remove(PreSaveWorldDelegateHandle);
-	GEditor->OnViewportClientListChanged().Remove(ViewportClientListChangedHandle);
+	if (FLevelEditorModule* LevelEditor = FModuleManager::GetModulePtr<FLevelEditorModule>("LevelEditor"))
+	{
+		LevelEditor->OnMapChanged().Remove(WorldTearDownDelegateHandle);
+		FEditorDelegates::BeginPIE.Remove(BeginPIEDelegateHandle);
+		FEditorDelegates::PreSaveWorld.Remove(PreSaveWorldDelegateHandle);
+		GEditor->OnViewportClientListChanged().Remove(ViewportClientListChangedHandle);
 
-	// auto-accept any in-progress tools
-	DeactivateAllActiveTools();
+		// auto-accept any in-progress tools
+		DeactivateAllActiveTools();
 
-	UInteractiveToolsContext::Shutdown();
+		UInteractiveToolsContext::Shutdown();
+	}
 }
 
 void UEdModeInteractiveToolsContext::InitializeContextFromEdMode(FEdMode* EditorModeIn, IToolsContextAssetAPI* UseAssetAPI)
diff --git a/Engine/Source/Editor/UnrealEd/Private/UnrealEdGlobals.cpp b/Engine/Source/Editor/UnrealEd/Private/UnrealEdGlobals.cpp
--- a/Engine/Source/Editor/UnrealEd/Private/UnrealEdGlobals.cpp
+++ b/Engine/Source/Editor/UnrealEd/Private/UnrealEdGlobals.cpp
@@ -44,10 +44,19 @@ DEFINE_LOG_CATEGORY_STATIC(LogUnrealEd, Log, All);
 /**
  * Provides access to the FEditorModeTools for the level editor
  */
+namespace Internal
+{
+	static TSharedPtr<FEditorModeTools> GetGlobalModeManager()
+	{
+		static TSharedPtr<FEditorModeTools> EditorModeToolsSingleton = MakeShared<FEditorModeTools>();
+		return EditorModeToolsSingleton;
+	}
+};
+
 FEditorModeTools& GLevelEditorModeTools()
 {
-	static TSharedRef<FEditorModeTools> EditorModeToolsSingleton = MakeShared<FEditorModeTools>();
-	return EditorModeToolsSingleton.Get();
+	check(Internal::GetGlobalModeManager().IsValid());
+	return *Internal::GetGlobalModeManager().Get();
 }
 
 FLevelEditorViewportClient* GCurrentLevelEditingViewportClient = NULL;
@@ -204,14 +213,14 @@ void EditorExit()
 {
 	TRACE_CPUPROFILER_EVENT_SCOPE(EditorExit);
 
-	GLevelEditorModeTools().RemoveAllDelegateHandlers();
-	GLevelEditorModeTools().SetDefaultMode(FBuiltinEditorModes::EM_Default);
-	GLevelEditorModeTools().DeactivateAllModes(); // this also activates the default mode
-
 	// Save out any config settings for the editor so they don't get lost
 	GEditor->SaveConfig();
 	GLevelEditorModeTools().SaveConfig();
 
+	// Shutdown the global static mode manager
+	GLevelEditorModeTools().SetDefaultMode(FBuiltinEditorModes::EM_Default);
+	Internal::GetGlobalModeManager().Reset();
+
 	// Clean up the actor folders singleton
 	FActorFolders::Cleanup();
 
