InteractiveToolsFramework: Add render query information to tell if a viewport is hovered and/or focused

diff --git a/Engine/Plugins/Experimental/MeshModelingToolset/Source/MeshModelingTools/Private/LatticeDeformerTool.cpp b/Engine/Plugins/Experimental/MeshModelingToolset/Source/MeshModelingTools/Private/LatticeDeformerTool.cpp
--- a/Engine/Plugins/Experimental/MeshModelingToolset/Source/MeshModelingTools/Private/LatticeDeformerTool.cpp
+++ b/Engine/Plugins/Experimental/MeshModelingToolset/Source/MeshModelingTools/Private/LatticeDeformerTool.cpp
@@ -68,7 +68,7 @@ FVector3i ULatticeDeformerTool::GetLatticeResolution() const
 	return FVector3i{ Settings->XAxisResolution, Settings->YAxisResolution, Settings->ZAxisResolution };
 }
 
-void ULatticeDeformerTool::DrawHUD(FCanvas* Canvas)
+void ULatticeDeformerTool::DrawHUD(FCanvas* Canvas, IToolsContextRenderAPI* RenderAPI)
 {
 	// Draw the drag rectangle if it's active
 	if (ControlPointsMechanic->bIsDragging)
diff --git a/Engine/Plugins/Experimental/MeshModelingToolset/Source/MeshModelingTools/Public/LatticeDeformerTool.h b/Engine/Plugins/Experimental/MeshModelingToolset/Source/MeshModelingTools/Public/LatticeDeformerTool.h
--- a/Engine/Plugins/Experimental/MeshModelingToolset/Source/MeshModelingTools/Public/LatticeDeformerTool.h
+++ b/Engine/Plugins/Experimental/MeshModelingToolset/Source/MeshModelingTools/Public/LatticeDeformerTool.h
@@ -91,7 +91,7 @@ class MESHMODELINGTOOLS_API ULatticeDeformerTool : public USingleSelectionTool
 
 public:
 
-	virtual void DrawHUD(FCanvas* Canvas) override;
+	virtual void DrawHUD(FCanvas* Canvas, IToolsContextRenderAPI* RenderAPI) override;
 
 	virtual void SetWorld(UWorld* World) { TargetWorld = World; }
 	virtual void SetAssetAPI(IToolsContextAssetAPI* NewAssetApi) { AssetAPI = NewAssetApi; }
diff --git a/Engine/Plugins/Experimental/ModelingToolsEditorMode/Source/ModelingToolsEditorMode/Private/ModelingToolsEditorMode.cpp b/Engine/Plugins/Experimental/ModelingToolsEditorMode/Source/ModelingToolsEditorMode/Private/ModelingToolsEditorMode.cpp
--- a/Engine/Plugins/Experimental/ModelingToolsEditorMode/Source/ModelingToolsEditorMode/Private/ModelingToolsEditorMode.cpp
+++ b/Engine/Plugins/Experimental/ModelingToolsEditorMode/Source/ModelingToolsEditorMode/Private/ModelingToolsEditorMode.cpp
@@ -205,7 +205,7 @@ void FModelingToolsEditorMode::DrawHUD(FEditorViewportClient* ViewportClient,FVi
 	FEdMode::DrawHUD(ViewportClient, Viewport, View, Canvas);
 	if (ToolsContext != nullptr)
 	{
-		ToolsContext->DrawHUD(Canvas);
+		ToolsContext->DrawHUD(ViewportClient, Viewport, View, Canvas);
 	}
 }
 
diff --git a/Engine/Source/Editor/Experimental/EditorInteractiveToolsFramework/Private/EdModeInteractiveToolsContext.cpp b/Engine/Source/Editor/Experimental/EditorInteractiveToolsFramework/Private/EdModeInteractiveToolsContext.cpp
--- a/Engine/Source/Editor/Experimental/EditorInteractiveToolsFramework/Private/EdModeInteractiveToolsContext.cpp
+++ b/Engine/Source/Editor/Experimental/EditorInteractiveToolsFramework/Private/EdModeInteractiveToolsContext.cpp
@@ -30,6 +30,7 @@
 #include "EdMode.h"
 
 #include "BaseGizmos/GizmoRenderingUtil.h"
+#include "UnrealClient.h"
 
 //#define ENABLE_DEBUG_PRINTING
 
@@ -635,11 +636,11 @@ public:
 	FPrimitiveDrawInterface* PDI;
 	const FSceneView* SceneView;
 	FViewCameraState ViewCameraState;
+	EViewInteractionState ViewInteractionState;
 
-	FEdModeTempRenderContext(const FSceneView* View, FViewport* Viewport, FEditorViewportClient* ViewportClient, FPrimitiveDrawInterface* DrawInterface)
+	FEdModeTempRenderContext(const FSceneView* View, FViewport* Viewport, FEditorViewportClient* ViewportClient, FPrimitiveDrawInterface* DrawInterface, EViewInteractionState ViewInteractionState)
+		:PDI(DrawInterface), SceneView(View), ViewInteractionState(ViewInteractionState)
 	{
-		PDI = DrawInterface;
-		SceneView = View;
 		CacheCurrentViewState(Viewport, ViewportClient);
 	}
 
@@ -658,6 +659,11 @@ public:
 		return ViewCameraState;
 	}
 
+	virtual EViewInteractionState GetViewInteractionState() override
+	{
+		return ViewInteractionState;
+	}
+
 	void CacheCurrentViewState(FViewport* Viewport, FEditorViewportClient* ViewportClient)
 	{
 		FViewportCameraTransform ViewTransform = ViewportClient->GetViewTransform();
@@ -739,14 +745,38 @@ void UEdModeInteractiveToolsContext::Render(const FSceneView* View, FViewport* V
 	}
 
 	// Render Tool and Gizmos
-	FEdModeTempRenderContext RenderContext(View, Viewport, ViewportClient, PDI);
+	const FEditorViewportClient* Focused = EditorModeManager->GetFocusedViewportClient();
+	const FEditorViewportClient* Hovered = EditorModeManager->GetHoveredViewportClient();
+	EViewInteractionState InteractionState = EViewInteractionState::None;
+	if (ViewportClient == Focused )
+	{
+		InteractionState |= EViewInteractionState::Focused;
+	}
+	if (ViewportClient == Hovered )
+	{
+		InteractionState |= EViewInteractionState::Hovered;
+	}
+	FEdModeTempRenderContext RenderContext(View, Viewport, ViewportClient, PDI, InteractionState);
 	ToolManager->Render(&RenderContext);
 	GizmoManager->Render(&RenderContext);
 }
 
-void UEdModeInteractiveToolsContext::DrawHUD(FCanvas* Canvas)
+void UEdModeInteractiveToolsContext::DrawHUD(FViewportClient* ViewportClient,FViewport* Viewport,const FSceneView* View, FCanvas* Canvas)
 {
-	ToolManager->DrawHUD(Canvas);
+	FEditorViewportClient* EditorViewportClient = static_cast<FEditorViewportClient*>(Viewport->GetClient());
+	const FViewportClient* Focused = EditorModeManager->GetFocusedViewportClient();
+	const FViewportClient* Hovered = EditorModeManager->GetHoveredViewportClient();
+	EViewInteractionState InteractionState = EViewInteractionState::None;
+	if (ViewportClient == Focused )
+	{
+		InteractionState |= EViewInteractionState::Focused;
+	}
+	if (ViewportClient == Hovered )
+	{
+		InteractionState |= EViewInteractionState::Hovered;
+	}
+	FEdModeTempRenderContext RenderContext(View, Viewport, EditorViewportClient, nullptr /*PDI*/, InteractionState);
+	ToolManager->DrawHUD(Canvas, &RenderContext);
 }
 
 
diff --git a/Engine/Source/Editor/Experimental/EditorInteractiveToolsFramework/Public/EdModeInteractiveToolsContext.h b/Engine/Source/Editor/Experimental/EditorInteractiveToolsFramework/Public/EdModeInteractiveToolsContext.h
--- a/Engine/Source/Editor/Experimental/EditorInteractiveToolsFramework/Public/EdModeInteractiveToolsContext.h
+++ b/Engine/Source/Editor/Experimental/EditorInteractiveToolsFramework/Public/EdModeInteractiveToolsContext.h
@@ -60,7 +60,7 @@ public:
 
 	virtual void Tick(FEditorViewportClient* ViewportClient, float DeltaTime);
 	virtual void Render(const FSceneView* View, FViewport* Viewport, FPrimitiveDrawInterface* PDI);
-	virtual void DrawHUD(FCanvas* Canvas);
+	virtual void DrawHUD(FViewportClient* ViewportClient,FViewport* Viewport,const FSceneView* View, FCanvas* Canvas);
 
 	virtual bool ProcessEditDelete();
 
diff --git a/Engine/Source/Editor/UnrealEd/Private/Tools/UEdMode.cpp b/Engine/Source/Editor/UnrealEd/Private/Tools/UEdMode.cpp
--- a/Engine/Source/Editor/UnrealEd/Private/Tools/UEdMode.cpp
+++ b/Engine/Source/Editor/UnrealEd/Private/Tools/UEdMode.cpp
@@ -299,7 +299,7 @@ void UEdMode::DrawHUD(FEditorViewportClient* ViewportClient, FViewport* Viewport
 	// give ToolsContext a chance to draw in screen space
 	if (ToolsContext != nullptr)
 	{
-		ToolsContext->DrawHUD(Canvas);
+		ToolsContext->DrawHUD(ViewportClient, Viewport, View, Canvas);
 	}
 
 	// Render the drag tool.
diff --git a/Engine/Source/Runtime/Experimental/InteractiveToolsFramework/Private/InteractiveTool.cpp b/Engine/Source/Runtime/Experimental/InteractiveToolsFramework/Private/InteractiveTool.cpp
--- a/Engine/Source/Runtime/Experimental/InteractiveToolsFramework/Private/InteractiveTool.cpp
+++ b/Engine/Source/Runtime/Experimental/InteractiveToolsFramework/Private/InteractiveTool.cpp
@@ -49,7 +49,7 @@ void UInteractiveTool::Render(IToolsContextRenderAPI* RenderAPI)
 {
 }
 
-void UInteractiveTool::DrawHUD(FCanvas* Canvas)
+void UInteractiveTool::DrawHUD(FCanvas* Canvas, IToolsContextRenderAPI* RenderAPI)
 {
 }
 
diff --git a/Engine/Source/Runtime/Experimental/InteractiveToolsFramework/Private/InteractiveToolManager.cpp b/Engine/Source/Runtime/Experimental/InteractiveToolsFramework/Private/InteractiveToolManager.cpp
--- a/Engine/Source/Runtime/Experimental/InteractiveToolsFramework/Private/InteractiveToolManager.cpp
+++ b/Engine/Source/Runtime/Experimental/InteractiveToolsFramework/Private/InteractiveToolManager.cpp
@@ -316,16 +316,16 @@ void UInteractiveToolManager::Render(IToolsContextRenderAPI* RenderAPI)
 	}
 }
 
-void UInteractiveToolManager::DrawHUD(FCanvas* Canvas)
+void UInteractiveToolManager::DrawHUD(FCanvas* Canvas, IToolsContextRenderAPI* RenderAPI)
 {
 	if (ActiveLeftTool != nullptr)
 	{
-		ActiveLeftTool->DrawHUD(Canvas);
+		ActiveLeftTool->DrawHUD(Canvas, RenderAPI);
 	}
 
 	if (ActiveRightTool != nullptr)
 	{
-		ActiveRightTool->DrawHUD(Canvas);
+		ActiveRightTool->DrawHUD(Canvas, RenderAPI);
 	}
 }
 
diff --git a/Engine/Source/Runtime/Experimental/InteractiveToolsFramework/Public/InteractiveTool.h b/Engine/Source/Runtime/Experimental/InteractiveToolsFramework/Public/InteractiveTool.h
--- a/Engine/Source/Runtime/Experimental/InteractiveToolsFramework/Public/InteractiveTool.h
+++ b/Engine/Source/Runtime/Experimental/InteractiveToolsFramework/Public/InteractiveTool.h
@@ -267,7 +267,12 @@ public:
 	 */
 	virtual void Render(IToolsContextRenderAPI* RenderAPI);
 
-	virtual void DrawHUD( FCanvas* Canvas );
+	/**
+	 * Allow the Tool to do any custom screen space drawing
+	 * @param Canvas the FCanvas to use to do the drawing
+	 * @param RenderAPI Abstraction that provides access to Rendering in the current ToolsContext
+	 */
+	virtual void DrawHUD( FCanvas* Canvas, IToolsContextRenderAPI* RenderAPI );
 
 	/**
 	 * Non overrideable func which does processing and calls the tool's OnTick
diff --git a/Engine/Source/Runtime/Experimental/InteractiveToolsFramework/Public/InteractiveToolManager.h b/Engine/Source/Runtime/Experimental/InteractiveToolsFramework/Public/InteractiveToolManager.h
--- a/Engine/Source/Runtime/Experimental/InteractiveToolsFramework/Public/InteractiveToolManager.h
+++ b/Engine/Source/Runtime/Experimental/InteractiveToolsFramework/Public/InteractiveToolManager.h
@@ -231,7 +231,7 @@ public:
 	virtual void Render(IToolsContextRenderAPI* RenderAPI);
 
 	/** Let active Tools do their screen space drawing. Called by UInteractiveToolsContext. */
-	virtual void DrawHUD(FCanvas* Canvas);
+	virtual void DrawHUD(FCanvas* Canvas, IToolsContextRenderAPI* RenderAPI);
 
 
 	//
diff --git a/Engine/Source/Runtime/Experimental/InteractiveToolsFramework/Public/ToolContextInterfaces.h b/Engine/Source/Runtime/Experimental/InteractiveToolsFramework/Public/ToolContextInterfaces.h
--- a/Engine/Source/Runtime/Experimental/InteractiveToolsFramework/Public/ToolContextInterfaces.h
+++ b/Engine/Source/Runtime/Experimental/InteractiveToolsFramework/Public/ToolContextInterfaces.h
@@ -5,6 +5,7 @@
 #include "CoreMinimal.h"
 #include "UObject/Interface.h"
 #include "ComponentSourceInterfaces.h"
+#include "UObject/ObjectMacros.h"
 
 
 // predeclarations so we don't have to include these in all tools
@@ -342,6 +343,14 @@ public:
 
 };
 
+UENUM()
+enum class EViewInteractionState {
+	None = 0,
+	Hovered = 1,
+	Focused = 2
+};
+ENUM_CLASS_FLAGS(EViewInteractionState);
+
 /**
  * Users of the Tools Framework need to implement IToolsContextRenderAPI to allow
  * Tools, Indicators, and Gizmos to make low-level rendering calls for things like line drawing.
@@ -361,6 +370,9 @@ public:
 
 	/** @return Current Camera State for this Render API */
 	virtual FViewCameraState GetCameraState() = 0;
+
+	/** @return Current interaction state of the view to render */
+	virtual EViewInteractionState GetViewInteractionState() = 0;
 };
 
 
