Editor: EndTracking returns true only if there's an active tool (ie. paired with StartTracking)

diff --git a/Engine/Source/Editor/Experimental/EditorInteractiveToolsFramework/Private/EdModeInteractiveToolsContext.cpp b/Engine/Source/Editor/Experimental/EditorInteractiveToolsFramework/Private/EdModeInteractiveToolsContext.cpp
--- a/Engine/Source/Editor/Experimental/EditorInteractiveToolsFramework/Private/EdModeInteractiveToolsContext.cpp
+++ b/Engine/Source/Editor/Experimental/EditorInteractiveToolsFramework/Private/EdModeInteractiveToolsContext.cpp
@@ -1150,14 +1150,18 @@ bool UEdModeInteractiveToolsContext::CapturedMouseMove(FEditorViewportClient* In
 
 bool UEdModeInteractiveToolsContext::EndTracking(FEditorViewportClient* InViewportClient, FViewport* InViewport)
 {
+	// unlock flight camera
+	InViewportClient->bLockFlightCamera = false;
+
+	if (ToolManager->HasActiveTool(EToolSide::Mouse))
+	{
 #ifdef ENABLE_DEBUG_PRINTING
 	UE_LOG(LogTemp, Warning, TEXT("END TRACKING"));
 #endif
+		return true;
+	}
 
-	// unlock flight camera
-	InViewportClient->bLockFlightCamera = false;
-
-	return true;
+	return false;
 }
 
 FRay UEdModeInteractiveToolsContext::GetRayFromMousePos(FEditorViewportClient* ViewportClient, FViewport* Viewport, int MouseX, int MouseY)
