Remove #pragma once from cpps

diff --git a/Engine/Plugins/Animation/IKRig/Source/IKRigEditor/Private/IKRigEditor.cpp b/Engine/Plugins/Animation/IKRig/Source/IKRigEditor/Private/IKRigEditor.cpp
--- a/Engine/Plugins/Animation/IKRig/Source/IKRigEditor/Private/IKRigEditor.cpp
+++ b/Engine/Plugins/Animation/IKRig/Source/IKRigEditor/Private/IKRigEditor.cpp
@@ -1,7 +1,5 @@
 // Copyright Epic Games, Inc. All Rights Reserved.
 
-#pragma once
-
 #include "IKRigEditor.h"
 #include "Features/IModularFeatures.h"
 #include "Engine/Engine.h"
diff --git a/Engine/Source/Editor/Experimental/EditorInteractiveToolsFramework/Private/ToolTargets/StaticMeshComponentToolTarget.cpp b/Engine/Source/Editor/Experimental/EditorInteractiveToolsFramework/Private/ToolTargets/StaticMeshComponentToolTarget.cpp
--- a/Engine/Source/Editor/Experimental/EditorInteractiveToolsFramework/Private/ToolTargets/StaticMeshComponentToolTarget.cpp
+++ b/Engine/Source/Editor/Experimental/EditorInteractiveToolsFramework/Private/ToolTargets/StaticMeshComponentToolTarget.cpp
@@ -1,7 +1,5 @@
 // Copyright Epic Games, Inc. All Rights Reserved.
 
-#pragma once
-
 #include "ToolTargets/StaticMeshComponentToolTarget.h"
 
 #include "ComponentReregisterContext.h"
