ModelingMode: prevent mesh editing tools from modifying assets in /Engine/, and display in-viewport error message if the user tries to do this
#rnx

diff --git a/Engine/Source/Editor/Experimental/EditorInteractiveToolsFramework/Private/ToolTargets/StaticMeshComponentToolTarget.cpp b/Engine/Source/Editor/Experimental/EditorInteractiveToolsFramework/Private/ToolTargets/StaticMeshComponentToolTarget.cpp
--- a/Engine/Source/Editor/Experimental/EditorInteractiveToolsFramework/Private/ToolTargets/StaticMeshComponentToolTarget.cpp
+++ b/Engine/Source/Editor/Experimental/EditorInteractiveToolsFramework/Private/ToolTargets/StaticMeshComponentToolTarget.cpp
@@ -18,6 +18,12 @@ void UStaticMeshComponentToolTarget::CommitMaterialSetUpdate(const FComponentMat
 
 	UStaticMesh* StaticMesh = Cast<UStaticMeshComponent>(Component)->GetStaticMesh();
 
+	if (StaticMesh->GetPathName().StartsWith(TEXT("/Engine/")))
+	{
+		UE_LOG(LogTemp, Warning, TEXT("CANNOT MODIFY BUILT-IN ENGINE ASSET %s"), *StaticMesh->GetPathName());
+		return;
+	}
+
 	// flush any pending rendering commands, which might touch this component while we are rebuilding it's mesh
 	FlushRenderingCommands();
 
@@ -61,6 +67,17 @@ void UStaticMeshComponentToolTarget::CommitMeshDescription(const FCommitter& Com
 
 	UStaticMesh* StaticMesh = Cast<UStaticMeshComponent>(Component)->GetStaticMesh();
 
+	if (StaticMesh->GetPathName().StartsWith(TEXT("/Engine/")))
+	{
+		FString DebugMessage = FString::Printf(TEXT("CANNOT MODIFY BUILT-IN ENGINE ASSET %s"), *StaticMesh->GetPathName());
+		if (GAreScreenMessagesEnabled)
+		{
+			GEngine->AddOnScreenDebugMessage(INDEX_NONE, 10.0f, FColor::Red, DebugMessage);
+		}
+		UE_LOG(LogTemp, Warning, TEXT("%s"), *DebugMessage);
+		return;
+	}
+
 	// flush any pending rendering commands, which might touch this component while we are rebuilding it's mesh
 	FlushRenderingCommands();
 
diff --git a/Engine/Source/Editor/UnrealEd/Private/Tools/EditorComponentSourceFactory.cpp b/Engine/Source/Editor/UnrealEd/Private/Tools/EditorComponentSourceFactory.cpp
--- a/Engine/Source/Editor/UnrealEd/Private/Tools/EditorComponentSourceFactory.cpp
+++ b/Engine/Source/Editor/UnrealEd/Private/Tools/EditorComponentSourceFactory.cpp
@@ -70,6 +70,12 @@ void FStaticMeshComponentTarget::CommitMaterialSetUpdate(const FComponentMateria
 	
 	UStaticMesh* StaticMesh = Cast<UStaticMeshComponent>(Component)->GetStaticMesh();
 
+	if (StaticMesh->GetPathName().StartsWith(TEXT("/Engine/")))
+	{
+		UE_LOG(LogTemp, Warning, TEXT("CANNOT MODIFY BUILT-IN ENGINE ASSET %s"), *StaticMesh->GetPathName());
+		return;
+	}
+
 	// flush any pending rendering commands, which might touch this component while we are rebuilding it's mesh
 	FlushRenderingCommands();
 
@@ -122,6 +128,17 @@ void FStaticMeshComponentTarget::CommitMesh( const FCommitter& Committer )
 	//check(bSaved);
 	UStaticMesh* StaticMesh = Cast<UStaticMeshComponent>(Component)->GetStaticMesh();
 
+	if (StaticMesh->GetPathName().StartsWith(TEXT("/Engine/")))
+	{
+		FString DebugMessage = FString::Printf(TEXT("CANNOT MODIFY BUILT-IN ENGINE ASSET %s"), *StaticMesh->GetPathName());
+		if (GAreScreenMessagesEnabled)
+		{
+			GEngine->AddOnScreenDebugMessage(INDEX_NONE, 10.0f, FColor::Red, DebugMessage);
+		}
+		UE_LOG(LogTemp, Warning, TEXT("%s"), *DebugMessage);
+		return;
+	}
+
 	// flush any pending rendering commands, which might touch this component while we are rebuilding it's mesh
 	FlushRenderingCommands();
 
