MeshModelingTools: Split the material-related functions into a separate interface in the tool target system, made them match the existing PrimitiveComponentTarget signatures, made CommitMaterialSet support applying to the component rather than the asset.

#rnx

diff --git a/Engine/Plugins/Experimental/MeshModelingToolset/Source/MeshModelingToolsEditorOnly/Private/ConversionUtils/VolumeMeshDescriptionToolTarget.cpp b/Engine/Plugins/Experimental/MeshModelingToolset/Source/MeshModelingToolsEditorOnly/Private/ConversionUtils/VolumeMeshDescriptionToolTarget.cpp
--- a/Engine/Plugins/Experimental/MeshModelingToolset/Source/MeshModelingToolsEditorOnly/Private/ConversionUtils/VolumeMeshDescriptionToolTarget.cpp
+++ b/Engine/Plugins/Experimental/MeshModelingToolset/Source/MeshModelingToolsEditorOnly/Private/ConversionUtils/VolumeMeshDescriptionToolTarget.cpp
@@ -25,7 +25,22 @@ UVolumeMeshDescriptionToolTarget::UVolumeMeshDescriptionToolTarget()
 	VolumeToMeshOptions.bOptimizeMesh = true;
 }
 
-void UVolumeMeshDescriptionToolTarget::GetMaterialSet(FComponentMaterialSet& MaterialSetOut) const
+int32 UVolumeMeshDescriptionToolTarget::GetNumMaterials() const
+{
+	return IsValid() ? 1 : 0;
+}
+
+UMaterialInterface* UVolumeMeshDescriptionToolTarget::GetMaterial(int32 MaterialIndex) const
+{
+	if (!IsValid() || MaterialIndex > 0)
+	{
+		return nullptr;
+	}
+
+	return ToolSetupUtil::GetDefaultEditVolumeMaterial();
+}
+
+void UVolumeMeshDescriptionToolTarget::GetMaterialSet(FComponentMaterialSet& MaterialSetOut, bool bPreferAssetMaterials) const
 {
 	MaterialSetOut.Materials.Reset();
 	if (IsValid() == false)
@@ -40,12 +55,6 @@ void UVolumeMeshDescriptionToolTarget::GetMaterialSet(FComponentMaterialSet& Mat
 	}
 }
 
-void UVolumeMeshDescriptionToolTarget::CommitMaterialSetUpdate(const FComponentMaterialSet& MaterialSet)
-{
-	check(IsValid());
-	// Do nothing
-}
-
 FMeshDescription* UVolumeMeshDescriptionToolTarget::GetMeshDescription()
 {
 	FTransform transform = GetWorldTransform();
diff --git a/Engine/Plugins/Experimental/MeshModelingToolset/Source/MeshModelingToolsEditorOnly/Private/MeshToVolumeTool.cpp b/Engine/Plugins/Experimental/MeshModelingToolset/Source/MeshModelingToolsEditorOnly/Private/MeshToVolumeTool.cpp
--- a/Engine/Plugins/Experimental/MeshModelingToolset/Source/MeshModelingToolsEditorOnly/Private/MeshToVolumeTool.cpp
+++ b/Engine/Plugins/Experimental/MeshModelingToolset/Source/MeshModelingToolsEditorOnly/Private/MeshToVolumeTool.cpp
@@ -12,6 +12,7 @@
 #include "ToolSetupUtil.h"
 #include "Selections/MeshConnectedComponents.h"
 #include "Selection/ToolSelectionUtil.h"
+#include "TargetInterfaces/MaterialProvider.h"
 #include "TargetInterfaces/MeshDescriptionProvider.h"
 #include "TargetInterfaces/PrimitiveComponentBackedTarget.h"
 #include "ToolTargetManager.h"
@@ -33,8 +34,9 @@ const FToolTargetTypeRequirements& UMeshToVolumeToolBuilder::GetTargetRequiremen
 {
 	static FToolTargetTypeRequirements TypeRequirements(nullptr, TArray<const UClass*>({
 		UMeshDescriptionProvider::StaticClass(),
-		UPrimitiveComponentBackedTarget::StaticClass() })
-		);
+		UPrimitiveComponentBackedTarget::StaticClass(),
+		UMaterialProvider::StaticClass()
+		}));
 	return TypeRequirements;
 }
 
@@ -78,7 +80,8 @@ void UMeshToVolumeTool::Setup()
 	PreviewMesh->SetTransform(TargetComponent->GetWorldTransform());
 
 	FComponentMaterialSet MaterialSet;
-	TargetComponent->GetMaterialSet(MaterialSet);
+	check(Cast<IMaterialProvider>(Target));
+	Cast<IMaterialProvider>(Target)->GetMaterialSet(MaterialSet, false);
 	PreviewMesh->SetMaterials(MaterialSet.Materials);
 
 	PreviewMesh->SetTangentsMode(EDynamicMeshTangentCalcType::ExternallyCalculated);
diff --git a/Engine/Plugins/Experimental/MeshModelingToolset/Source/MeshModelingToolsEditorOnly/Public/ConversionUtils/VolumeMeshDescriptionToolTarget.h b/Engine/Plugins/Experimental/MeshModelingToolset/Source/MeshModelingToolsEditorOnly/Public/ConversionUtils/VolumeMeshDescriptionToolTarget.h
--- a/Engine/Plugins/Experimental/MeshModelingToolset/Source/MeshModelingToolsEditorOnly/Public/ConversionUtils/VolumeMeshDescriptionToolTarget.h
+++ b/Engine/Plugins/Experimental/MeshModelingToolset/Source/MeshModelingToolsEditorOnly/Public/ConversionUtils/VolumeMeshDescriptionToolTarget.h
@@ -5,6 +5,7 @@
 #include "CoreMinimal.h"
 
 #include "ConversionUtils/VolumeToDynamicMesh.h"
+#include "TargetInterfaces/MaterialProvider.h"
 #include "TargetInterfaces/MeshDescriptionCommitter.h"
 #include "TargetInterfaces/MeshDescriptionProvider.h"
 #include "ToolTargets/PrimitiveComponentToolTarget.h"
@@ -16,7 +17,7 @@
  */
 UCLASS(Transient)
 class MESHMODELINGTOOLSEDITORONLY_API UVolumeMeshDescriptionToolTarget : public UPrimitiveComponentToolTarget,
-	public IMeshDescriptionCommitter, public IMeshDescriptionProvider
+	public IMeshDescriptionCommitter, public IMeshDescriptionProvider, public IMaterialProvider
 {
 	GENERATED_BODY()
 
@@ -31,9 +32,15 @@ public:
 	// IMeshDescritpionCommitter implementation
 	virtual void CommitMeshDescription(const FCommitter& Committer) override;
 
-	// IPrimitiveComponentBackedTarget implementation
-	virtual void GetMaterialSet(FComponentMaterialSet& MaterialSetOut) const override;
-	virtual void CommitMaterialSetUpdate(const FComponentMaterialSet& MaterialSet) override;
+	// IMaterialProvider implementation
+	virtual int32 GetNumMaterials() const override;
+	virtual UMaterialInterface* GetMaterial(int32 MaterialIndex) const override;
+	// Ignores bPreferAssetMaterials
+	virtual void GetMaterialSet(FComponentMaterialSet& MaterialSetOut, bool bPreferAssetMaterials) const override; 
+
+	// Doesn't actually do anything for a volume
+	virtual bool CommitMaterialSetUpdate(const FComponentMaterialSet& MaterialSet, bool bApplyToAsset) override { return false; }
+
 	// Rest provided by parent class
 
 protected:
diff --git a/Engine/Plugins/Experimental/MeshModelingToolset/Source/ModelingComponents/Private/PreviewMesh.cpp b/Engine/Plugins/Experimental/MeshModelingToolset/Source/ModelingComponents/Private/PreviewMesh.cpp
--- a/Engine/Plugins/Experimental/MeshModelingToolset/Source/ModelingComponents/Private/PreviewMesh.cpp
+++ b/Engine/Plugins/Experimental/MeshModelingToolset/Source/ModelingComponents/Private/PreviewMesh.cpp
@@ -4,7 +4,7 @@
 
 #include "Containers/StaticArray.h"
 #include "Engine/Classes/Materials/Material.h"
-#include "TargetInterfaces/PrimitiveComponentBackedTarget.h"
+#include "TargetInterfaces/MaterialProvider.h" //FComponentMaterialSet
 
 
 APreviewMeshActor::APreviewMeshActor()
diff --git a/Engine/Source/Editor/Experimental/EditorInteractiveToolsFramework/Private/ToolTargets/StaticMeshComponentToolTarget.cpp b/Engine/Source/Editor/Experimental/EditorInteractiveToolsFramework/Private/ToolTargets/StaticMeshComponentToolTarget.cpp
--- a/Engine/Source/Editor/Experimental/EditorInteractiveToolsFramework/Private/ToolTargets/StaticMeshComponentToolTarget.cpp
+++ b/Engine/Source/Editor/Experimental/EditorInteractiveToolsFramework/Private/ToolTargets/StaticMeshComponentToolTarget.cpp
@@ -12,45 +12,102 @@ namespace UStaticMeshComponentToolTargetLocals
 	int32 LODIndex = 0;
 }
 
-void UStaticMeshComponentToolTarget::CommitMaterialSetUpdate(const FComponentMaterialSet& MaterialSet)
+int32 UStaticMeshComponentToolTarget::GetNumMaterials() const
 {
-	check(IsValid());
+	return ensure(IsValid()) ? Component->GetNumMaterials() : 0;
+}
 
-	UStaticMesh* StaticMesh = Cast<UStaticMeshComponent>(Component)->GetStaticMesh();
+UMaterialInterface* UStaticMeshComponentToolTarget::GetMaterial(int32 MaterialIndex) const
+{
+	return ensure(IsValid()) ? Component->GetMaterial(MaterialIndex) : nullptr;
+}
 
-	if (StaticMesh->GetPathName().StartsWith(TEXT("/Engine/")))
+void UStaticMeshComponentToolTarget::GetMaterialSet(FComponentMaterialSet& MaterialSetOut, bool bPreferAssetMaterials) const
+{
+	if (!ensure(IsValid())) return;
+
+	if (bPreferAssetMaterials)
 	{
-		UE_LOG(LogTemp, Warning, TEXT("CANNOT MODIFY BUILT-IN ENGINE ASSET %s"), *StaticMesh->GetPathName());
-		return;
+		UStaticMesh* StaticMesh = Cast<UStaticMeshComponent>(Component)->GetStaticMesh();
+		int32 NumMaterials = Component->GetNumMaterials();
+		MaterialSetOut.Materials.SetNum(NumMaterials);
+		for (int32 k = 0; k < NumMaterials; ++k)
+		{
+			MaterialSetOut.Materials[k] = StaticMesh->GetMaterial(k);
+		}
 	}
+	else
+	{
+		int32 NumMaterials = Component->GetNumMaterials();
+		MaterialSetOut.Materials.SetNum(NumMaterials);
+		for (int32 k = 0; k < NumMaterials; ++k)
+		{
+			MaterialSetOut.Materials[k] = Component->GetMaterial(k);
+		}
+	}
+}
 
-	// flush any pending rendering commands, which might touch this component while we are rebuilding it's mesh
-	FlushRenderingCommands();
+bool UStaticMeshComponentToolTarget::CommitMaterialSetUpdate(const FComponentMaterialSet& MaterialSet, bool bApplyToAsset)
+{
+	if (!ensure(IsValid())) return false;
 
-	// unregister the component while we update it's static mesh
-	TUniquePtr<FComponentReregisterContext> ComponentReregisterContext = MakeUnique<FComponentReregisterContext>(Component);
+	if (bApplyToAsset)
+	{
+		UStaticMesh* StaticMesh = Cast<UStaticMeshComponent>(Component)->GetStaticMesh();
 
-	// make sure transactional flag is on
-	StaticMesh->SetFlags(RF_Transactional);
+		if (StaticMesh->GetPathName().StartsWith(TEXT("/Engine/")))
+		{
+			UE_LOG(LogTemp, Warning, TEXT("CANNOT MODIFY BUILT-IN ENGINE ASSET %s"), *StaticMesh->GetPathName());
+			return false;
+		}
 
-	bool bSavedToTransactionBuffer = StaticMesh->Modify();
-	check(bSavedToTransactionBuffer);
+		// flush any pending rendering commands, which might touch this component while we are rebuilding its mesh
+		FlushRenderingCommands();
 
-	int NewNumMaterials = MaterialSet.Materials.Num();
-	if (NewNumMaterials != StaticMesh->GetStaticMaterials().Num())
-	{
-		StaticMesh->GetStaticMaterials().SetNum(NewNumMaterials);
+		// unregister the component while we update it's static mesh
+		TUniquePtr<FComponentReregisterContext> ComponentReregisterContext = MakeUnique<FComponentReregisterContext>(Component);
+
+		// make sure transactional flag is on
+		StaticMesh->SetFlags(RF_Transactional);
+
+		StaticMesh->Modify();
+
+		int NewNumMaterials = MaterialSet.Materials.Num();
+		if (NewNumMaterials != StaticMesh->GetStaticMaterials().Num())
+		{
+			StaticMesh->GetStaticMaterials().SetNum(NewNumMaterials);
+		}
+		for (int k = 0; k < NewNumMaterials; ++k)
+		{
+			if (StaticMesh->GetMaterial(k) != MaterialSet.Materials[k])
+			{
+				StaticMesh->SetMaterial(k, MaterialSet.Materials[k]);
+			}
+		}
+
+		StaticMesh->PostEditChange();
 	}
-	for (int k = 0; k < NewNumMaterials; ++k)
+	else
 	{
-		if (StaticMesh->GetMaterial(k) != MaterialSet.Materials[k])
+		int32 NumMaterialsNeeded = Component->GetNumMaterials();
+		int32 NumMaterialsGiven = MaterialSet.Materials.Num();
+
+		// We wrote the below code to support a mismatch in the number of materials.
+		// However, it is not yet clear whether this might be desirable, and we don't
+		// want to inadvertantly hide bugs in the meantime. So, we keep this ensure here
+		// for now, and we can remove it if we decide that we want the ability.
+		ensure(NumMaterialsNeeded == NumMaterialsGiven);
+
+		check(NumMaterialsGiven > 0);
+
+		for (int32 i = 0; i < NumMaterialsNeeded; ++i)
 		{
-			StaticMesh->SetMaterial(k, MaterialSet.Materials[k]);
+			int32 MaterialToUseIndex = FMath::Min(i, NumMaterialsGiven - 1);
+			Component->SetMaterial(i, MaterialSet.Materials[MaterialToUseIndex]);
 		}
 	}
 
-	// right?
-	StaticMesh->PostEditChange();
+	return true;
 }
 
 FMeshDescription* UStaticMeshComponentToolTarget::GetMeshDescription()
@@ -63,7 +120,7 @@ FMeshDescription* UStaticMeshComponentToolTarget::GetMeshDescription()
 void UStaticMeshComponentToolTarget::CommitMeshDescription(const FCommitter& Committer)
 {
 	using namespace UStaticMeshComponentToolTargetLocals;
-	check(IsValid());
+	if (ensure(IsValid()) == false) return;
 
 	UStaticMesh* StaticMesh = Cast<UStaticMeshComponent>(Component)->GetStaticMesh();
 
diff --git a/Engine/Source/Editor/Experimental/EditorInteractiveToolsFramework/Public/ToolTargets/StaticMeshComponentToolTarget.h b/Engine/Source/Editor/Experimental/EditorInteractiveToolsFramework/Public/ToolTargets/StaticMeshComponentToolTarget.h
--- a/Engine/Source/Editor/Experimental/EditorInteractiveToolsFramework/Public/ToolTargets/StaticMeshComponentToolTarget.h
+++ b/Engine/Source/Editor/Experimental/EditorInteractiveToolsFramework/Public/ToolTargets/StaticMeshComponentToolTarget.h
@@ -4,6 +4,7 @@
 
 #include "CoreMinimal.h"
 
+#include "TargetInterfaces/MaterialProvider.h"
 #include "TargetInterfaces/MeshDescriptionCommitter.h"
 #include "TargetInterfaces/MeshDescriptionProvider.h"
 #include "ToolTargets/PrimitiveComponentToolTarget.h"
@@ -16,7 +17,7 @@
  */
 UCLASS(Transient)
 class EDITORINTERACTIVETOOLSFRAMEWORK_API UStaticMeshComponentToolTarget : public UPrimitiveComponentToolTarget,
-	public IMeshDescriptionCommitter, public IMeshDescriptionProvider
+	public IMeshDescriptionCommitter, public IMeshDescriptionProvider, public IMaterialProvider
 {
 	GENERATED_BODY()
 
@@ -28,8 +29,12 @@ public:
 	// IMeshDescritpionCommitter implementation
 	virtual void CommitMeshDescription(const FCommitter& Committer) override;
 
-	// IPrimitiveComponentBackedTarget implementation
-	virtual void CommitMaterialSetUpdate(const FComponentMaterialSet& MaterialSet) override;
+	// IMaterialProvider implementation
+	int32 GetNumMaterials() const override;
+	UMaterialInterface* GetMaterial(int32 MaterialIndex) const override;
+	void GetMaterialSet(FComponentMaterialSet& MaterialSetOut, bool bPreferAssetMaterials) const override;
+	virtual bool CommitMaterialSetUpdate(const FComponentMaterialSet& MaterialSet, bool bApplyToAsset) override;	
+
 	// Rest provided by parent class
 
 protected:
diff --git a/Engine/Source/Editor/UnrealEd/Private/Tools/EditorComponentSourceFactory.cpp b/Engine/Source/Editor/UnrealEd/Private/Tools/EditorComponentSourceFactory.cpp
--- a/Engine/Source/Editor/UnrealEd/Private/Tools/EditorComponentSourceFactory.cpp
+++ b/Engine/Source/Editor/UnrealEd/Private/Tools/EditorComponentSourceFactory.cpp
@@ -64,45 +64,64 @@ void FStaticMeshComponentTarget::GetMaterialSet(FComponentMaterialSet& MaterialS
 
 void FStaticMeshComponentTarget::CommitMaterialSetUpdate(const FComponentMaterialSet& MaterialSet, bool bApplyToAsset)
 {
-	// we only support this right now...
-	check(bApplyToAsset == true);
 	if (ensure(IsValid()) == false) return;
-	
-	UStaticMesh* StaticMesh = Cast<UStaticMeshComponent>(Component)->GetStaticMesh();
 
-	if (StaticMesh->GetPathName().StartsWith(TEXT("/Engine/")))
+	if (bApplyToAsset)
 	{
-		UE_LOG(LogTemp, Warning, TEXT("CANNOT MODIFY BUILT-IN ENGINE ASSET %s"), *StaticMesh->GetPathName());
-		return;
-	}
+		UStaticMesh* StaticMesh = Cast<UStaticMeshComponent>(Component)->GetStaticMesh();
 
-	// flush any pending rendering commands, which might touch this component while we are rebuilding it's mesh
-	FlushRenderingCommands();
+		if (StaticMesh->GetPathName().StartsWith(TEXT("/Engine/")))
+		{
+			UE_LOG(LogTemp, Warning, TEXT("CANNOT MODIFY BUILT-IN ENGINE ASSET %s"), *StaticMesh->GetPathName());
+			return;
+		}
 
-	// unregister the component while we update it's static mesh
-	TUniquePtr<FComponentReregisterContext> ComponentReregisterContext = MakeUnique<FComponentReregisterContext>(Component);
+		// flush any pending rendering commands, which might touch this component while we are rebuilding its mesh
+		FlushRenderingCommands();
 
-	// make sure transactional flag is on
-	StaticMesh->SetFlags(RF_Transactional);
+		// unregister the component while we update it's static mesh
+		TUniquePtr<FComponentReregisterContext> ComponentReregisterContext = MakeUnique<FComponentReregisterContext>(Component);
 
-	bool bSavedToTransactionBuffer = StaticMesh->Modify();
-	check(bSavedToTransactionBuffer);
-	
-	int NewNumMaterials = MaterialSet.Materials.Num();
-	if (NewNumMaterials != StaticMesh->GetStaticMaterials().Num())
-	{
-		StaticMesh->GetStaticMaterials().SetNum(NewNumMaterials);
+		// make sure transactional flag is on
+		StaticMesh->SetFlags(RF_Transactional);
+
+		StaticMesh->Modify();
+
+		int NewNumMaterials = MaterialSet.Materials.Num();
+		if (NewNumMaterials != StaticMesh->GetStaticMaterials().Num())
+		{
+			StaticMesh->GetStaticMaterials().SetNum(NewNumMaterials);
+		}
+		for (int k = 0; k < NewNumMaterials; ++k)
+		{
+			if (StaticMesh->GetMaterial(k) != MaterialSet.Materials[k])
+			{
+				StaticMesh->SetMaterial(k, MaterialSet.Materials[k]);
+			}
+		}
+
+		StaticMesh->PostEditChange();
 	}
-	for (int k = 0; k < NewNumMaterials; ++k)
+	else
 	{
-		if (StaticMesh->GetMaterial(k) != MaterialSet.Materials[k])
+		int32 NumMaterialsNeeded = Component->GetNumMaterials();
+		int32 NumMaterialsGiven = MaterialSet.Materials.Num();
+
+		// We wrote the below code to support a mismatch in the number of materials.
+		// However, it is not yet clear whether this might be desirable, and we don't
+		// want to inadvertantly hide bugs in the meantime. So, we keep this ensure here
+		// for now, and we can remove it if we decide that we want the ability.
+		ensure(NumMaterialsNeeded == NumMaterialsGiven);
+
+		check(NumMaterialsGiven > 0);
+
+		for (int32 i = 0; i < NumMaterialsNeeded; ++i)
 		{
-			StaticMesh->SetMaterial(k, MaterialSet.Materials[k]);
+			int32 MaterialToUseIndex = FMath::Min(i, NumMaterialsGiven - 1);
+			Component->SetMaterial(i, MaterialSet.Materials[MaterialToUseIndex]);
 		}
 	}
-
-	// right?
-	StaticMesh->PostEditChange();
+	
 }
 
 
diff --git a/Engine/Source/Runtime/Experimental/InteractiveToolsFramework/Private/ToolTargets/PrimitiveComponentToolTarget.cpp b/Engine/Source/Runtime/Experimental/InteractiveToolsFramework/Private/ToolTargets/PrimitiveComponentToolTarget.cpp
--- a/Engine/Source/Runtime/Experimental/InteractiveToolsFramework/Private/ToolTargets/PrimitiveComponentToolTarget.cpp
+++ b/Engine/Source/Runtime/Experimental/InteractiveToolsFramework/Private/ToolTargets/PrimitiveComponentToolTarget.cpp
@@ -27,29 +27,6 @@ void UPrimitiveComponentToolTarget::SetOwnerVisibility(bool bVisible) const
 	}
 }
 
-int32 UPrimitiveComponentToolTarget::GetNumMaterials() const
-{
-	return IsValid() ? Component->GetNumMaterials() : 0;
-}
-
-UMaterialInterface* UPrimitiveComponentToolTarget::GetMaterial(int32 MaterialIndex) const
-{
-	return IsValid() ? Component->GetMaterial(MaterialIndex) : nullptr;
-}
-
-void UPrimitiveComponentToolTarget::GetMaterialSet(FComponentMaterialSet& MaterialSetOut) const
-{
-	if (IsValid())
-	{
-		int32 NumMaterials = Component->GetNumMaterials();
-		MaterialSetOut.Materials.SetNum(NumMaterials);
-		for (int32 k = 0; k < NumMaterials; ++k)
-		{
-			MaterialSetOut.Materials[k] = Component->GetMaterial(k);
-		}
-	}
-}
-
 FTransform UPrimitiveComponentToolTarget::GetWorldTransform() const
 {
 	return IsValid() ? Component->GetComponentTransform() : FTransform::Identity;
diff --git a/Engine/Source/Runtime/Experimental/InteractiveToolsFramework/Public/ComponentSourceInterfaces.h b/Engine/Source/Runtime/Experimental/InteractiveToolsFramework/Public/ComponentSourceInterfaces.h
--- a/Engine/Source/Runtime/Experimental/InteractiveToolsFramework/Public/ComponentSourceInterfaces.h
+++ b/Engine/Source/Runtime/Experimental/InteractiveToolsFramework/Public/ComponentSourceInterfaces.h
@@ -6,7 +6,7 @@
 #include "Math/UnrealMath.h"
 #include "Misc/Optional.h"
 #include "Templates/Function.h"
-#include "TargetInterfaces/PrimitiveComponentBackedTarget.h" // FComponentMaterialSet
+#include "TargetInterfaces/MaterialProvider.h" // FComponentMaterialSet
 #include "Engine/EngineTypes.h"    // FHitResult
 
 // predeclarations
diff --git a/Engine/Source/Runtime/Experimental/InteractiveToolsFramework/Public/TargetInterfaces/MaterialProvider.h b/Engine/Source/Runtime/Experimental/InteractiveToolsFramework/Public/TargetInterfaces/MaterialProvider.h
new file mode 100644
--- /dev/null
+++ b/Engine/Source/Runtime/Experimental/InteractiveToolsFramework/Public/TargetInterfaces/MaterialProvider.h
@@ -0,0 +1,71 @@
+// Copyright Epic Games, Inc. All Rights Reserved.
+
+#pragma once
+
+#include "CoreMinimal.h"
+#include "UObject/Interface.h"
+
+#include "MaterialProvider.generated.h"
+
+class UMaterialInterface;
+
+/**
+ * FComponentMaterialSet is the set of materials assigned to a component (ie Material Slots on a StaticMesh)
+ */
+struct INTERACTIVETOOLSFRAMEWORK_API FComponentMaterialSet
+{
+	TArray<UMaterialInterface*> Materials;
+
+	bool operator!=(const FComponentMaterialSet& Other) const
+	{
+		return Materials != Other.Materials;
+	}
+};
+
+UINTERFACE()
+class INTERACTIVETOOLSFRAMEWORK_API UMaterialProvider : public UInterface
+{
+	GENERATED_BODY()
+};
+
+class INTERACTIVETOOLSFRAMEWORK_API IMaterialProvider
+{
+	GENERATED_BODY()
+
+public:
+
+	/** @return number of material indices in use by this Component */
+	virtual int32 GetNumMaterials() const = 0;
+
+	/**
+	 * Get pointer to a Material provided by this Source
+	 * @param MaterialIndex index of the material
+	 * @return MaterialInterface pointer, or null if MaterialIndex is invalid
+	 */
+	virtual UMaterialInterface* GetMaterial(int32 MaterialIndex) const = 0;
+
+	/**
+	 * Get material set provided by this source
+	 * @param MaterialSetOut returned material set
+	 * @param bAssetMaterials Prefer the underlying asset materials. This may be ignored by targets
+	 *  to which this is not relevant.
+	 */
+	virtual void GetMaterialSet(FComponentMaterialSet& MaterialSetOut, bool bPreferAssetMaterials) const = 0;
+
+	/**
+	 * Commit an update to the material set. This may generate a transaction.
+	 *
+	 * Note that a target may not allow applying the material to an underlying asset, or may not even
+	 * allow having the material applied to the instance, in which case the function will return false
+	 * without doing anything.
+	 *
+	 * TODO: We may decide that we want some functions to query whether the object allows
+	 * one or the other, or we may break this out into a separate interface. We'll wait to
+	 * see when this actually becomes necessary to see the most convenient way to go at that
+	 * point.
+	 *
+	 * @param MaterialSet new list of materials
+	 * @return true if successful
+	 */
+	virtual bool CommitMaterialSetUpdate(const FComponentMaterialSet& MaterialSet, bool bApplyToAsset) = 0;
+};
\ No newline at end of file
diff --git a/Engine/Source/Runtime/Experimental/InteractiveToolsFramework/Public/TargetInterfaces/PrimitiveComponentBackedTarget.h b/Engine/Source/Runtime/Experimental/InteractiveToolsFramework/Public/TargetInterfaces/PrimitiveComponentBackedTarget.h
--- a/Engine/Source/Runtime/Experimental/InteractiveToolsFramework/Public/TargetInterfaces/PrimitiveComponentBackedTarget.h
+++ b/Engine/Source/Runtime/Experimental/InteractiveToolsFramework/Public/TargetInterfaces/PrimitiveComponentBackedTarget.h
@@ -8,24 +8,10 @@
 
 #include "PrimitiveComponentBackedTarget.generated.h"
 
-class UMaterialInterface;
 class UPrimitiveComponent;
 class AActor;
 struct FTransform;
 
-/**
- * FComponentMaterialSet is the set of materials assigned to a component (ie Material Slots on a StaticMesh)
- */
-struct INTERACTIVETOOLSFRAMEWORK_API FComponentMaterialSet
-{
-	TArray<UMaterialInterface*> Materials;
-
-	bool operator!=(const FComponentMaterialSet& Other) const
-	{
-		return Materials != Other.Materials;
-	}
-};
-
 UINTERFACE()
 class INTERACTIVETOOLSFRAMEWORK_API UPrimitiveComponentBackedTarget : public UInterface
 {
@@ -50,22 +36,6 @@ public:
 	 */
 	virtual void SetOwnerVisibility(bool bVisible) const = 0;
 
-	/** @return number of material indices in use by this Component */
-	virtual int32 GetNumMaterials() const = 0;
-
-	/**
-	 * Get pointer to a Material provided by this Source
-	 * @param MaterialIndex index of the material
-	 * @return MaterialInterface pointer, or null if MaterialIndex is invalid
-	 */
-	virtual UMaterialInterface* GetMaterial(int32 MaterialIndex) const = 0;
-
-	/**
-	 * Get material set provided by this source
-	 * @param MaterialSetOut returned material set
-	 */
-	virtual void GetMaterialSet(FComponentMaterialSet& MaterialSetOut) const = 0;
-
 	/**
 	 * @return the transform on this component
 	 * @todo Do we need to return a list of transforms here?
@@ -79,12 +49,4 @@ public:
 	 * @return true if ray intersected Component
 	 */
 	virtual bool HitTestComponent(const FRay& WorldRay, FHitResult& OutHit) const = 0;
-
-
-	// TODO: Should this be in a separate interface?
-	/**
-	 * Commit an update to the material set. This may generate a transaction.
-	 * @param MaterialSet new list of materials
-	 */
-	virtual void CommitMaterialSetUpdate(const FComponentMaterialSet& MaterialSet) = 0;
 };
\ No newline at end of file
diff --git a/Engine/Source/Runtime/Experimental/InteractiveToolsFramework/Public/ToolTargets/PrimitiveComponentToolTarget.h b/Engine/Source/Runtime/Experimental/InteractiveToolsFramework/Public/ToolTargets/PrimitiveComponentToolTarget.h
--- a/Engine/Source/Runtime/Experimental/InteractiveToolsFramework/Public/ToolTargets/PrimitiveComponentToolTarget.h
+++ b/Engine/Source/Runtime/Experimental/InteractiveToolsFramework/Public/ToolTargets/PrimitiveComponentToolTarget.h
@@ -24,14 +24,9 @@ public:
 	virtual UPrimitiveComponent* GetOwnerComponent() const override;
 	virtual AActor* GetOwnerActor() const override;
 	virtual void SetOwnerVisibility(bool bVisible) const override;
-	virtual int32 GetNumMaterials() const override;
-	virtual UMaterialInterface* GetMaterial(int32 MaterialIndex) const override;
-	virtual void GetMaterialSet(FComponentMaterialSet& MaterialSetOut) const override;
 	virtual FTransform GetWorldTransform() const override;
 	virtual bool HitTestComponent(const FRay& WorldRay, FHitResult& OutHit) const override;
 
-	// Not implemented by this class:
-	virtual void CommitMaterialSetUpdate(const FComponentMaterialSet& MaterialSet) override PURE_VIRTUAL(UPrimitiveComponentToolTarget::CommitMaterialSetUpdate);
 protected:
 	UPrimitiveComponent* Component;
 };
\ No newline at end of file
