MeshModelingTools: Made Esc and Enter keys work even when viewport is not focused.

#rnx

diff --git a/Engine/Plugins/Experimental/ModelingToolsEditorMode/Source/ModelingToolsEditorMode/Private/ModelingToolsManagerActions.cpp b/Engine/Plugins/Experimental/ModelingToolsEditorMode/Source/ModelingToolsEditorMode/Private/ModelingToolsManagerActions.cpp
--- a/Engine/Plugins/Experimental/ModelingToolsEditorMode/Source/ModelingToolsEditorMode/Private/ModelingToolsManagerActions.cpp
+++ b/Engine/Plugins/Experimental/ModelingToolsEditorMode/Source/ModelingToolsEditorMode/Private/ModelingToolsManagerActions.cpp
@@ -126,9 +126,10 @@ void FModelingToolsManagerCommands::RegisterCommands()
 	UI_COMMAND(BeginPolyModelTool_Outset, "Outset", "Start the PolyGroup Outset Tool", EUserInterfaceActionType::ToggleButton, FInputChord());
 	UI_COMMAND(BeginPolyModelTool_CutFaces, "Cut", "Start the PolyGroup Cut Faces Tool", EUserInterfaceActionType::ToggleButton, FInputChord());
 
-	UI_COMMAND(AcceptActiveTool, "Accept", "Accept the active tool", EUserInterfaceActionType::Button, FInputChord());
+	// Note that passing a chord into one of these calls hooks the key press to the respective action. 
+	UI_COMMAND(AcceptActiveTool, "Accept", "Accept the active tool", EUserInterfaceActionType::Button, FInputChord(EKeys::Enter));
 	UI_COMMAND(CancelActiveTool, "Cancel", "Cancel the active tool", EUserInterfaceActionType::Button, FInputChord());
-	UI_COMMAND(CompleteActiveTool, "Complete", "Complete the active tool", EUserInterfaceActionType::Button, FInputChord());
+	UI_COMMAND(CompleteActiveTool, "Complete", "Complete the active tool", EUserInterfaceActionType::Button, FInputChord(EKeys::Enter));
 	UI_COMMAND(CancelOrCompleteActiveTool, "Cancel or Complete", "Cancel or complete the active tool", EUserInterfaceActionType::Button, FInputChord(EKeys::Escape));
 }
 
diff --git a/Engine/Source/Editor/Experimental/EditorInteractiveToolsFramework/Private/EdModeInteractiveToolsContext.cpp b/Engine/Source/Editor/Experimental/EditorInteractiveToolsFramework/Private/EdModeInteractiveToolsContext.cpp
--- a/Engine/Source/Editor/Experimental/EditorInteractiveToolsFramework/Private/EdModeInteractiveToolsContext.cpp
+++ b/Engine/Source/Editor/Experimental/EditorInteractiveToolsFramework/Private/EdModeInteractiveToolsContext.cpp
@@ -927,41 +927,6 @@ bool UEdModeInteractiveToolsContext::InputKey(FEditorViewportClient* ViewportCli
 
 	bool bHandled = false;
 
-
-	// escape key cancels current tool
-	if (Key == EKeys::Escape && Event == IE_Released )
-	{
-		if (ToolManager->HasAnyActiveTool())
-		{
-			if (ToolManager->HasActiveTool(EToolSide::Mouse))
-			{
-				DeactivateActiveTool(EToolSide::Mouse, EToolShutdownType::Cancel);
-			}
-			return true;
-		}
-	}
-
-	// enter key accepts current tool, or ends tool if it does not have accept state
-	if (Key == EKeys::Enter && Event == IE_Released && ToolManager->HasAnyActiveTool())
-	{
-		if (ToolManager->HasActiveTool(EToolSide::Mouse))
-		{
-			if (ToolManager->GetActiveTool(EToolSide::Mouse)->HasAccept())
-			{
-				if (ToolManager->CanAcceptActiveTool(EToolSide::Mouse))
-				{
-					DeactivateActiveTool(EToolSide::Mouse, EToolShutdownType::Accept);
-					return true;
-				}
-			}
-			else
-			{
-				DeactivateActiveTool(EToolSide::Mouse, EToolShutdownType::Completed);
-				return true;
-			}
-		}
-	}
-
 	// This is true if we are using the fly camera controls (ie right-mouse possibly + WASD). 
 	// Those controls do *not* capture the mouse and so we still get the events, and we need to ignore them.
 	// Note that it is possible to enter fly camera by holding right-mouse, then hold another button and release right-mouse,
