Editor: Change to SLevelViewport

diff --git a/Engine/Plugins/Cameras/CameraShakePreviewer/Source/CameraShakePreviewer/Private/CameraShakePreviewerModule.cpp b/Engine/Plugins/Cameras/CameraShakePreviewer/Source/CameraShakePreviewer/Private/CameraShakePreviewerModule.cpp
--- a/Engine/Plugins/Cameras/CameraShakePreviewer/Source/CameraShakePreviewer/Private/CameraShakePreviewerModule.cpp
+++ b/Engine/Plugins/Cameras/CameraShakePreviewer/Source/CameraShakePreviewer/Private/CameraShakePreviewerModule.cpp
@@ -16,6 +16,7 @@
 #include "WorkspaceMenuStructure.h"
 #include "WorkspaceMenuStructureModule.h"
 #include "IAssetViewport.h"
+#include "SLevelViewport.h"
 #include "EditorViewportClient.h"
 #include "LevelEditorMenuContext.h"
 
@@ -199,7 +200,7 @@ TSharedRef<FExtender> FCameraShakePreviewerModule::OnExtendLevelViewportOptionMe
 	// Find the viewport for which we're opening this options menu.
 	FLevelEditorModule& LevelEditorModule = FModuleManager::LoadModuleChecked<FLevelEditorModule>(LevelEditorModuleName);
 	TSharedPtr<ILevelEditor> LevelEditor = LevelEditorModule.GetLevelEditorInstance().Pin();
-	TSharedPtr<IAssetViewport> ViewportInterface = LevelEditor->GetActiveViewportInterface();
+	TSharedPtr<SLevelViewport> ViewportInterface = LevelEditor->GetActiveViewportInterface();
 	FLevelEditorViewportClient* ViewportClient = (FLevelEditorViewportClient*)&ViewportInterface->GetAssetViewportClient();
 	if (ViewportClient->ViewportType != ELevelViewportType::LVT_Perspective)
 	{
diff --git a/Engine/Plugins/Experimental/ModelingToolsEditorMode/Source/ModelingToolsEditorMode/Private/ModelingToolsEditorMode.cpp b/Engine/Plugins/Experimental/ModelingToolsEditorMode/Source/ModelingToolsEditorMode/Private/ModelingToolsEditorMode.cpp
--- a/Engine/Plugins/Experimental/ModelingToolsEditorMode/Source/ModelingToolsEditorMode/Private/ModelingToolsEditorMode.cpp
+++ b/Engine/Plugins/Experimental/ModelingToolsEditorMode/Source/ModelingToolsEditorMode/Private/ModelingToolsEditorMode.cpp
@@ -794,8 +794,8 @@ void UModelingToolsEditorMode::ConfigureRealTimeViewportsOverride(bool bEnable)
 	TSharedPtr<ILevelEditor> LevelEditor = LevelEditorModule.GetFirstLevelEditor();
 	if (LevelEditor.IsValid())
 	{
-		TArray<TSharedPtr<IAssetViewport>> Viewports = LevelEditor->GetViewports();
-		for (const TSharedPtr<IAssetViewport>& ViewportWindow : Viewports)
+		TArray<TSharedPtr<SLevelViewport>> Viewports = LevelEditor->GetViewports();
+		for (const TSharedPtr<SLevelViewport>& ViewportWindow : Viewports)
 		{
 			if (ViewportWindow.IsValid())
 			{
diff --git a/Engine/Source/Developer/GameplayDebugger/Private/Editor/GameplayDebuggerToolkit.cpp b/Engine/Source/Developer/GameplayDebugger/Private/Editor/GameplayDebuggerToolkit.cpp
--- a/Engine/Source/Developer/GameplayDebugger/Private/Editor/GameplayDebuggerToolkit.cpp
+++ b/Engine/Source/Developer/GameplayDebugger/Private/Editor/GameplayDebuggerToolkit.cpp
@@ -6,6 +6,7 @@
 #if WITH_EDITOR
 #include "LevelEditor.h"
 #include "IAssetViewport.h"
+#include "SLevelViewport.h"
 #include "LevelEditorViewport.h"
 #include "Widgets/Layout/SBorder.h"
 #include "Widgets/Input/SButton.h"
@@ -89,8 +90,8 @@ FReply FGameplayDebuggerToolkit::OnClickedDisableTool()
 	TSharedPtr<ILevelEditor> LevelEditor = LevelEditorModule.GetFirstLevelEditor();
 	if (LevelEditor.IsValid())
 	{
-		TArray<TSharedPtr<IAssetViewport>> Viewports = LevelEditor->GetViewports();
-		for (const TSharedPtr<IAssetViewport>& ViewportWindow : Viewports)
+		TArray<TSharedPtr<SLevelViewport>> Viewports = LevelEditor->GetViewports();
+		for (const TSharedPtr<SLevelViewport>& ViewportWindow : Viewports)
 		{
 			if (ViewportWindow.IsValid())
 			{
diff --git a/Engine/Source/Editor/Experimental/EditorInteractiveToolsFramework/Private/EdModeInteractiveToolsContext.cpp b/Engine/Source/Editor/Experimental/EditorInteractiveToolsFramework/Private/EdModeInteractiveToolsContext.cpp
--- a/Engine/Source/Editor/Experimental/EditorInteractiveToolsFramework/Private/EdModeInteractiveToolsContext.cpp
+++ b/Engine/Source/Editor/Experimental/EditorInteractiveToolsFramework/Private/EdModeInteractiveToolsContext.cpp
@@ -1246,8 +1246,8 @@ void UEdModeInteractiveToolsContext::SetEditorStateForTool()
 	TSharedPtr<ILevelEditor> LevelEditor = LevelEditorModule.GetFirstLevelEditor();
 	if (LevelEditor.IsValid())
 	{
-		TArray<TSharedPtr<IAssetViewport>> Viewports = LevelEditor->GetViewports();
-		for (const TSharedPtr<IAssetViewport>& ViewportWindow : Viewports)
+		TArray<TSharedPtr<SLevelViewport>> Viewports = LevelEditor->GetViewports();
+		for (const TSharedPtr<SLevelViewport>& ViewportWindow : Viewports)
 		{
 			if (ViewportWindow.IsValid())
 			{
@@ -1270,8 +1270,8 @@ void UEdModeInteractiveToolsContext::RestoreEditorState()
 	TSharedPtr<ILevelEditor> LevelEditor  = LevelEditorModule.GetFirstLevelEditor();
 	if (LevelEditor.IsValid())
 	{
-		TArray<TSharedPtr<IAssetViewport>> Viewports = LevelEditor->GetViewports();
-		for (const TSharedPtr<IAssetViewport>& ViewportWindow : Viewports)
+		TArray<TSharedPtr<SLevelViewport>> Viewports = LevelEditor->GetViewports();
+		for (const TSharedPtr<SLevelViewport>& ViewportWindow : Viewports)
 		{
 			if (ViewportWindow.IsValid())
 			{
diff --git a/Engine/Source/Editor/LandscapeEditor/Private/LandscapeEdMode.cpp b/Engine/Source/Editor/LandscapeEditor/Private/LandscapeEdMode.cpp
--- a/Engine/Source/Editor/LandscapeEditor/Private/LandscapeEdMode.cpp
+++ b/Engine/Source/Editor/LandscapeEditor/Private/LandscapeEdMode.cpp
@@ -36,6 +36,7 @@
 //Slate dependencies
 #include "Misc/FeedbackContext.h"
 #include "IAssetViewport.h"
+#include "SLevelViewport.h"
 #include "SLandscapeEditor.h"
 #include "Framework/Application/SlateApplication.h"
 
@@ -2837,8 +2838,8 @@ void FEdModeLandscape::ForceRealTimeViewports(const bool bEnable)
 	TSharedPtr<ILevelEditor> LevelEditor = LevelEditorModule.GetFirstLevelEditor();
 	if (LevelEditor.IsValid())
 	{
-		TArray<TSharedPtr<IAssetViewport>> Viewports = LevelEditor->GetViewports();
-		for (const TSharedPtr<IAssetViewport>& ViewportWindow : Viewports)
+		TArray<TSharedPtr<SLevelViewport>> Viewports = LevelEditor->GetViewports();
+		for (const TSharedPtr<SLevelViewport>& ViewportWindow : Viewports)
 		{
 			if (ViewportWindow.IsValid())
 			{
diff --git a/Engine/Source/Editor/LevelEditor/Private/SLevelEditor.cpp b/Engine/Source/Editor/LevelEditor/Private/SLevelEditor.cpp
--- a/Engine/Source/Editor/LevelEditor/Private/SLevelEditor.cpp
+++ b/Engine/Source/Editor/LevelEditor/Private/SLevelEditor.cpp
@@ -1769,9 +1769,9 @@ const TArray< TSharedPtr< IToolkit > >& SLevelEditor::GetHostedToolkits() const
 	return HostedToolkits;
 }
 
-TArray< TSharedPtr< IAssetViewport > > SLevelEditor::GetViewports() const
+TArray< TSharedPtr< SLevelViewport > > SLevelEditor::GetViewports() const
 {
-	TArray< TSharedPtr<IAssetViewport> > OutViewports;
+	TArray< TSharedPtr<SLevelViewport> > OutViewports;
 
 	for( int32 TabIndex = 0; TabIndex < ViewportTabs.Num(); ++TabIndex )
 	{
@@ -1810,7 +1810,7 @@ TArray< TSharedPtr< IAssetViewport > > SLevelEditor::GetViewports() const
 	return OutViewports;
 }
  
-TSharedPtr<IAssetViewport> SLevelEditor::GetActiveViewportInterface()
+TSharedPtr<SLevelViewport> SLevelEditor::GetActiveViewportInterface()
 {
 	return GetActiveViewport();
 }
diff --git a/Engine/Source/Editor/LevelEditor/Private/SLevelEditor.h b/Engine/Source/Editor/LevelEditor/Private/SLevelEditor.h
--- a/Engine/Source/Editor/LevelEditor/Private/SLevelEditor.h
+++ b/Engine/Source/Editor/LevelEditor/Private/SLevelEditor.h
@@ -70,8 +70,8 @@ public:
 	virtual void SummonLevelViewportContextMenu(AActor* HitProxyActor = nullptr) override;
 	virtual void SummonLevelViewportViewOptionMenu(const ELevelViewportType ViewOption) override;
 	virtual const TArray< TSharedPtr< class IToolkit > >& GetHostedToolkits() const override;
-	virtual TArray< TSharedPtr< IAssetViewport > > GetViewports() const override;
-	virtual TSharedPtr<IAssetViewport> GetActiveViewportInterface() override;
+	virtual TArray< TSharedPtr< SLevelViewport > > GetViewports() const override;
+	virtual TSharedPtr<SLevelViewport> GetActiveViewportInterface() override;
 	PRAGMA_DISABLE_DEPRECATION_WARNINGS
 	virtual TSharedPtr< class FAssetThumbnailPool > GetThumbnailPool() const override;
 	PRAGMA_ENABLE_DEPRECATION_WARNINGS
diff --git a/Engine/Source/Editor/LevelEditor/Public/ILevelEditor.h b/Engine/Source/Editor/LevelEditor/Public/ILevelEditor.h
--- a/Engine/Source/Editor/LevelEditor/Public/ILevelEditor.h
+++ b/Engine/Source/Editor/LevelEditor/Public/ILevelEditor.h
@@ -41,10 +41,10 @@ public:
 	virtual const TArray< TSharedPtr< IToolkit > >& GetHostedToolkits() const = 0;
 
 	/** Gets an array of all viewports in this level editor */
-	virtual TArray< TSharedPtr< IAssetViewport > > GetViewports() const = 0;
+	virtual TArray< TSharedPtr< SLevelViewport > > GetViewports() const = 0;
 	
 	/** Gets the active level viewport for this level editor */
-	virtual TSharedPtr<IAssetViewport> GetActiveViewportInterface() = 0;
+	virtual TSharedPtr<SLevelViewport> GetActiveViewportInterface() = 0;
 
 	/** Get the thumbnail pool used by this level editor */
 	UE_DEPRECATED(5.0, "GetThumbnailPool has been replaced by UThumbnailManager::Get().GetSharedThumbnailPool().")
