Fixing dropped undo transactions and fly mode wonkiness in viewports with the ITF tools contexts.

The tracking transactions were not happening at the right times, so the tools context had special considerations for some viewport interactions, but not all. This was causing some scale, rotate, and translate operations in the viewport to be lost when the tools context would evaluate the right mouse button behavior for fly mode.

This rearranges how the tracking transaction makes it to the tools context from the viewport, and splits the tools context processing of mouse tracking away from the mode manager, so that the ordering of mouse processing can be maintained by the viewport, instead of the tools context having to maintain knowledge of the viewport.

diff --git a/Engine/Plugins/Experimental/ModelingToolsEditorMode/Source/ModelingToolsEditorMode/Private/ModelingToolsEditorMode.cpp b/Engine/Plugins/Experimental/ModelingToolsEditorMode/Source/ModelingToolsEditorMode/Private/ModelingToolsEditorMode.cpp
--- a/Engine/Plugins/Experimental/ModelingToolsEditorMode/Source/ModelingToolsEditorMode/Private/ModelingToolsEditorMode.cpp
+++ b/Engine/Plugins/Experimental/ModelingToolsEditorMode/Source/ModelingToolsEditorMode/Private/ModelingToolsEditorMode.cpp
@@ -769,20 +769,6 @@ bool UModelingToolsEditorMode::GetPivotForOrbit(FVector& OutPivot) const
 	return false;
 }
 
-bool UModelingToolsEditorMode::InputKey(FEditorViewportClient* ViewportClient, FViewport* Viewport, FKey Key, EInputEvent Event)
-{
-	// TODO: This is a bit of a hack around the fact that when we fly around with right click + wasd,
-	// we still get the key presses passed to us. UEdMode currently does this check  for ToolCommandList,
-	// but we want our hotkeys to live in the toolkit command list so that we respond to them when the
-	// viewport is not focused.
-	// This should get removed when we have fixed wasd flying to capture its keys properly.
-	if (ToolsContext->ShouldIgnoreHotkeys())
-	{
-		return false;
-	}
-	return UEdMode::InputKey(ViewportClient, Viewport, Key, Event);
-}
-
 
 
 void UModelingToolsEditorMode::ConfigureRealTimeViewportsOverride(bool bEnable)
diff --git a/Engine/Plugins/Experimental/ModelingToolsEditorMode/Source/ModelingToolsEditorMode/Public/ModelingToolsEditorMode.h b/Engine/Plugins/Experimental/ModelingToolsEditorMode/Source/ModelingToolsEditorMode/Public/ModelingToolsEditorMode.h
--- a/Engine/Plugins/Experimental/ModelingToolsEditorMode/Source/ModelingToolsEditorMode/Public/ModelingToolsEditorMode.h
+++ b/Engine/Plugins/Experimental/ModelingToolsEditorMode/Source/ModelingToolsEditorMode/Public/ModelingToolsEditorMode.h
@@ -42,8 +42,6 @@ public:
 
 	virtual bool GetPivotForOrbit(FVector& OutPivot) const override;
 
-	virtual bool InputKey(FEditorViewportClient* ViewportClient, FViewport* Viewport, FKey Key, EInputEvent Event) override;
-
 	/*
 	 * focus events
 	 */
diff --git a/Engine/Source/Editor/Experimental/EditorInteractiveToolsFramework/Private/EdModeInteractiveToolsContext.cpp b/Engine/Source/Editor/Experimental/EditorInteractiveToolsFramework/Private/EdModeInteractiveToolsContext.cpp
--- a/Engine/Source/Editor/Experimental/EditorInteractiveToolsFramework/Private/EdModeInteractiveToolsContext.cpp
+++ b/Engine/Source/Editor/Experimental/EditorInteractiveToolsFramework/Private/EdModeInteractiveToolsContext.cpp
@@ -534,7 +534,6 @@ UEdModeInteractiveToolsContext::UEdModeInteractiveToolsContext()
 	QueriesAPI = nullptr;
 	TransactionAPI = nullptr;
 	AssetAPI = nullptr;
-	bIsTrackingMouse = false;
 }
 
 
@@ -572,21 +571,6 @@ void UEdModeInteractiveToolsContext::Initialize(IToolsContextQueriesAPI* Queries
 		RestoreEditorState();
 	});
 
-
-	// If user right-press-drags, this enables "fly mode" in the main viewport, and in that mode the QEWASD keys should
-	// be used for flying control. However the EdMode InputKey/etc system doesn't enforce any of this, we can still also
-	// get that mouse input and hotkeys. So we register a dummy behavior that captures all right-mouse dragging, and
-	// in that mode we set bInFlyMode=true, so that Modes based on this Context will know to skip hotkey processing
-	ULocalClickDragInputBehavior* RightMouseBehavior = NewObject<ULocalClickDragInputBehavior>(this);
-	RightMouseBehavior->CanBeginClickDragFunc = [](const FInputDeviceRay& PressPos) { return  FInputRayHit(0); };
-	RightMouseBehavior->OnClickPressFunc = [this](const FInputDeviceRay&) { bInFlyMode = true; };
-	RightMouseBehavior->OnClickReleaseFunc = [this](const FInputDeviceRay&) { bInFlyMode = false; };
-	RightMouseBehavior->OnTerminateFunc = [this]() { bInFlyMode = false; };
-	RightMouseBehavior->SetDefaultPriority(FInputCapturePriority(0));
-	RightMouseBehavior->SetUseRightMouseButton();
-	RightMouseBehavior->Initialize();
-	InputRouter->RegisterBehavior(RightMouseBehavior, this);
-
 	InvalidationTimestamp = 0;
 }
 
@@ -940,88 +924,9 @@ bool UEdModeInteractiveToolsContext::InputKey(FEditorViewportClient* ViewportCli
 	else if (Event == IE_DoubleClick) { UE_LOG(LogTemp, Warning, TEXT("DOUBLECLICK EVENT")); }
 #endif
 
-	bool bHandled = false;
-
-	// This is true if we are using the fly camera controls (ie right-mouse possibly + WASD). 
-	// Those controls do *not* capture the mouse and so we still get the events, and we need to ignore them.
-	// Note that it is possible to enter fly camera by holding right-mouse, then hold another button and release right-mouse,
-	// and that stays in fly mode, so we cannot rely on right-mouse state alone.
-	if (ViewportClient->IsMovingCamera())
-	{
-		// We are still in this state when user releases right-mouse button but is still holding down left-mouse.
-		// In that state we need to allow the InputRouter to see the event, so that the right-mouse-capture behavior can release
-		bool bIsReleaseRightNavButton = Key.IsMouseButton() && (Key == EKeys::RightMouseButton) && (Event == IE_Released);
-		if (bIsReleaseRightNavButton == false)
-		{
-			return false;
-		}
-	}
-
-	// convert doubleclick events to pressed, for now...this is a hack!
-	if (Event == IE_DoubleClick)
-	{
-		Event = IE_Pressed;
-	}
-
 	if (Event == IE_Pressed || Event == IE_Released)
 	{
-		if (Key.IsMouseButton())
-		{
-			bool bIsLeftMouse = (Key == EKeys::LeftMouseButton);
-			bool bIsMiddleMouse = (Key == EKeys::MiddleMouseButton);
-			bool bIsRightMouse = (Key == EKeys::RightMouseButton);
-
-			if (bIsLeftMouse || bIsMiddleMouse || bIsRightMouse)
-			{
-				// if alt is down and we are not capturing, somewhere higher in the ViewportClient/EdMode stack 
-				// is going to start doing alt+mouse camera manipulation. So we should ignore this mouse event.
-				if (ViewportClient->IsAltPressed() && InputRouter->HasActiveMouseCapture() == false)
-				{
-					return false;
-				}
-
-				FInputDeviceState InputState = CurrentMouseState;
-				InputState.InputDevice = EInputDevices::Mouse;
-				InputState.SetModifierKeyStates(
-					ViewportClient->IsShiftPressed(), ViewportClient->IsAltPressed(),
-					ViewportClient->IsCtrlPressed(), ViewportClient->IsCmdPressed());
-
-				if (bIsLeftMouse)
-				{
-					InputState.Mouse.Left.SetStates(
-						(Event == IE_Pressed), (Event == IE_Pressed), (Event == IE_Released));
-					CurrentMouseState.Mouse.Left.bDown = (Event == IE_Pressed);
-				}
-				else if (bIsMiddleMouse)
-				{
-					InputState.Mouse.Middle.SetStates(
-						(Event == IE_Pressed), (Event == IE_Pressed), (Event == IE_Released));
-					CurrentMouseState.Mouse.Middle.bDown = (Event == IE_Pressed);
-				}
-				else
-				{
-					InputState.Mouse.Right.SetStates(
-						(Event == IE_Pressed), (Event == IE_Pressed), (Event == IE_Released));
-					CurrentMouseState.Mouse.Right.bDown = (Event == IE_Pressed);
-				}
-
-				InputRouter->PostInputEvent(InputState);
-
-				if (InputRouter->HasActiveMouseCapture() && bInFlyMode == false)
-				{
-					// what is this about? MeshPaintMode has it...
-					ViewportClient->bLockFlightCamera = true;
-					bHandled = true;   // indicate that we handled this event,
-									   // which will disable camera movement/etc ?
-				}
-				else
-				{
-					//ViewportClient->bLockFlightCamera = false;
-				}
-
-			}
-		}
-		else if (Key.IsGamepadKey())
+		if (Key.IsGamepadKey())
 		{
 			// not supported yet
 		}
@@ -1043,16 +948,13 @@ bool UEdModeInteractiveToolsContext::InputKey(FEditorViewportClient* ViewportCli
 			InputState.Keyboard.ActiveKey.Button = Key;
 			bool bPressed = (Event == IE_Pressed);
 			InputState.Keyboard.ActiveKey.SetStates(bPressed, bPressed, !bPressed);
-			InputRouter->PostInputEvent(InputState);
+			return InputRouter->PostInputEvent(InputState);
 		}
-
 	}
 
-	return bHandled;
+	return false;
 }
 
-
-
 bool UEdModeInteractiveToolsContext::MouseEnter(FEditorViewportClient* ViewportClient, FViewport* Viewport, int32 x, int32 y)
 {
 #ifdef ENABLE_DEBUG_PRINTING
@@ -1111,18 +1013,45 @@ bool UEdModeInteractiveToolsContext::MouseLeave(FEditorViewportClient* ViewportC
 
 bool UEdModeInteractiveToolsContext::StartTracking(FEditorViewportClient* InViewportClient, FViewport* InViewport)
 {
-	bIsTrackingMouse = InputRouter->HasActiveMouseCapture();
-	return bIsTrackingMouse;
-}
+	bool bIsLeftMouse = InViewport->KeyState(EKeys::LeftMouseButton);
+	bool bIsMiddleMouse = InViewport->KeyState(EKeys::MiddleMouseButton);
+	bool bIsRightMouse = InViewport->KeyState(EKeys::RightMouseButton);
 
-bool UEdModeInteractiveToolsContext::CapturedMouseMove(FEditorViewportClient* InViewportClient, FViewport* InViewport, int32 InMouseX, int32 InMouseY)
-{
-	// if alt is down we will not allow client to see this event
-	if (InViewportClient->IsAltPressed())
+	if (bIsLeftMouse || bIsMiddleMouse || bIsRightMouse)
 	{
-		return false;
+		FInputDeviceState InputState = CurrentMouseState;
+		InputState.InputDevice = EInputDevices::Mouse;
+		InputState.SetModifierKeyStates(
+			InViewportClient->IsShiftPressed(), InViewportClient->IsAltPressed(),
+			InViewportClient->IsCtrlPressed(), InViewportClient->IsCmdPressed());
+
+		if (bIsLeftMouse)
+		{
+			InputState.Mouse.Left.SetStates(true, true, false);
+			CurrentMouseState.Mouse.Left.bDown = true;
+		}
+		else if (bIsMiddleMouse)
+		{
+			InputState.Mouse.Middle.SetStates(true, true, false);
+			CurrentMouseState.Mouse.Middle.bDown = true;
+		}
+		else
+		{
+			InputState.Mouse.Right.SetStates(true, true, false);
+			CurrentMouseState.Mouse.Right.bDown = true;
+		}
+
+		if (InputRouter->PostInputEvent(InputState))
+		{
+			InViewportClient->bLockFlightCamera = true;
+			return true;
+		}
 	}
+	return false;
+}
 
+bool UEdModeInteractiveToolsContext::CapturedMouseMove(FEditorViewportClient* InViewportClient, FViewport* InViewport, int32 InMouseX, int32 InMouseY)
+{
 	FVector2D OldPosition = CurrentMouseState.Mouse.Position2D;
 	CurrentMouseState.Mouse.Position2D = FVector2D(InMouseX, InMouseY);
 	CurrentMouseState.Mouse.WorldRay = GetRayFromMousePos(InViewportClient, InViewport, InMouseX, InMouseY);
@@ -1151,9 +1080,31 @@ bool UEdModeInteractiveToolsContext::EndTracking(FEditorViewportClient* InViewpo
 	// unlock flight camera
 	InViewportClient->bLockFlightCamera = false;
 
-	if (bIsTrackingMouse)
+	FInputDeviceState InputState = CurrentMouseState;
+	InputState.InputDevice = EInputDevices::Mouse;
+	InputState.SetModifierKeyStates(
+		InViewportClient->IsShiftPressed(), InViewportClient->IsAltPressed(),
+		InViewportClient->IsCtrlPressed(), InViewportClient->IsCmdPressed());
+
+	// Clear the current mouse state button down
+	if (CurrentMouseState.Mouse.Left.bDown)
+	{
+		CurrentMouseState.Mouse.Left.bDown = false;
+		InputState.Mouse.Left.SetStates(false, false, true);
+	}
+	if (CurrentMouseState.Mouse.Middle.bDown)
+	{
+		CurrentMouseState.Mouse.Middle.bDown = false;
+		InputState.Mouse.Middle.SetStates(false, false, true);
+	}
+	if (CurrentMouseState.Mouse.Right.bDown)
+	{
+		CurrentMouseState.Mouse.Right.bDown = false;
+		InputState.Mouse.Right.SetStates(false, false, true);
+	}
+
+	if (InputRouter->PostInputEvent(InputState))
 	{
-		bIsTrackingMouse = false;
 		return true;
 	}
 
diff --git a/Engine/Source/Editor/Experimental/EditorInteractiveToolsFramework/Public/EdModeInteractiveToolsContext.h b/Engine/Source/Editor/Experimental/EditorInteractiveToolsFramework/Public/EdModeInteractiveToolsContext.h
--- a/Engine/Source/Editor/Experimental/EditorInteractiveToolsFramework/Public/EdModeInteractiveToolsContext.h
+++ b/Engine/Source/Editor/Experimental/EditorInteractiveToolsFramework/Public/EdModeInteractiveToolsContext.h
@@ -91,8 +91,6 @@ public:
 	void StartTool(const FString ToolTypeIdentifier);
 	void EndTool(EToolShutdownType ShutdownType);
 
-	bool ShouldIgnoreHotkeys() const { return bInFlyMode; }
-
 	FRay GetLastWorldRay() const;
 
 protected:
@@ -138,9 +136,6 @@ protected:
 	// Copy-pasted from other Editor code, seems kind of expensive?
 	static FRay GetRayFromMousePos(FEditorViewportClient* ViewportClient, FViewport* Viewport, int MouseX, int MouseY);
 
-	/** This will be set to true if user is in right-mouse "fly mode", which requires special handling to intercept hotkeys/etc */
-	bool bInFlyMode = false;
-
 	// editor UI state that we set before starting tool and when exiting tool
 	// Currently disabling anti-aliasing during active Tools because it causes PDI flickering
 	void SetEditorStateForTool();
@@ -151,5 +146,4 @@ protected:
 
 private:
 	FEditorModeTools* EditorModeManager = nullptr;
-	bool bIsTrackingMouse;
 };
diff --git a/Engine/Source/Editor/UnrealEd/Private/EditorModeManager.cpp b/Engine/Source/Editor/UnrealEd/Private/EditorModeManager.cpp
--- a/Engine/Source/Editor/UnrealEd/Private/EditorModeManager.cpp
+++ b/Engine/Source/Editor/UnrealEd/Private/EditorModeManager.cpp
@@ -1137,7 +1137,7 @@ bool FEditorModeTools::StartTracking(FEditorViewportClient* InViewportClient, FV
 	bIsTracking = true;
 	CachedLocation = PivotLocation;	// Cache the pivot location
 
-	bool bTransactionHandled = InteractiveToolsContext->StartTracking(InViewportClient, InViewport);
+	bool bTransactionHandled = false;
 	for (UEdMode* Mode : ActiveScriptableModes)
 	{
 		bTransactionHandled |= Mode->StartTracking(InViewportClient, InViewportClient->Viewport);
@@ -1150,7 +1150,7 @@ bool FEditorModeTools::StartTracking(FEditorViewportClient* InViewportClient, FV
 bool FEditorModeTools::EndTracking(FEditorViewportClient* InViewportClient, FViewport* InViewport)
 {
 	bIsTracking = false;
-	bool bTransactionHandled = InteractiveToolsContext->EndTracking(InViewportClient, InViewport);
+	bool bTransactionHandled = false;
 
 	for (UEdMode* Mode : ActiveScriptableModes)
 	{
@@ -1380,21 +1380,15 @@ bool FEditorModeTools::ProcessCapturedMouseMoves( FEditorViewportClient* InViewp
 /** Notifies all active modes of keyboard input */
 bool FEditorModeTools::InputKey(FEditorViewportClient* InViewportClient, FViewport* Viewport, FKey Key, EInputEvent Event)
 {
-	bool bToolsContextHandled = InteractiveToolsContext->InputKey(InViewportClient, Viewport, Key, Event);
 	bool bHandled = false;
 	
-	//Copy the modes and iterate of that since a key may remove the edit mode and change CopyActiveScriptableModes
+	// Copy the modes and iterate of that since a key may remove the edit mode and change CopyActiveScriptableModes
 	TArray<UEdMode*> CopyActiveScriptableModes(ActiveScriptableModes);
 	for (UEdMode* Mode : CopyActiveScriptableModes)
 	{
-		// If the interactive tools context did not handle the InputKey, then pass along to the UEdMode
-		// Always pass through for the FEdModes to handle for backwards compat.
-		if (!bToolsContextHandled || Mode->AsLegacyMode())
-		{
-			bHandled |= Mode->InputKey(InViewportClient, Viewport, Key, Event);
-		}
+		bHandled |= Mode->InputKey(InViewportClient, Viewport, Key, Event);
 	}
-	return bHandled || bToolsContextHandled;
+	return bHandled;
 }
 
 /** Notifies all active modes of axis movement */
diff --git a/Engine/Source/Editor/UnrealEd/Private/EditorViewportClient.cpp b/Engine/Source/Editor/UnrealEd/Private/EditorViewportClient.cpp
--- a/Engine/Source/Editor/UnrealEd/Private/EditorViewportClient.cpp
+++ b/Engine/Source/Editor/UnrealEd/Private/EditorViewportClient.cpp
@@ -57,6 +57,7 @@
 #include "UnrealEngine.h"
 #include "BufferVisualizationData.h"
 #include "UnrealWidget.h"
+#include "EdModeInteractiveToolsContext.h"
 
 #include "CustomEditorStaticScreenPercentage.h"
 
@@ -2929,6 +2930,12 @@ bool FEditorViewportClient::InputKey(FViewport* InViewport, int32 ControllerId,
 		bHandled |= true;
 	}
 
+	// Give the interactive tools context a shot at input if it hasn't been handled already.
+	if (!bHandled)
+	{
+		bHandled = ModeTools->GetInteractiveToolsContext()->InputKey(this, Viewport, Key, Event);
+	}
+
 	//apply the visibility and set the cursor positions
 	ApplyRequiredCursorVisibility( true );
 	return bHandled;
@@ -4683,7 +4690,7 @@ bool FEditorViewportClient::ShouldPanOrDollyCamera() const
 {
 	const bool bIsCtrlDown = IsCtrlPressed();
 
-	const bool bLeftMouseButtonDown = Viewport->KeyState( EKeys::LeftMouseButton );
+	const bool bLeftMouseButtonDown = Viewport->KeyState( EKeys::LeftMouseButton ) && !bLockFlightCamera;
 	const bool bRightMouseButtonDown = Viewport->KeyState( EKeys::RightMouseButton );
 	const bool bIsMarqueeSelect = IsOrtho() && bLeftMouseButtonDown;
 
@@ -5045,7 +5052,7 @@ void FEditorViewportClient::MoveViewportCamera(const FVector& InDrag, const FRot
 				CameraUserImpulseData->RotatePitchVelocityModifier += VelModRotSpeed * RotEuler.Y / ViewportSettings->MouseSensitivty;
 				CameraUserImpulseData->RotateYawVelocityModifier += VelModRotSpeed * RotEuler.Z / ViewportSettings->MouseSensitivty;
 			}
-			else
+			else if (!bLockFlightCamera)
 			{
 				MoveViewportPerspectiveCamera( InDrag, InRot, bDollyCamera );
 			}
diff --git a/Engine/Source/Editor/UnrealEd/Private/LevelEditorViewport.cpp b/Engine/Source/Editor/UnrealEd/Private/LevelEditorViewport.cpp
--- a/Engine/Source/Editor/UnrealEd/Private/LevelEditorViewport.cpp
+++ b/Engine/Source/Editor/UnrealEd/Private/LevelEditorViewport.cpp
@@ -93,6 +93,7 @@
 #include "Materials/MaterialExpressionCustom.h"
 #include "Materials/MaterialExpressionWorldPosition.h"
 #include "UnrealWidget.h"
+#include "EdModeInteractiveToolsContext.h"
 
 DEFINE_LOG_CATEGORY(LogEditorViewport);
 
@@ -3023,6 +3024,11 @@ void FLevelEditorViewportClient::TrackingStarted( const FInputEventState& InInpu
 			// Suspend actor/component modification during each delta step to avoid recording unnecessary overhead into the transaction buffer
 			GEditor->DisableDeltaModification(true);
 		}
+		else if (!IsFlightCameraActive())
+		{
+			// Give the interactive tools context a shot at tracking.
+			ModeTools->GetInteractiveToolsContext()->StartTracking(this, Viewport);
+		}
 	}
 }
 
@@ -3106,6 +3112,10 @@ void FLevelEditorViewportClient::TrackingStopped()
 		// Restore actor/component delta modification
 		GEditor->DisableDeltaModification(false);
 	}
+	else
+	{
+		ModeTools->GetInteractiveToolsContext()->EndTracking(this, Viewport);
+	}
 
 	ModeTools->ActorMoveNotify();
 
diff --git a/Engine/Source/Editor/UnrealEd/Private/Tools/UEdMode.cpp b/Engine/Source/Editor/UnrealEd/Private/Tools/UEdMode.cpp
--- a/Engine/Source/Editor/UnrealEd/Private/Tools/UEdMode.cpp
+++ b/Engine/Source/Editor/UnrealEd/Private/Tools/UEdMode.cpp
@@ -115,12 +115,9 @@ bool UEdMode::InputKey(FEditorViewportClient* ViewportClient, FViewport* Viewpor
 
 	if (Event != IE_Released)
 	{
-		if (ToolsContext->ShouldIgnoreHotkeys() == false)		// allow the context to capture keyboard input if necessary
+		if (ToolCommandList->ProcessCommandBindings(Key, FSlateApplication::Get().GetModifierKeys(), false/*Event == IE_Repeat*/))
 		{
-			if (ToolCommandList->ProcessCommandBindings(Key, FSlateApplication::Get().GetModifierKeys(), false/*Event == IE_Repeat*/))
-			{
-				return true;
-			}
+			return true;
 		}
 	}
 
diff --git a/Engine/Source/Runtime/Experimental/InteractiveToolsFramework/Private/InputRouter.cpp b/Engine/Source/Runtime/Experimental/InteractiveToolsFramework/Private/InputRouter.cpp
--- a/Engine/Source/Runtime/Experimental/InteractiveToolsFramework/Private/InputRouter.cpp
+++ b/Engine/Source/Runtime/Experimental/InteractiveToolsFramework/Private/InputRouter.cpp
@@ -54,26 +54,28 @@ void UInputRouter::RegisterBehavior(UInputBehavior* Behavior, void* Source, cons
 }
 
 
-void UInputRouter::PostInputEvent(const FInputDeviceState& Input)
+bool UInputRouter::PostInputEvent(const FInputDeviceState& Input)
 {
 	if (ActiveInputBehaviors->IsEmpty())
 	{
-		return;
+		return false;
 	}
 
 	if (Input.IsFromDevice(EInputDevices::Mouse))
 	{
 		PostInputEvent_Mouse(Input);
+		return HasActiveMouseCapture();
 	}
 	else if (Input.IsFromDevice(EInputDevices::Keyboard))
 	{
 		PostInputEvent_Keyboard(Input);
+		return (ActiveKeyboardCapture != nullptr);
 	}
 	else
 	{
 		unimplemented();
 		TransactionsAPI->DisplayMessage(LOCTEXT("PostInputEventMessage", "UInteractiveToolManager::PostInputEvent - input device is not currently supported."), EToolMessageLevel::Internal);
-		return;
+		return false;
 	}
 }
 
diff --git a/Engine/Source/Runtime/Experimental/InteractiveToolsFramework/Public/InputRouter.h b/Engine/Source/Runtime/Experimental/InteractiveToolsFramework/Public/InputRouter.h
--- a/Engine/Source/Runtime/Experimental/InteractiveToolsFramework/Public/InputRouter.h
+++ b/Engine/Source/Runtime/Experimental/InteractiveToolsFramework/Public/InputRouter.h
@@ -52,7 +52,7 @@ public:
 
 
 	/** Insert a new input event which is used to check for new captures, or forwarded to active capture */
-	virtual void PostInputEvent(const FInputDeviceState& Input);
+	virtual bool PostInputEvent(const FInputDeviceState& Input);
 
 	/** Returns true if there is an active mouse capture */
 	virtual bool HasActiveMouseCapture() const;
