Fixing broken Ctrl+click and other modifier behavior in MeshModelingTools by reverting most of the input processing changes from 15676489. Made tools not process inputs captured by viewport flying to keep another issue fixed. Made KeyAsModifierInputBehavior not capture input.

#rnx

diff --git a/Engine/Plugins/Experimental/MeshModelingToolset/Source/MeshModelingTools/Private/MirrorTool.cpp b/Engine/Plugins/Experimental/MeshModelingToolset/Source/MeshModelingTools/Private/MirrorTool.cpp
--- a/Engine/Plugins/Experimental/MeshModelingToolset/Source/MeshModelingTools/Private/MirrorTool.cpp
+++ b/Engine/Plugins/Experimental/MeshModelingToolset/Source/MeshModelingTools/Private/MirrorTool.cpp
@@ -286,8 +286,7 @@ void UMirrorTool::Setup()
 
 	// Add modifier button for snapping
 	UKeyAsModifierInputBehavior* SnapToggleBehavior = NewObject<UKeyAsModifierInputBehavior>();
-	SnapToggleBehavior->ModifierCheckFunc = FInputDeviceState::IsShiftKeyDown;
-	SnapToggleBehavior->Initialize(this, SnappingToggleModifierId, EKeys::AnyKey);
+	SnapToggleBehavior->Initialize(this, SnappingToggleModifierId, FInputDeviceState::IsShiftKeyDown);
 	AddInputBehavior(SnapToggleBehavior);
 
 	// Start the preview calculations
diff --git a/Engine/Plugins/Experimental/MeshModelingToolset/Source/MeshModelingTools/Private/PlaneCutTool.cpp b/Engine/Plugins/Experimental/MeshModelingToolset/Source/MeshModelingTools/Private/PlaneCutTool.cpp
--- a/Engine/Plugins/Experimental/MeshModelingToolset/Source/MeshModelingTools/Private/PlaneCutTool.cpp
+++ b/Engine/Plugins/Experimental/MeshModelingToolset/Source/MeshModelingTools/Private/PlaneCutTool.cpp
@@ -107,8 +107,7 @@ void UPlaneCutTool::Setup()
 
 	// add modifier button for snapping
 	UKeyAsModifierInputBehavior* SnapToggleBehavior = NewObject<UKeyAsModifierInputBehavior>();
-	SnapToggleBehavior->ModifierCheckFunc = FInputDeviceState::IsShiftKeyDown;
-	SnapToggleBehavior->Initialize(this, IgnoreSnappingModifier, EKeys::AnyKey);
+	SnapToggleBehavior->Initialize(this, IgnoreSnappingModifier, FInputDeviceState::IsShiftKeyDown);
 	AddInputBehavior(SnapToggleBehavior);
 
 	// hide input StaticMeshComponents
diff --git a/Engine/Plugins/Experimental/MeshModelingToolset/Source/ModelingComponents/Private/Mechanics/DragAlignmentMechanic.cpp b/Engine/Plugins/Experimental/MeshModelingToolset/Source/ModelingComponents/Private/Mechanics/DragAlignmentMechanic.cpp
--- a/Engine/Plugins/Experimental/MeshModelingToolset/Source/ModelingComponents/Private/Mechanics/DragAlignmentMechanic.cpp
+++ b/Engine/Plugins/Experimental/MeshModelingToolset/Source/ModelingComponents/Private/Mechanics/DragAlignmentMechanic.cpp
@@ -22,8 +22,7 @@ void UDragAlignmentMechanic::Setup(UInteractiveTool* ParentToolIn)
 
 	// Register modifier listeners.
 	UKeyAsModifierInputBehavior* AlignmentToggleBehavior = NewObject<UKeyAsModifierInputBehavior>();
-	AlignmentToggleBehavior->ModifierCheckFunc = FInputDeviceState::IsCtrlKeyDown;
-	AlignmentToggleBehavior->Initialize(this, AlignmentModifierID, EKeys::AnyKey);
+	AlignmentToggleBehavior->Initialize(this, AlignmentModifierID, FInputDeviceState::IsCtrlKeyDown);
 	ParentTool->AddInputBehavior(AlignmentToggleBehavior);
 
 	// Set up the function that casts rays into the world.
diff --git a/Engine/Source/Editor/Experimental/EditorInteractiveToolsFramework/Private/EdModeInteractiveToolsContext.cpp b/Engine/Source/Editor/Experimental/EditorInteractiveToolsFramework/Private/EdModeInteractiveToolsContext.cpp
--- a/Engine/Source/Editor/Experimental/EditorInteractiveToolsFramework/Private/EdModeInteractiveToolsContext.cpp
+++ b/Engine/Source/Editor/Experimental/EditorInteractiveToolsFramework/Private/EdModeInteractiveToolsContext.cpp
@@ -926,7 +926,43 @@ bool UEdModeInteractiveToolsContext::InputKey(FEditorViewportClient* ViewportCli
 
 	if (Event == IE_Pressed || Event == IE_Released)
 	{
-		if (Key.IsGamepadKey())
+		if (Key.IsMouseButton())
+		{
+			bool bIsLeftMouse = (Key == EKeys::LeftMouseButton);
+			bool bIsMiddleMouse = (Key == EKeys::MiddleMouseButton);
+			bool bIsRightMouse = (Key == EKeys::RightMouseButton);
+			if (bIsLeftMouse || bIsMiddleMouse || bIsRightMouse)
+			{
+				FInputDeviceState InputState = CurrentMouseState;
+				InputState.InputDevice = EInputDevices::Mouse;
+				InputState.SetModifierKeyStates(
+					ViewportClient->IsShiftPressed(), ViewportClient->IsAltPressed(),
+					ViewportClient->IsCtrlPressed(), ViewportClient->IsCmdPressed());
+				if (bIsLeftMouse)
+				{
+					InputState.Mouse.Left.SetStates(
+						(Event == IE_Pressed), (Event == IE_Pressed), (Event == IE_Released));
+					CurrentMouseState.Mouse.Left.bDown = (Event == IE_Pressed);
+				}
+				else if (bIsMiddleMouse)
+				{
+					InputState.Mouse.Middle.SetStates(
+						(Event == IE_Pressed), (Event == IE_Pressed), (Event == IE_Released));
+					CurrentMouseState.Mouse.Middle.bDown = (Event == IE_Pressed);
+				}
+				else
+				{
+					InputState.Mouse.Right.SetStates(
+						(Event == IE_Pressed), (Event == IE_Pressed), (Event == IE_Released));
+					CurrentMouseState.Mouse.Right.bDown = (Event == IE_Pressed);
+				}
+				if (InputRouter->PostInputEvent(InputState))
+				{
+					return true;
+				}
+			}
+		}
+		else if (Key.IsGamepadKey())
 		{
 			// not supported yet
 		}
@@ -1013,41 +1049,8 @@ bool UEdModeInteractiveToolsContext::MouseLeave(FEditorViewportClient* ViewportC
 
 bool UEdModeInteractiveToolsContext::StartTracking(FEditorViewportClient* InViewportClient, FViewport* InViewport)
 {
-	bool bIsLeftMouse = InViewport->KeyState(EKeys::LeftMouseButton);
-	bool bIsMiddleMouse = InViewport->KeyState(EKeys::MiddleMouseButton);
-	bool bIsRightMouse = InViewport->KeyState(EKeys::RightMouseButton);
-
-	if (bIsLeftMouse || bIsMiddleMouse || bIsRightMouse)
-	{
-		FInputDeviceState InputState = CurrentMouseState;
-		InputState.InputDevice = EInputDevices::Mouse;
-		InputState.SetModifierKeyStates(
-			InViewportClient->IsShiftPressed(), InViewportClient->IsAltPressed(),
-			InViewportClient->IsCtrlPressed(), InViewportClient->IsCmdPressed());
-
-		if (bIsLeftMouse)
-		{
-			InputState.Mouse.Left.SetStates(true, true, false);
-			CurrentMouseState.Mouse.Left.bDown = true;
-		}
-		else if (bIsMiddleMouse)
-		{
-			InputState.Mouse.Middle.SetStates(true, true, false);
-			CurrentMouseState.Mouse.Middle.bDown = true;
-		}
-		else
-		{
-			InputState.Mouse.Right.SetStates(true, true, false);
-			CurrentMouseState.Mouse.Right.bDown = true;
-		}
-
-		if (InputRouter->PostInputEvent(InputState))
-		{
-			InViewportClient->bLockFlightCamera = true;
-			return true;
-		}
-	}
-	return false;
+	bIsTrackingMouse = InputRouter->HasActiveMouseCapture();
+	return bIsTrackingMouse;
 }
 
 bool UEdModeInteractiveToolsContext::CapturedMouseMove(FEditorViewportClient* InViewportClient, FViewport* InViewport, int32 InMouseX, int32 InMouseY)
@@ -1077,34 +1080,9 @@ bool UEdModeInteractiveToolsContext::CapturedMouseMove(FEditorViewportClient* In
 
 bool UEdModeInteractiveToolsContext::EndTracking(FEditorViewportClient* InViewportClient, FViewport* InViewport)
 {
-	// unlock flight camera
-	InViewportClient->bLockFlightCamera = false;
-
-	FInputDeviceState InputState = CurrentMouseState;
-	InputState.InputDevice = EInputDevices::Mouse;
-	InputState.SetModifierKeyStates(
-		InViewportClient->IsShiftPressed(), InViewportClient->IsAltPressed(),
-		InViewportClient->IsCtrlPressed(), InViewportClient->IsCmdPressed());
-
-	// Clear the current mouse state button down
-	if (CurrentMouseState.Mouse.Left.bDown)
-	{
-		CurrentMouseState.Mouse.Left.bDown = false;
-		InputState.Mouse.Left.SetStates(false, false, true);
-	}
-	if (CurrentMouseState.Mouse.Middle.bDown)
-	{
-		CurrentMouseState.Mouse.Middle.bDown = false;
-		InputState.Mouse.Middle.SetStates(false, false, true);
-	}
-	if (CurrentMouseState.Mouse.Right.bDown)
-	{
-		CurrentMouseState.Mouse.Right.bDown = false;
-		InputState.Mouse.Right.SetStates(false, false, true);
-	}
-
-	if (InputRouter->PostInputEvent(InputState))
+	if (bIsTrackingMouse)
 	{
+		bIsTrackingMouse = false;
 		return true;
 	}
 
diff --git a/Engine/Source/Editor/Experimental/EditorInteractiveToolsFramework/Public/EdModeInteractiveToolsContext.h b/Engine/Source/Editor/Experimental/EditorInteractiveToolsFramework/Public/EdModeInteractiveToolsContext.h
--- a/Engine/Source/Editor/Experimental/EditorInteractiveToolsFramework/Public/EdModeInteractiveToolsContext.h
+++ b/Engine/Source/Editor/Experimental/EditorInteractiveToolsFramework/Public/EdModeInteractiveToolsContext.h
@@ -146,4 +146,5 @@ protected:
 
 private:
 	FEditorModeTools* EditorModeManager = nullptr;
+	bool bIsTrackingMouse;
 };
diff --git a/Engine/Source/Editor/UnrealEd/Private/EditorModeManager.cpp b/Engine/Source/Editor/UnrealEd/Private/EditorModeManager.cpp
--- a/Engine/Source/Editor/UnrealEd/Private/EditorModeManager.cpp
+++ b/Engine/Source/Editor/UnrealEd/Private/EditorModeManager.cpp
@@ -1137,7 +1137,7 @@ bool FEditorModeTools::StartTracking(FEditorViewportClient* InViewportClient, FV
 	bIsTracking = true;
 	CachedLocation = PivotLocation;	// Cache the pivot location
 
-	bool bTransactionHandled = false;
+	bool bTransactionHandled = InteractiveToolsContext->StartTracking(InViewportClient, InViewport);
 	for (UEdMode* Mode : ActiveScriptableModes)
 	{
 		bTransactionHandled |= Mode->StartTracking(InViewportClient, InViewportClient->Viewport);
@@ -1150,7 +1150,7 @@ bool FEditorModeTools::StartTracking(FEditorViewportClient* InViewportClient, FV
 bool FEditorModeTools::EndTracking(FEditorViewportClient* InViewportClient, FViewport* InViewport)
 {
 	bIsTracking = false;
-	bool bTransactionHandled = false;
+	bool bTransactionHandled = InteractiveToolsContext->EndTracking(InViewportClient, InViewport);
 
 	for (UEdMode* Mode : ActiveScriptableModes)
 	{
@@ -1378,8 +1378,12 @@ bool FEditorModeTools::ProcessCapturedMouseMoves( FEditorViewportClient* InViewp
 }
 
 /** Notifies all active modes of keyboard input */
-bool FEditorModeTools::InputKey(FEditorViewportClient* InViewportClient, FViewport* Viewport, FKey Key, EInputEvent Event)
+bool FEditorModeTools::InputKey(FEditorViewportClient* InViewportClient, FViewport* Viewport, FKey Key, EInputEvent Event, bool bRouteToToolsContext)
 {
+	if (bRouteToToolsContext && InteractiveToolsContext->InputKey(InViewportClient, Viewport, Key, Event))
+	{
+		return true;
+	}
 	bool bHandled = false;
 	
 	// Copy the modes and iterate of that since a key may remove the edit mode and change CopyActiveScriptableModes
diff --git a/Engine/Source/Editor/UnrealEd/Private/EditorViewportClient.cpp b/Engine/Source/Editor/UnrealEd/Private/EditorViewportClient.cpp
--- a/Engine/Source/Editor/UnrealEd/Private/EditorViewportClient.cpp
+++ b/Engine/Source/Editor/UnrealEd/Private/EditorViewportClient.cpp
@@ -2786,8 +2786,15 @@ bool FEditorViewportClient::InputKey(FViewport* InViewport, int32 ControllerId,
 		return true;
 	}
 
-	// Let the current mode have a look at the input before reacting to it.
-	if (ModeTools->InputKey(this, Viewport, Key, Event))
+	// Let the current mode have a look at the input before reacting to it. 
+	// Note that bIsTracking tells us whether the viewport client is capturing mouse behavior to fly around,
+	// move objects, etc. In this case we don't want to pass the input the input router, because it is already
+	// captured. One may reasonably say that we don't want to pass the input to the modes in this case either, 
+	// but unfortunately the legacy behavior is that modes do get this input behavior, which we have to keep
+	// for now. In the long term, only old FEdModes will get this input, while the new UEdModes won't, but
+	// this is not yet split apart. Also in the long term, the viewport behaviors will be inside the
+	// input router and we won't need this bIsTracking flag to begin with.
+	if (ModeTools->InputKey(this, Viewport, Key, Event, /*bRouteToToolsContext*/ !bIsTracking))
 	{
 		return true;
 	}
@@ -2930,12 +2937,6 @@ bool FEditorViewportClient::InputKey(FViewport* InViewport, int32 ControllerId,
 		bHandled |= true;
 	}
 
-	// Give the interactive tools context a shot at input if it hasn't been handled already.
-	if (!bHandled)
-	{
-		ModeTools->GetInteractiveToolsContext()->InputKey(this, Viewport, Key, Event);
-	}
-
 	//apply the visibility and set the cursor positions
 	ApplyRequiredCursorVisibility( true );
 	return bHandled;
diff --git a/Engine/Source/Editor/UnrealEd/Private/LevelEditorViewport.cpp b/Engine/Source/Editor/UnrealEd/Private/LevelEditorViewport.cpp
--- a/Engine/Source/Editor/UnrealEd/Private/LevelEditorViewport.cpp
+++ b/Engine/Source/Editor/UnrealEd/Private/LevelEditorViewport.cpp
@@ -3024,11 +3024,6 @@ void FLevelEditorViewportClient::TrackingStarted( const FInputEventState& InInpu
 			// Suspend actor/component modification during each delta step to avoid recording unnecessary overhead into the transaction buffer
 			GEditor->DisableDeltaModification(true);
 		}
-		else if (!IsFlightCameraActive() && !ShouldOrbitCamera())
-		{
-			// Give the interactive tools context a shot at tracking.
-			ModeTools->GetInteractiveToolsContext()->StartTracking(this, Viewport);
-		}
 	}
 }
 
@@ -3112,10 +3107,6 @@ void FLevelEditorViewportClient::TrackingStopped()
 		// Restore actor/component delta modification
 		GEditor->DisableDeltaModification(false);
 	}
-	else
-	{
-		ModeTools->GetInteractiveToolsContext()->EndTracking(this, Viewport);
-	}
 
 	ModeTools->ActorMoveNotify();
 
diff --git a/Engine/Source/Editor/UnrealEd/Public/EditorModeManager.h b/Engine/Source/Editor/UnrealEd/Public/EditorModeManager.h
--- a/Engine/Source/Editor/UnrealEd/Public/EditorModeManager.h
+++ b/Engine/Source/Editor/UnrealEd/Public/EditorModeManager.h
@@ -249,8 +249,14 @@ public:
 	/** Notifies all active modes of all captured mouse movement */	
 	bool ProcessCapturedMouseMoves( FEditorViewportClient* InViewportClient, FViewport* InViewport, const TArrayView<FIntPoint>& CapturedMouseMoves );
 
-	/** Notifies all active modes of keyboard input */
-	bool InputKey( FEditorViewportClient* InViewportClient, FViewport* Viewport, FKey Key, EInputEvent Event);
+	/** 
+	 * Notifies all active modes of keyboard input 
+	 * @param bRouteToToolsContext If true, routes to the tools context and its input router before routing
+	 *  to modes (and does not route to modes if tools context handles it). We currently need the ability to
+	 *  set this to false due to some behaviors being routed in different conditions to legacy modes compared
+	 *  to the input router (see its use in EditorViewportClient.cpp).
+	 */
+	bool InputKey( FEditorViewportClient* InViewportClient, FViewport* Viewport, FKey Key, EInputEvent Event, bool bRouteToToolsContext = true);
 
 	/** Notifies all active modes of axis movement */
 	bool InputAxis( FEditorViewportClient* InViewportClient, FViewport* Viewport, int32 ControllerId, FKey Key, float Delta, float DeltaTime);
diff --git a/Engine/Source/Runtime/Experimental/InteractiveToolsFramework/Private/BaseBehaviors/KeyAsModifierInputBehavior.cpp b/Engine/Source/Runtime/Experimental/InteractiveToolsFramework/Private/BaseBehaviors/KeyAsModifierInputBehavior.cpp
--- a/Engine/Source/Runtime/Experimental/InteractiveToolsFramework/Private/BaseBehaviors/KeyAsModifierInputBehavior.cpp
+++ b/Engine/Source/Runtime/Experimental/InteractiveToolsFramework/Private/BaseBehaviors/KeyAsModifierInputBehavior.cpp
@@ -8,51 +8,34 @@ UKeyAsModifierInputBehavior::UKeyAsModifierInputBehavior()
 
 void UKeyAsModifierInputBehavior::Initialize(IModifierToggleBehaviorTarget* TargetIn, int ModifierID, const FKey& ModifierKeyIn)
 {
-	this->Target = TargetIn;
-	ModifierKey = ModifierKeyIn;
-	FKey TempKey = ModifierKeyIn;
-	Modifiers.RegisterModifier(ModifierID, [TempKey](const FInputDeviceState& Input)
+	Initialize(TargetIn, ModifierID, [ModifierKeyIn](const FInputDeviceState& Input)
 	{
-		return (TempKey == EKeys::AnyKey || Input.Keyboard.ActiveKey.Button == TempKey) && Input.Keyboard.ActiveKey.bPressed;
+		return (ModifierKeyIn == EKeys::AnyKey || Input.Keyboard.ActiveKey.Button == ModifierKeyIn) && Input.Keyboard.ActiveKey.bPressed;
 	});
 }
 
+void UKeyAsModifierInputBehavior::Initialize(IModifierToggleBehaviorTarget* TargetIn, int ModifierID, TFunction<bool(const FInputDeviceState&)> ModifierCheckFunc)
+{
+	this->Target = TargetIn;
+	Modifiers.RegisterModifier(ModifierID, ModifierCheckFunc);
+}
+
 FInputCaptureRequest UKeyAsModifierInputBehavior::WantsCapture(const FInputDeviceState& Input)
 {
-	if ((ModifierCheckFunc == nullptr || ModifierCheckFunc(Input)))
-	{
-		if ((ModifierKey == EKeys::AnyKey || Input.Keyboard.ActiveKey.Button == ModifierKey) && Input.Keyboard.ActiveKey.bPressed)
-		{
-			return FInputCaptureRequest::Begin(this, EInputCaptureSide::Any);
-		}
-	}
+	Modifiers.UpdateModifiers(Input, Target);
 	return FInputCaptureRequest::Ignore();
 }
 
 
 FInputCaptureUpdate UKeyAsModifierInputBehavior::BeginCapture(const FInputDeviceState& Input, EInputCaptureSide Side)
 {
-	PressedButton = Input.Keyboard.ActiveKey.Button;
-	Modifiers.UpdateModifiers(Input, Target);
-	return FInputCaptureUpdate::Begin(this, EInputCaptureSide::Any);
+	return FInputCaptureUpdate::End();
 }
 
 
 FInputCaptureUpdate UKeyAsModifierInputBehavior::UpdateCapture(const FInputDeviceState& Input, const FInputCaptureData& Data)
 {
-	if (Input.Keyboard.ActiveKey.Button != PressedButton)
-	{
-		return FInputCaptureUpdate::Continue();
-	}
-
-	Modifiers.UpdateModifiers(Input, Target);
-
-	if (Input.Keyboard.ActiveKey.bReleased)
-	{
-		return FInputCaptureUpdate::End();
-	}
-
-	return FInputCaptureUpdate::Continue();
+	return FInputCaptureUpdate::End();
 }
 
 
diff --git a/Engine/Source/Runtime/Experimental/InteractiveToolsFramework/Private/BaseBehaviors/SingleKeyCaptureBehavior.cpp b/Engine/Source/Runtime/Experimental/InteractiveToolsFramework/Private/BaseBehaviors/SingleKeyCaptureBehavior.cpp
new file mode 100644
--- /dev/null
+++ b/Engine/Source/Runtime/Experimental/InteractiveToolsFramework/Private/BaseBehaviors/SingleKeyCaptureBehavior.cpp
@@ -0,0 +1,62 @@
+// Copyright Epic Games, Inc. All Rights Reserved.
+
+#include "BaseBehaviors/SingleKeyCaptureBehavior.h"
+
+USingleKeyCaptureBehavior::USingleKeyCaptureBehavior()
+{
+}
+
+void USingleKeyCaptureBehavior::Initialize(IModifierToggleBehaviorTarget* TargetIn, int ModifierID, const FKey& ModifierKeyIn)
+{
+	this->Target = TargetIn;
+	ModifierKey = ModifierKeyIn;
+	FKey TempKey = ModifierKeyIn;
+	Modifiers.RegisterModifier(ModifierID, [TempKey](const FInputDeviceState& Input)
+		{
+			return (TempKey == EKeys::AnyKey || Input.Keyboard.ActiveKey.Button == TempKey) && Input.Keyboard.ActiveKey.bPressed;
+		});
+}
+
+FInputCaptureRequest USingleKeyCaptureBehavior::WantsCapture(const FInputDeviceState& Input)
+{
+	if ((ModifierCheckFunc == nullptr || ModifierCheckFunc(Input)))
+	{
+		if ((ModifierKey == EKeys::AnyKey || Input.Keyboard.ActiveKey.Button == ModifierKey) && Input.Keyboard.ActiveKey.bPressed)
+		{
+			return FInputCaptureRequest::Begin(this, EInputCaptureSide::Any);
+		}
+	}
+	return FInputCaptureRequest::Ignore();
+}
+
+
+FInputCaptureUpdate USingleKeyCaptureBehavior::BeginCapture(const FInputDeviceState& Input, EInputCaptureSide Side)
+{
+	PressedButton = Input.Keyboard.ActiveKey.Button;
+	Modifiers.UpdateModifiers(Input, Target);
+	return FInputCaptureUpdate::Begin(this, EInputCaptureSide::Any);
+}
+
+
+FInputCaptureUpdate USingleKeyCaptureBehavior::UpdateCapture(const FInputDeviceState& Input, const FInputCaptureData& Data)
+{
+	if (Input.Keyboard.ActiveKey.Button != PressedButton)
+	{
+		return FInputCaptureUpdate::Continue();
+	}
+
+	Modifiers.UpdateModifiers(Input, Target);
+
+	if (Input.Keyboard.ActiveKey.bReleased)
+	{
+		return FInputCaptureUpdate::End();
+	}
+
+	return FInputCaptureUpdate::Continue();
+}
+
+
+void USingleKeyCaptureBehavior::ForceEndCapture(const FInputCaptureData& Data)
+{
+}
+
diff --git a/Engine/Source/Runtime/Experimental/InteractiveToolsFramework/Public/BaseBehaviors/KeyAsModifierInputBehavior.h b/Engine/Source/Runtime/Experimental/InteractiveToolsFramework/Public/BaseBehaviors/KeyAsModifierInputBehavior.h
--- a/Engine/Source/Runtime/Experimental/InteractiveToolsFramework/Public/BaseBehaviors/KeyAsModifierInputBehavior.h
+++ b/Engine/Source/Runtime/Experimental/InteractiveToolsFramework/Public/BaseBehaviors/KeyAsModifierInputBehavior.h
@@ -12,7 +12,10 @@
 
 /**
  * UKeyAsModifierInputBehavior converts a specific key press/release into
- * a "Modifier" toggle via the IModifierToggleBehaviorTarget interface
+ * a "Modifier" toggle via the IModifierToggleBehaviorTarget interface. It does
+ * not capture the key press; rather, it updates the modifier its WantsCapture call.
+ * This means that the modifier won't be updated if another behavior captures the
+ * keyboard.
  */
 UCLASS()
 class INTERACTIVETOOLSFRAMEWORK_API UKeyAsModifierInputBehavior : public UInputBehavior
@@ -36,13 +39,14 @@ public:
 	virtual void Initialize(IModifierToggleBehaviorTarget* Target, int ModifierID, const FKey& ModifierKey);
 
 	/**
-	 * WantsCapture() will only return capture request if this function returns true (or is null)
-	 * Intended to be used for alt/ctrl/cmd/shift modifiers on the main ModifierKey
+	 * Initialize this behavior with the given Target
+	 * @param Target implementor of modifier-toggle behavior
+	 * @param ModifierID integer ID that identifiers the modifier toggle
+	 * @param ModifierCheckFunc the function that will be checked to set the the modifier.
 	 */
-	TFunction<bool(const FInputDeviceState&)> ModifierCheckFunc = nullptr;
+	void Initialize(IModifierToggleBehaviorTarget* TargetIn, int ModifierID, TFunction<bool(const FInputDeviceState&)> ModifierCheckFunc);
 
 	// UInputBehavior implementation
-
 	virtual FInputCaptureRequest WantsCapture(const FInputDeviceState& Input) override;
 	virtual FInputCaptureUpdate BeginCapture(const FInputDeviceState& Input, EInputCaptureSide Side) override;
 	virtual FInputCaptureUpdate UpdateCapture(const FInputDeviceState& Input, const FInputCaptureData& Data) override;
@@ -53,12 +57,6 @@ protected:
 	/** Modifier Target object */
 	IModifierToggleBehaviorTarget* Target;
 
-	/** Key that is used as modifier */
-	FKey ModifierKey;
-
 	/** Modifier set for this behavior, internally initialized with check on ModifierKey */
 	FInputBehaviorModifierStates Modifiers;
-
-	/** The key that was pressed to activate capture */
-	FKey PressedButton;
 };
diff --git a/Engine/Source/Runtime/Experimental/InteractiveToolsFramework/Public/BaseBehaviors/SingleKeyCaptureBehavior.h b/Engine/Source/Runtime/Experimental/InteractiveToolsFramework/Public/BaseBehaviors/SingleKeyCaptureBehavior.h
new file mode 100644
--- /dev/null
+++ b/Engine/Source/Runtime/Experimental/InteractiveToolsFramework/Public/BaseBehaviors/SingleKeyCaptureBehavior.h
@@ -0,0 +1,68 @@
+// Copyright Epic Games, Inc. All Rights Reserved.
+
+#pragma once
+
+#include "CoreMinimal.h"
+#include "InputBehavior.h"
+#include "InteractiveTool.h"
+#include "BaseBehaviors/BehaviorTargetInterfaces.h"
+#include "BaseBehaviors/InputBehaviorModifierStates.h"
+#include "SingleKeyCaptureBehavior.generated.h"
+
+
+/**
+ * USingleKeyCaptureBehavior captures a key press and routes it to target via
+ * the IModifierToggleBehaviorTarget interface. If you want similar behavior
+ * without actually capturing the key, you should use UKeyAsModifierInputBehavior.
+ */
+UCLASS()
+class INTERACTIVETOOLSFRAMEWORK_API USingleKeyCaptureBehavior : public UInputBehavior
+{
+	GENERATED_BODY()
+
+public:
+	USingleKeyCaptureBehavior();
+
+	virtual EInputDevices GetSupportedDevices() override
+	{
+		return EInputDevices::Keyboard;
+	}
+
+	/**
+	 * Initialize this behavior with the given Target. Note that though it interacts through the
+	 * IModifierToggleBehaviorTarget interface, the modifier here captures keyboard input, unlike
+	 * the modifiers in UKeyAsMOdifierInputBehavior and other places.
+	 * 
+	 * @param Target implementor of IModifierToggleBehaviorTarget
+	 * @param ModifierID integer ID that identifiers the modifier toggle
+	 * @param ModifierKey the key that will be used as the modifier toggle
+	 */
+	virtual void Initialize(IModifierToggleBehaviorTarget* Target, int ModifierID, const FKey& ModifierKey);
+
+	/**
+	 * WantsCapture() will only return capture request if this function returns true (or is null)
+	 * Intended to be used for alt/ctrl/cmd/shift modifiers on the main key
+	 */
+	TFunction<bool(const FInputDeviceState&)> ModifierCheckFunc = nullptr;
+
+	// UInputBehavior implementation
+
+	virtual FInputCaptureRequest WantsCapture(const FInputDeviceState& Input) override;
+	virtual FInputCaptureUpdate BeginCapture(const FInputDeviceState& Input, EInputCaptureSide Side) override;
+	virtual FInputCaptureUpdate UpdateCapture(const FInputDeviceState& Input, const FInputCaptureData& Data) override;
+	virtual void ForceEndCapture(const FInputCaptureData& Data) override;
+
+
+protected:
+	/** Modifier Target object */
+	IModifierToggleBehaviorTarget* Target;
+
+	/** Key that is used as modifier */
+	FKey ModifierKey;
+
+	/** Modifier set for this behavior, internally initialized with check on ModifierKey */
+	FInputBehaviorModifierStates Modifiers;
+
+	/** The key that was pressed to activate capture */
+	FKey PressedButton;
+};
