Adding a missing special case back into EdModeInteractiveToolsContext that disallows ITF input router from capturing alt+mouse clicks.

#rnx

diff --git a/Engine/Source/Editor/Experimental/EditorInteractiveToolsFramework/Private/EdModeInteractiveToolsContext.cpp b/Engine/Source/Editor/Experimental/EditorInteractiveToolsFramework/Private/EdModeInteractiveToolsContext.cpp
--- a/Engine/Source/Editor/Experimental/EditorInteractiveToolsFramework/Private/EdModeInteractiveToolsContext.cpp
+++ b/Engine/Source/Editor/Experimental/EditorInteractiveToolsFramework/Private/EdModeInteractiveToolsContext.cpp
@@ -935,6 +935,16 @@ bool UEdModeInteractiveToolsContext::InputKey(FEditorViewportClient* ViewportCli
 			bool bIsRightMouse = (Key == EKeys::RightMouseButton);
 			if (bIsLeftMouse || bIsMiddleMouse || bIsRightMouse)
 			{
+				// Currently, we don't capture mouse clicks that start with Alt being down because we want
+				// Alt camera manipulation to have priority over tools. So, we let those inputs pass on up
+				// to wherever they get handled.
+				// Someday these kinds of prioritizations will be handled by having camera manipulation be
+				// in a common input router so that behavior priorities can determine the ordering.
+				if (ViewportClient->IsAltPressed() && InputRouter->HasActiveMouseCapture() == false)
+				{
+					return false;
+				}
+
 				FInputDeviceState InputState = CurrentMouseState;
 				InputState.InputDevice = EInputDevices::Mouse;
 				InputState.SetModifierKeyStates(
