ModelingMode: stop logging all tool messages as Warnings in the log
#rnx

diff --git a/Engine/Source/Editor/Experimental/EditorInteractiveToolsFramework/Private/EdModeInteractiveToolsContext.cpp b/Engine/Source/Editor/Experimental/EditorInteractiveToolsFramework/Private/EdModeInteractiveToolsContext.cpp
--- a/Engine/Source/Editor/Experimental/EditorInteractiveToolsFramework/Private/EdModeInteractiveToolsContext.cpp
+++ b/Engine/Source/Editor/Experimental/EditorInteractiveToolsFramework/Private/EdModeInteractiveToolsContext.cpp
@@ -443,11 +443,11 @@ public:
 
 	virtual void DisplayMessage(const FText& Message, EToolMessageLevel Level) override
 	{
-		if (Level == EToolMessageLevel::UserNotification)
+		if (Level == EToolMessageLevel::UserNotification || Level == EToolMessageLevel::UserMessage)
 		{
 			ToolsContext->PostToolNotificationMessage(Message);
 		}
-		if (Level == EToolMessageLevel::UserWarning)
+		else if (Level == EToolMessageLevel::UserWarning || Level == EToolMessageLevel::UserError)
 		{
 			ToolsContext->PostToolWarningMessage(Message);
 		}
