Adding post tool build and post set up notifications to the tool manager.

diff --git a/Engine/Source/Editor/Experimental/EditorInteractiveToolsFramework/Private/EdModeInteractiveToolsContext.cpp b/Engine/Source/Editor/Experimental/EditorInteractiveToolsFramework/Private/EdModeInteractiveToolsContext.cpp
--- a/Engine/Source/Editor/Experimental/EditorInteractiveToolsFramework/Private/EdModeInteractiveToolsContext.cpp
+++ b/Engine/Source/Editor/Experimental/EditorInteractiveToolsFramework/Private/EdModeInteractiveToolsContext.cpp
@@ -562,10 +562,8 @@ void UEdModeInteractiveToolsContext::Initialize(IToolsContextQueriesAPI* Queries
 		}
 	});
 
-	ToolManager->OnToolEnded.AddLambda([this](UInteractiveToolManager*, UInteractiveTool*)
-	{
-		RestoreEditorState();
-	});
+	ToolManager->OnToolEnded.AddUObject(this, &UEdModeInteractiveToolsContext::OnToolEnded);
+	ToolManager->OnToolPostBuild.AddUObject(this, &UEdModeInteractiveToolsContext::OnToolPostBuild);
 
 	// if viewport clients change we will discard our overrides as we aren't sure what happened
 	ViewportClientListChangedHandle = GEditor->OnViewportClientListChanged().AddLambda([this]()
@@ -578,6 +576,12 @@ void UEdModeInteractiveToolsContext::Initialize(IToolsContextQueriesAPI* Queries
 
 void UEdModeInteractiveToolsContext::Shutdown()
 {
+	if (ToolManager)
+	{
+		ToolManager->OnToolPostBuild.RemoveAll(this);
+		ToolManager->OnToolEnded.RemoveAll(this);
+	}
+
 	if (FLevelEditorModule* LevelEditor = FModuleManager::GetModulePtr<FLevelEditorModule>("LevelEditor"))
 	{
 		LevelEditor->OnMapChanged().Remove(WorldTearDownDelegateHandle);
@@ -1236,3 +1240,13 @@ void UEdModeInteractiveToolsContext::RestoreEditorState()
 		}
 	}
 }
+
+void UEdModeInteractiveToolsContext::OnToolEnded(UInteractiveToolManager* InToolManager, UInteractiveTool* InEndedTool)
+{
+	RestoreEditorState();
+}
+
+void UEdModeInteractiveToolsContext::OnToolPostBuild(UInteractiveToolManager* InToolManager, EToolSide InSide, UInteractiveTool* InBuiltTool, UInteractiveToolBuilder* InToolBuilder, const FToolBuilderState& ToolState)
+{
+	// todo: Add any shared tool targets for the mode toolkit
+}
diff --git a/Engine/Source/Editor/Experimental/EditorInteractiveToolsFramework/Public/EdModeInteractiveToolsContext.h b/Engine/Source/Editor/Experimental/EditorInteractiveToolsFramework/Public/EdModeInteractiveToolsContext.h
--- a/Engine/Source/Editor/Experimental/EditorInteractiveToolsFramework/Public/EdModeInteractiveToolsContext.h
+++ b/Engine/Source/Editor/Experimental/EditorInteractiveToolsFramework/Public/EdModeInteractiveToolsContext.h
@@ -58,8 +58,6 @@ public:
 	// UObject Interface
 	virtual UWorld* GetWorld() const override;
 
-	// call these from your UEdMode functions of the same name
-
 	void Tick(FEditorViewportClient* ViewportClient, float DeltaTime);
 	void Render(const FSceneView* View, FViewport* Viewport, FPrimitiveDrawInterface* PDI);
 	void DrawHUD(FViewportClient* ViewportClient,FViewport* Viewport,const FSceneView* View, FCanvas* Canvas);
@@ -141,6 +139,9 @@ protected:
 	void SetEditorStateForTool();
 	void RestoreEditorState();
 
+	void OnToolEnded(UInteractiveToolManager* InToolManager, UInteractiveTool* InEndedTool);
+	void OnToolPostBuild(UInteractiveToolManager* InToolManager, EToolSide InSide, UInteractiveTool* InBuiltTool, UInteractiveToolBuilder* InToolBuilder, const FToolBuilderState& ToolState);
+
 	TOptional<FString> PendingToolToStart = {};
 	TOptional<EToolShutdownType> PendingToolShutdownType = {};
 
diff --git a/Engine/Source/Runtime/Experimental/InteractiveToolsFramework/Private/InteractiveToolManager.cpp b/Engine/Source/Runtime/Experimental/InteractiveToolsFramework/Private/InteractiveToolManager.cpp
--- a/Engine/Source/Runtime/Experimental/InteractiveToolsFramework/Private/InteractiveToolManager.cpp
+++ b/Engine/Source/Runtime/Experimental/InteractiveToolsFramework/Private/InteractiveToolManager.cpp
@@ -53,6 +53,15 @@ void UInteractiveToolManager::Shutdown()
 	bIsActive = false;
 }
 
+void UInteractiveToolManager::DoPostBuild(EToolSide Side, UInteractiveTool* InBuiltTool, UInteractiveToolBuilder* InToolBuilder, const FToolBuilderState& InBuilderState)
+{
+	OnToolPostBuild.Broadcast(this, Side, InBuiltTool, InToolBuilder, InBuilderState);
+}
+
+void UInteractiveToolManager::DoPostSetup(EToolSide Side, UInteractiveTool* InInteractiveTool)
+{
+	OnToolPostSetup.Broadcast(this, Side, InInteractiveTool);
+}
 
 
 void UInteractiveToolManager::RegisterToolType(const FString& Identifier, UInteractiveToolBuilder* Builder)
@@ -182,8 +191,10 @@ bool UInteractiveToolManager::ActivateToolInternal(EToolSide Side)
 		return false;
 	}
 	ActiveLeftToolName = ActiveLeftBuilderName;
+	DoPostBuild(Side, ActiveLeftTool, ActiveLeftBuilder, InputState);
 
 	ActiveLeftTool->Setup();
+	DoPostSetup(Side, ActiveLeftTool);
 
 	// register new active input behaviors
 	InputRouter->RegisterSource(ActiveLeftTool);
diff --git a/Engine/Source/Runtime/Experimental/InteractiveToolsFramework/Public/InteractiveToolManager.h b/Engine/Source/Runtime/Experimental/InteractiveToolsFramework/Public/InteractiveToolManager.h
--- a/Engine/Source/Runtime/Experimental/InteractiveToolsFramework/Public/InteractiveToolManager.h
+++ b/Engine/Source/Runtime/Experimental/InteractiveToolsFramework/Public/InteractiveToolManager.h
@@ -73,6 +73,12 @@ protected:
 	/** Shutdown the ToolManager. Called by UInteractiveToolsContext. */
 	virtual void Shutdown();
 
+	/** Called immediately after a tool is built. Broadcasts OnToolPostBuild. */
+	virtual void DoPostBuild(EToolSide Side, UInteractiveTool* InBuiltTool, UInteractiveToolBuilder* InToolBuilder, const FToolBuilderState& InBuilderState);
+	
+	/** Called immediately after a tool's Setup is called. Broadcasts OnToolPostSetup. */
+	virtual void DoPostSetup(EToolSide Side, UInteractiveTool* InInteractiveTool);
+
 public:
 
 
@@ -274,6 +280,12 @@ public:
 	DECLARE_MULTICAST_DELEGATE_TwoParams(FToolManagerToolStartedSignature, UInteractiveToolManager*, UInteractiveTool*);
 	FToolManagerToolStartedSignature OnToolStarted;
 
+	DECLARE_MULTICAST_DELEGATE_FiveParams(FToolManagerToolPostBuildSignature, UInteractiveToolManager*, EToolSide, UInteractiveTool*, UInteractiveToolBuilder*, const FToolBuilderState&);
+	FToolManagerToolPostBuildSignature OnToolPostBuild;
+
+	DECLARE_MULTICAST_DELEGATE_ThreeParams(FToolManagerToolPostSetupSignature, UInteractiveToolManager*, EToolSide, UInteractiveTool*);
+	FToolManagerToolPostSetupSignature OnToolPostSetup;
+
 	DECLARE_MULTICAST_DELEGATE_TwoParams(FToolManagerToolEndedSignature, UInteractiveToolManager*, UInteractiveTool*);
 	FToolManagerToolEndedSignature OnToolEnded;
 
