First pass at Level Editor TRS widget implementation using new gizmo architecture. Translate and Scale are implemented. Rotation, Screen Space and Planar Translation/Scale and Arcball Rotation still TBD.

diff --git a/Engine/Source/Editor/Experimental/EditorInteractiveToolsFramework/Private/EditorGizmos/EditorAxisSources.cpp b/Engine/Source/Editor/Experimental/EditorInteractiveToolsFramework/Private/EditorGizmos/EditorAxisSources.cpp
new file mode 100644
--- /dev/null
+++ b/Engine/Source/Editor/Experimental/EditorInteractiveToolsFramework/Private/EditorGizmos/EditorAxisSources.cpp
@@ -0,0 +1,36 @@
+// Copyright Epic Games, Inc. All Rights Reserved.
+
+#include "EditorGizmos/EditorAxisSources.h"
+#include "EditorModeManager.h"
+
+FVector UGizmoEditorAxisSource::GetOrigin() const
+{
+	// @todo get this from the UTransformProxy instead of global?
+	// @todo - per Brooke's comment, the UTransformProxy could possibly have
+	// a tool target and the toolkit host could be queried from the tool target.
+	// Also refer to UTypedElementViewportInteraction, the transform proxy 
+	// should work in conjunction with this for typed element support.
+	FEditorModeTools& EditorModeTools = GLevelEditorModeTools();
+	return EditorModeTools.GetWidgetLocation();
+}
+
+FVector UGizmoEditorAxisSource::GetDirection() const
+{
+	// @todo get this from the UTransformProxy instead of global?
+	FEditorModeTools& EditorModeTools = GLevelEditorModeTools();
+	FMatrix GizmoToWorldMatrix = EditorModeTools.GetCustomInputCoordinateSystem();
+
+	if (bLocalAxes)
+	{
+		EAxis::Type Axis = (AxisIndex == 0 ? EAxis::X : AxisIndex == 1 ? EAxis::Y : EAxis::Z);
+		return GizmoToWorldMatrix.GetUnitAxis(Axis);
+	} 
+	else
+	{
+		FVector Axis(0, 0, 0);
+		Axis[FMath::Clamp(AxisIndex, 0, 2)] = 1.0;
+		return Axis;
+	}
+}
+
+
diff --git a/Engine/Source/Editor/Experimental/EditorInteractiveToolsFramework/Private/EditorGizmos/EditorTransformGizmo.cpp b/Engine/Source/Editor/Experimental/EditorInteractiveToolsFramework/Private/EditorGizmos/EditorTransformGizmo.cpp
new file mode 100644
--- /dev/null
+++ b/Engine/Source/Editor/Experimental/EditorInteractiveToolsFramework/Private/EditorGizmos/EditorTransformGizmo.cpp
@@ -0,0 +1,781 @@
+// Copyright Epic Games, Inc. All Rights Reserved.
+
+#include "EditorGizmos/EditorTransformGizmo.h"
+#include "BaseGizmos/AxisPositionGizmo.h"
+#include "BaseGizmos/GizmoArrowObject.h"
+#include "BaseGizmos/GizmoBoxObject.h"
+#include "BaseGizmos/GizmoConeObject.h"
+#include "BaseGizmos/GizmoCylinderObject.h"
+#include "BaseGizmos/GizmoGroupObject.h"
+#include "BaseGizmos/TransformSources.h"
+#include "EditorGizmos/EditorAxisSources.h"
+#include "EditorGizmos/EditorParameterToTransformAdapters.h"
+#include "Engine/CollisionProfile.h"
+#include "Engine/World.h"
+#include "UnrealEngine.h"
+
+#define LOCTEXT_NAMESPACE "UEditorTransformGizmo"
+
+constexpr float UEditorTransformGizmo::AXIS_LENGTH;
+constexpr float UEditorTransformGizmo::AXIS_RADIUS;
+constexpr float UEditorTransformGizmo::AXIS_CONE_ANGLE;
+constexpr float UEditorTransformGizmo::AXIS_CONE_HEIGHT;
+constexpr float UEditorTransformGizmo::AXIS_CONE_HEAD_OFFSET;
+constexpr float UEditorTransformGizmo::AXIS_CUBE_SIZE;
+constexpr float UEditorTransformGizmo::AXIS_CUBE_HEAD_OFFSET;
+constexpr float UEditorTransformGizmo::TRANSLATE_ROTATE_AXIS_CIRCLE_RADIUS;
+constexpr float UEditorTransformGizmo::TWOD_AXIS_CIRCLE_RADIUS;
+constexpr float UEditorTransformGizmo::INNER_AXIS_CIRCLE_RADIUS;
+constexpr float UEditorTransformGizmo::OUTER_AXIS_CIRCLE_RADIUS;
+constexpr float UEditorTransformGizmo::ROTATION_TEXT_RADIUS;
+constexpr int32 UEditorTransformGizmo::AXIS_CIRCLE_SIDES;
+constexpr float UEditorTransformGizmo::ARCALL_RELATIVE_INNER_SIZE;
+constexpr float UEditorTransformGizmo::AXIS_LENGTH_SCALE;
+constexpr float UEditorTransformGizmo::AXIS_LENGTH_SCALE_OFFSET;
+
+constexpr FLinearColor UEditorTransformGizmo::AxisColorX;
+constexpr FLinearColor UEditorTransformGizmo::AxisColorY;
+constexpr FLinearColor UEditorTransformGizmo::AxisColorZ;
+constexpr FLinearColor UEditorTransformGizmo::ScreenAxisColor;
+constexpr FColor UEditorTransformGizmo::PlaneColorXY;
+constexpr FColor UEditorTransformGizmo::ArcBallColor;
+constexpr FColor UEditorTransformGizmo::ScreenSpaceColor;
+constexpr FColor UEditorTransformGizmo::CurrentColor;
+
+UInteractiveGizmo* UEditorTransformGizmoBuilder::BuildGizmo(const FToolBuilderState& SceneState) const
+{
+	UEditorTransformGizmo* NewGizmo = NewObject<UEditorTransformGizmo>(SceneState.GizmoManager);
+	NewGizmo->SetWorld(SceneState.World);
+	return NewGizmo;
+}
+
+void UEditorTransformGizmo::SetWorld(UWorld* WorldIn)
+{
+	this->World = WorldIn;
+}
+
+void UEditorTransformGizmo::SetElements(ETransformGizmoSubElements InEnableElements)
+{
+	EnableElements = InEnableElements;
+}
+
+void UEditorTransformGizmo::SetDisallowNegativeScaling(bool bDisallow)
+{
+	if (bDisallowNegativeScaling != bDisallow)
+	{
+		bDisallowNegativeScaling = bDisallow;
+		for (UInteractiveGizmo* SubGizmo : this->ActiveGizmos)
+		{
+			if (UAxisPositionGizmo* CastGizmo = Cast<UAxisPositionGizmo>(SubGizmo))
+			{
+				if (UGizmoAxisScaleParameterSource* ParamSource = Cast<UGizmoAxisScaleParameterSource>(CastGizmo->ParameterSource.GetObject()))
+				{
+					ParamSource->bClampToZero = bDisallow;
+				}
+			}
+			/* @todo
+			if (UPlanePositionGizmo* CastGizmo = Cast<UEditorPlanePositionGizmo>(SubGizmo))
+			{
+				if (UGizmoPlaneScaleParameterSource* ParamSource = Cast<UGizmoPlaneScaleParameterSource>(CastGizmo->ParameterSource.GetObject()))
+				{
+					ParamSource->bClampToZero = bDisallow;
+				}
+			}
+			*/
+		}
+	}
+}
+
+void UEditorTransformGizmo::SetIsNonUniformScaleAllowedFunction(TUniqueFunction<bool()>&& IsNonUniformScaleAllowedIn)
+{
+	IsNonUniformScaleAllowed = MoveTemp(IsNonUniformScaleAllowedIn);
+}
+
+
+void UEditorTransformGizmo::Setup()
+{
+	UInteractiveGizmo::Setup();
+
+	UMaterial* AxisMaterialBase = GEngine->ArrowMaterial;
+
+	AxisMaterialX = UMaterialInstanceDynamic::Create(AxisMaterialBase, NULL);
+	AxisMaterialX->SetVectorParameterValue("GizmoColor", AxisColorX);
+
+	AxisMaterialY = UMaterialInstanceDynamic::Create(AxisMaterialBase, NULL);
+	AxisMaterialY->SetVectorParameterValue("GizmoColor", AxisColorY);
+
+	AxisMaterialZ = UMaterialInstanceDynamic::Create(AxisMaterialBase, NULL);
+	AxisMaterialZ->SetVectorParameterValue("GizmoColor", AxisColorZ);
+
+	CurrentAxisMaterial = UMaterialInstanceDynamic::Create(AxisMaterialBase, NULL);
+	CurrentAxisMaterial->SetVectorParameterValue("GizmoColor", CurrentColor);
+
+	OpaquePlaneMaterialXY = UMaterialInstanceDynamic::Create(AxisMaterialBase, NULL);
+	OpaquePlaneMaterialXY->SetVectorParameterValue("GizmoColor", FLinearColor::White);
+
+	TransparentPlaneMaterialXY = (UMaterial*)StaticLoadObject(
+		UMaterial::StaticClass(), NULL,
+		TEXT("/Engine/EditorMaterials/WidgetVertexColorMaterial.WidgetVertexColorMaterial"), NULL, LOAD_None, NULL);
+
+	GridMaterial = (UMaterial*)StaticLoadObject(
+		UMaterial::StaticClass(), NULL,
+		TEXT("/Engine/EditorMaterials/WidgetGridVertexColorMaterial_Ma.WidgetGridVertexColorMaterial_Ma"), NULL,
+		LOAD_None, NULL);
+	if (!GridMaterial)
+	{
+		GridMaterial = TransparentPlaneMaterialXY;
+	}
+}
+
+
+void UEditorTransformGizmo::Shutdown()
+{
+	ClearActiveTarget();
+}
+
+
+void UEditorTransformGizmo::UpdateCameraAxisSource()
+{
+	FViewCameraState CameraState;
+	GetGizmoManager()->GetContextQueriesAPI()->GetCurrentViewState(CameraState);
+	if (CameraAxisSource != nullptr) // && GizmoActor != nullptr)
+	{
+		// @todo get this from the UTransformProxy instead of global?
+		FEditorModeTools& EditorModeTools = GLevelEditorModeTools();
+		CameraAxisSource->Origin = EditorModeTools.GetWidgetLocation();// GizmoActor->GetTransform().GetLocation();
+		CameraAxisSource->Direction = -CameraState.Forward();
+		CameraAxisSource->TangentX = CameraState.Right();
+		CameraAxisSource->TangentY = CameraState.Up();
+	}
+}
+
+
+void UEditorTransformGizmo::Tick(float DeltaTime)
+{	
+	if (bUseContextCoordinateSystem)
+	{
+		CurrentCoordinateSystem = GetGizmoManager()->GetContextQueriesAPI()->GetCurrentCoordinateSystem();
+	}
+	
+	check(CurrentCoordinateSystem == EToolContextCoordinateSystem::World || CurrentCoordinateSystem == EToolContextCoordinateSystem::Local)
+	bool bUseLocalAxes = (CurrentCoordinateSystem == EToolContextCoordinateSystem::Local);
+
+	if (AxisXSource != nullptr && AxisYSource != nullptr && AxisZSource != nullptr)
+	{
+		AxisXSource->bLocalAxes = bUseLocalAxes;
+		AxisYSource->bLocalAxes = bUseLocalAxes;
+		AxisZSource->bLocalAxes = bUseLocalAxes;
+	}
+
+	for (UGizmoBaseObject* Object : ActiveObjects)
+	{
+		Object->SetWorldLocalState(CurrentCoordinateSystem == EToolContextCoordinateSystem::World);
+
+		float Scale = GLevelEditorModeTools().GetWidgetScale();
+		Object->SetGizmoScale(Scale);
+	}
+
+	bool bShouldShowNonUniformScale = IsNonUniformScaleAllowed();
+
+	for (UGizmoBaseObject* Object : NonuniformScaleObjects)
+	{
+		Object->SetVisibility(bShouldShowNonUniformScale);
+	}
+
+	UpdateCameraAxisSource();
+}
+
+
+void UEditorTransformGizmo::SetActiveTarget(UTransformProxy* Target, IToolContextTransactionProvider* TransactionProvider)
+{
+	if (ActiveTarget != nullptr)
+	{
+		ClearActiveTarget();
+	}
+
+	ActiveTarget = Target;
+
+	// move gizmo to target location
+
+	FTransform TargetTransform = Target->GetTransform();
+	FTransform GizmoTransform = TargetTransform;
+
+	// @todo this needs to be queried and updated in the TransformProxy
+	float GizmoScale = GLevelEditorModeTools().GetWidgetScale();
+	FMatrix GizmoLocalToWorld = GizmoTransform.ToMatrixNoScale();
+	FTransform GizmoLocalToWorldTransform = GizmoTransform;
+	GizmoLocalToWorldTransform.SetScale3D(FVector(GizmoScale));
+
+	// create group object to which all active objects will be added
+	GizmoGroupObject = NewObject<UGizmoGroupObject>();
+
+	// root component provides local X/Y/Z axis, identified by AxisIndex
+	AxisXSource = UGizmoEditorAxisSource::Construct(0, true, this);
+	AxisYSource = UGizmoEditorAxisSource::Construct(1, true, this);
+	AxisZSource = UGizmoEditorAxisSource::Construct(2, true, this);
+
+	auto MakeArrowObjectFunc = [](FVector Axis, UMaterialInterface* Material, UMaterialInterface* CurrentMaterial, FMatrix LocalToWorld, float GizmoScale, FTransform LocalToWorldTransform)
+	{
+		UGizmoArrowObject* ArrowObject = NewObject<UGizmoArrowObject>();
+		ArrowObject->CylinderObject->Direction = Axis;
+		ArrowObject->ConeObject->Direction = -Axis;
+		ArrowObject->SetMaterial(Material);
+		ArrowObject->SetCurrentMaterial(CurrentMaterial);
+		ArrowObject->SetGizmoScale(GizmoScale);
+		ArrowObject->SetLocalToWorldTransform(LocalToWorldTransform);
+		return ArrowObject;
+	};
+
+	UGizmoArrowObject* ArrowXObject = MakeArrowObjectFunc(FVector(1.0f, 0.0f, 0.0f), AxisMaterialX, CurrentAxisMaterial, GizmoLocalToWorld, GizmoScale, GizmoLocalToWorldTransform);
+	UGizmoArrowObject* ArrowYObject = MakeArrowObjectFunc(FVector(0.0f, 1.0f, 0.0f), AxisMaterialY, CurrentAxisMaterial, GizmoLocalToWorld, GizmoScale, GizmoLocalToWorldTransform);
+	UGizmoArrowObject* ArrowZObject = MakeArrowObjectFunc(FVector(0.0f, 0.0f, 1.0f), AxisMaterialZ, CurrentAxisMaterial, GizmoLocalToWorld, GizmoScale, GizmoLocalToWorldTransform);
+
+	UGizmoScaledAndUnscaledTransformSources* TransformSource = UGizmoScaledAndUnscaledTransformSources::Construct(
+		UGizmoTransformProxyTransformSource::Construct(ActiveTarget, this),
+		UGizmoObjectWorldTransformSource::Construct(GizmoGroupObject, this));
+	// This state target emits an explicit FChange that moves the GizmoActor root component during undo/redo.
+	// It also opens/closes the Transaction that saves/restores the target object locations.
+	if (TransactionProvider == nullptr)
+	{
+		TransactionProvider = GetGizmoManager();
+	}
+	StateTarget = UGizmoObjectTransformChangeStateTarget::Construct(GizmoGroupObject,
+		LOCTEXT("UEditorTransformGizmoTransaction", "Transform"), TransactionProvider, this);
+	StateTarget->DependentChangeSources.Add(MakeUnique<FTransformProxyChangeSource>(Target));
+
+	CameraAxisSource = NewObject<UGizmoConstantFrameAxisSource>(this);
+
+	FVector TargetWorldOrigin = TargetTransform.GetLocation();
+
+	if ((EnableElements & ETransformGizmoSubElements::TranslateAxisX) != ETransformGizmoSubElements::None)
+	{
+		AddAxisTranslationGizmo(ArrowXObject, AxisXSource, TransformSource, StateTarget, EAxisList::X, AxisColorX);
+		ActiveObjects.Add(ArrowXObject);
+		GizmoGroupObject->Add(ArrowXObject);
+	}
+	if ((EnableElements & ETransformGizmoSubElements::TranslateAxisY) != ETransformGizmoSubElements::None)
+	{
+		AddAxisTranslationGizmo(ArrowYObject, AxisYSource, TransformSource, StateTarget, EAxisList::Y, AxisColorY);
+		ActiveObjects.Add(ArrowYObject);
+		GizmoGroupObject->Add(ArrowYObject);
+	}
+	if ((EnableElements & ETransformGizmoSubElements::TranslateAxisZ) != ETransformGizmoSubElements::None)
+	{
+		AddAxisTranslationGizmo(ArrowZObject,  AxisZSource, TransformSource, StateTarget, EAxisList::Z, AxisColorZ);
+		ActiveObjects.Add(ArrowZObject);
+		GizmoGroupObject->Add(ArrowZObject);
+	}
+
+/*
+	// @todo: add plane translation
+	if (GizmoActor->TranslateYZ != nullptr)
+	{
+		AddPlaneTranslationGizmo(GizmoActor->TranslateYZ, GizmoComponent, AxisXSource, TransformSource, StateTarget);
+		ActiveComponents.Add(GizmoActor->TranslateYZ);
+	}
+	if (GizmoActor->TranslateXZ != nullptr)
+	{
+		AddPlaneTranslationGizmo(GizmoActor->TranslateXZ, GizmoComponent, AxisYSource, TransformSource, StateTarget);
+		ActiveComponents.Add(GizmoActor->TranslateXZ);
+	}
+	if (GizmoActor->TranslateXY != nullptr)
+	{
+		AddPlaneTranslationGizmo(GizmoActor->TranslateXY, GizmoComponent, AxisZSource, TransformSource, StateTarget);
+		ActiveComponents.Add(GizmoActor->TranslateXY);
+	}
+
+	// @todo: finish rotation implementation
+	if (GizmoActor->RotateX != nullptr)
+	{
+		AddAxisRotationGizmo(GizmoActor->RotateX, GizmoComponent, AxisXSource, SpaceSource, TransformSource, StateTarget, EAxisList::X, AxisColorX);
+		ActiveComponents.Add(GizmoActor->RotateX);
+	}
+	if (GizmoActor->RotateY != nullptr)
+	{
+		AddAxisRotationGizmo(GizmoActor->RotateY, GizmoComponent, AxisYSource, SpaceSource, TransformSource, StateTarget, EAxisList::Y, AxisColorY);
+		ActiveComponents.Add(GizmoActor->RotateY);
+	}
+	if (GizmoActor->RotateZ != nullptr)
+	{
+		AddAxisRotationGizmo(GizmoActor->RotateZ, GizmoComponent, AxisZSource, SpaceSource, TransformSource, StateTarget, EAxisList::Z, AxisColorZ);
+		ActiveComponents.Add(GizmoActor->RotateZ);
+	}
+*/
+
+	// Create objects for scale gizmo
+	UGizmoArrowObject* ScaleArrowXObject = MakeArrowObjectFunc(FVector(1.0f, 0.0f, 0.0f), AxisMaterialX, CurrentAxisMaterial, GizmoLocalToWorld, GizmoScale, GizmoLocalToWorldTransform);
+	UGizmoArrowObject* ScaleArrowYObject = MakeArrowObjectFunc(FVector(0.0f, 1.0f, 0.0f), AxisMaterialY, CurrentAxisMaterial, GizmoLocalToWorld, GizmoScale, GizmoLocalToWorldTransform);
+	UGizmoArrowObject* ScaleArrowZObject = MakeArrowObjectFunc(FVector(0.0f, 0.0f, 1.0f), AxisMaterialZ, CurrentAxisMaterial, GizmoLocalToWorld, GizmoScale, GizmoLocalToWorldTransform);
+
+	// only need these if scaling enabled. Essentially these are just the unit axes, regardless
+	// of what 3D axis is in use, we will tell the ParameterSource-to-3D-Scale mapper to
+	// use the coordinate axes
+	UnitAxisXSource = UGizmoEditorAxisSource::Construct(0, false, this);
+	UnitAxisYSource = UGizmoEditorAxisSource::Construct(1, false, this);
+	UnitAxisZSource = UGizmoEditorAxisSource::Construct(2, false, this);
+
+/*
+	// @todo: add uniform scale handle
+	if (GizmoActor->UniformScale != nullptr)
+	{
+		AddUniformScaleGizmo(GizmoActor->UniformScale, GizmoComponent, CameraAxisSource, CameraAxisSource, TransformSource, StateTarget);
+		ActiveComponents.Add(GizmoActor->UniformScale);
+	}
+*/
+
+	if ((EnableElements & ETransformGizmoSubElements::ScaleAxisX) != ETransformGizmoSubElements::None)
+	{
+		AddAxisScaleGizmo(ScaleArrowXObject, AxisXSource, UnitAxisXSource, TransformSource, StateTarget, EAxisList::X, AxisColorX);
+		ActiveObjects.Add(ScaleArrowXObject);
+		NonuniformScaleObjects.Add(ScaleArrowXObject);
+		GizmoGroupObject->Add(ScaleArrowXObject);
+	}
+	if ((EnableElements & ETransformGizmoSubElements::ScaleAxisY) != ETransformGizmoSubElements::None)
+	{
+		AddAxisScaleGizmo(ScaleArrowYObject, AxisYSource, UnitAxisYSource, TransformSource, StateTarget, EAxisList::Y, AxisColorY);
+		ActiveObjects.Add(ScaleArrowYObject);
+		NonuniformScaleObjects.Add(ScaleArrowYObject);
+		GizmoGroupObject->Add(ScaleArrowYObject);
+	}
+	if ((EnableElements & ETransformGizmoSubElements::ScaleAxisZ) != ETransformGizmoSubElements::None)
+	{
+		AddAxisScaleGizmo(ScaleArrowZObject, AxisZSource, UnitAxisZSource, TransformSource, StateTarget, EAxisList::Z, AxisColorZ);
+		ActiveObjects.Add(ScaleArrowZObject);
+		NonuniformScaleObjects.Add(ScaleArrowZObject);
+		GizmoGroupObject->Add(ScaleArrowZObject);
+	}
+
+/*
+	// @todo: add plane scale
+	if (GizmoActor->PlaneScaleYZ != nullptr)
+	{
+		AddPlaneScaleGizmo(GizmoActor->PlaneScaleYZ, GizmoComponent, AxisXSource, UnitAxisXSource, TransformSource, StateTarget);
+		ActiveComponents.Add(GizmoActor->PlaneScaleYZ);
+		NonuniformScaleComponents.Add(GizmoActor->PlaneScaleYZ);
+	}
+	if (GizmoActor->PlaneScaleXZ != nullptr)
+	{
+		AddPlaneScaleGizmo(GizmoActor->PlaneScaleXZ, GizmoComponent, AxisYSource, UnitAxisYSource, TransformSource, StateTarget);
+		ActiveComponents.Add(GizmoActor->PlaneScaleXZ);
+		NonuniformScaleComponents.Add(GizmoActor->PlaneScaleXZ);
+	}
+	if (GizmoActor->PlaneScaleXY != nullptr)
+	{
+		AddPlaneScaleGizmo(GizmoActor->PlaneScaleXY, GizmoComponent, AxisZSource, UnitAxisZSource, TransformSource, StateTarget);
+		ActiveComponents.Add(GizmoActor->PlaneScaleXY);
+		NonuniformScaleComponents.Add(GizmoActor->PlaneScaleXY);
+	}
+*/ 
+}
+
+
+void UEditorTransformGizmo::ReinitializeGizmoTransform(const FTransform& NewTransform)
+{
+	// @todo update gizmo objects here?
+
+	// The underlying proxy has an existing way to reinitialize its transform without callbacks.
+	TGuardValue<bool>(ActiveTarget->bSetPivotMode, true);
+	ActiveTarget->SetTransform(NewTransform);
+}
+
+
+void UEditorTransformGizmo::SetNewGizmoTransform(const FTransform& NewTransform)
+{
+	// @todo update gizmo objects here?
+
+	check(ActiveTarget != nullptr);
+
+	StateTarget->BeginUpdate();
+
+	ActiveTarget->SetTransform(NewTransform);
+
+	StateTarget->EndUpdate();
+}
+
+
+// @todo: This should either be named to "SetScale" or removed, since it can be done with ReinitializeGizmoTransform
+void UEditorTransformGizmo::SetNewChildScale(const FVector& NewChildScale)
+{
+	FTransform NewTransform = ActiveTarget->GetTransform();
+	NewTransform.SetScale3D(NewChildScale);
+
+	TGuardValue<bool>(ActiveTarget->bSetPivotMode, true);
+	ActiveTarget->SetTransform(NewTransform);
+}
+
+
+void UEditorTransformGizmo::SetVisibility(bool bVisible)
+{
+	for (UGizmoBaseObject* Object : ActiveObjects)
+	{
+		Object->SetVisibility(bVisible);
+	}
+}
+
+UInteractiveGizmo* UEditorTransformGizmo::AddAxisTranslationGizmo(
+	UGizmoArrowObject * InArrowObject,
+	IGizmoAxisSource * InAxisSource,
+	IGizmoTransformSource * InTransformSource,
+	IGizmoStateTarget * InStateTarget,
+	EAxisList::Type InAxisType,
+	const FLinearColor InAxisColor)
+{
+	check(InArrowObject);
+	check(InArrowObject->CylinderObject);
+	check(InArrowObject->ConeObject);
+
+	// create axis-position gizmo, axis-position parameter will drive translation
+	UAxisPositionGizmo* TranslateGizmo = Cast<UAxisPositionGizmo>(GetGizmoManager()->CreateGizmo(
+		UInteractiveGizmoManager::DefaultAxisPositionBuilderIdentifier));
+	check(TranslateGizmo);
+
+	InArrowObject->CylinderObject->Length = AXIS_LENGTH;
+	InArrowObject->CylinderObject->Radius = AXIS_RADIUS;
+	InArrowObject->bHasConeHead = true;
+
+	InArrowObject->ConeObject->Angle = FMath::DegreesToRadians(AXIS_CONE_ANGLE);
+	InArrowObject->ConeObject->Height = AXIS_CONE_HEIGHT;
+	InArrowObject->ConeObject->Offset = -(AXIS_LENGTH + AXIS_CONE_HEAD_OFFSET);
+
+	// arrow object provides the render capability
+	TranslateGizmo->GizmoObject = InArrowObject;
+
+	// axis source provides the translation axis
+	TranslateGizmo->AxisSource = Cast<UObject>(InAxisSource);
+
+	// parameter source maps axis-parameter-change to translation of TransformSource's transform
+	UGizmoEditorAxisTranslationParameterSource* ParamSource = UGizmoEditorAxisTranslationParameterSource::Construct(InAxisSource, InTransformSource, this);
+	ParamSource->AxisTranslationParameterSource->PositionConstraintFunction = [this](const FVector& Pos, FVector& Snapped) { return PositionSnapFunction(Pos, Snapped); };
+	TranslateGizmo->ParameterSource = ParamSource;
+
+	UGizmoObjectHitTarget* HitTarget = UGizmoObjectHitTarget::Construct(InArrowObject, this);
+
+	TranslateGizmo->HitTarget = HitTarget;
+	TranslateGizmo->StateTarget = Cast<UObject>(InStateTarget);
+
+	TranslateGizmo->ShouldUseCustomDestinationFunc = [this]() { return ShouldAlignDestination(); };
+	TranslateGizmo->CustomDestinationFunc =
+		[this](const UAxisPositionGizmo::FCustomDestinationParams& Params, FVector& OutputPoint) {
+		return DestinationAlignmentRayCaster(*Params.WorldRay, OutputPoint);
+	};
+
+	ActiveGizmos.Add(TranslateGizmo);
+
+	return TranslateGizmo;
+}
+
+UInteractiveGizmo* UEditorTransformGizmo::AddPlaneTranslationGizmo(
+	IGizmoAxisSource* InAxisSource,
+	IGizmoTransformSource* InTransformSource,
+	IGizmoStateTarget* InStateTarget)
+{
+	/* @todo
+	// create axis-position gizmo, axis-position parameter will drive translation
+	UEditorPlanePositionGizmo* TranslateGizmo = Cast<UEditorPlanePositionGizmo>(GetGizmoManager()->CreateGizmo(
+		UEditorInteractiveGizmoManager::DefaultEditorPlanePositionBuilderIdentifier));
+	check(TranslateGizmo);
+
+	// axis source provides the translation axis
+	TranslateGizmo->AxisSource = Cast<UObject>(InAxisSource);
+
+	// parameter source maps axis-parameter-change to translation of TransformSource's transform
+	UGizmoPlaneTranslationParameterSource* ParamSource = UGizmoPlaneTranslationParameterSource::Construct(InAxisSource, InTransformSource, this);
+	ParamSource->PositionConstraintFunction = [this](const FVector& Pos, FVector& Snapped) { return PositionSnapFunction(Pos, Snapped); };
+	TranslateGizmo->ParameterSource = ParamSource;
+
+	// sub-component provides hit target
+	UGizmoComponentHitTarget* HitTarget = UGizmoComponentHitTarget::Construct(InAxisComponent, this);
+	if (this->UpdateHoverFunctionComponent)
+	{
+		HitTarget->UpdateHoverFunction = [InAxisComponent, this](bool bHovering) { this->UpdateHoverFunctionComponent(InAxisComponent, bHovering); };
+	}
+	TranslateGizmo->HitTarget = HitTarget;
+
+	TranslateGizmo->StateTarget = Cast<UObject>(InStateTarget);
+
+	TranslateGizmo->ShouldUseCustomDestinationFunc = [this]() { return ShouldAlignDestination(); };
+	TranslateGizmo->CustomDestinationFunc =
+		[this](const UEditorPlanePositionGizmo::FCustomDestinationParams& Params, FVector& OutputPoint) {
+		return DestinationAlignmentRayCaster(*Params.WorldRay, OutputPoint);
+	};
+
+	ActiveGizmos.Add(TranslateGizmo);
+	return TranslateGizmo;
+	*/
+	return nullptr;
+}
+
+UInteractiveGizmo* UEditorTransformGizmo::AddAxisRotationGizmo(
+	IGizmoAxisSource* InAxisSource,
+	IGizmoTransformSource* InTransformSource,
+	IGizmoStateTarget* InStateTarget,
+	EAxisList::Type InAxisType,
+	const FLinearColor InAxisColor)
+{
+	/* @todo
+	// create axis-angle gizmo, angle will drive axis-rotation
+	UEditorAxisAngleGizmo* RotateGizmo = Cast<UEditorAxisAngleGizmo>(GetGizmoManager()->CreateGizmo(
+		UEditorInteractiveGizmoManager::DefaultEditorAxisAngleBuilderIdentifier));
+	check(RotateGizmo);
+
+	// axis source provides the translation axis
+	RotateGizmo->Initialize(InAxisType, InAxisColor, CurrentColor);
+
+	// axis source provides the rotation axis
+	RotateGizmo->AxisSource = Cast<UObject>(InAxisSource);
+
+	// space source provides coordinate space info specific to the gizmo
+	RotateGizmo->GizmoSpaceSource = Cast<UObject>(SpaceSource);
+
+	// parameter source maps angle-parameter-change to rotation of TransformSource's transform
+	UGizmoAxisRotationParameterSource* AngleSource = UGizmoAxisRotationParameterSource::Construct(InAxisSource, InTransformSource, this);
+	AngleSource->RotationConstraintFunction = [this](const FQuat& DeltaRotation){ return RotationSnapFunction(DeltaRotation); };
+	RotateGizmo->AngleSource = AngleSource;
+
+	// sub-component provides hit target
+	UGizmoComponentHitTarget* HitTarget = UGizmoComponentHitTarget::Construct(InAxisComponent, this);
+	if (this->UpdateHoverFunctionComponent)
+	{
+		HitTarget->UpdateHoverFunction = [InAxisComponent, this](bool bHovering) { this->UpdateHoverFunctionComponent(InAxisComponent, bHovering); };
+	}
+	RotateGizmo->HitTarget = HitTarget;
+
+	RotateGizmo->StateTarget = Cast<UObject>(StateTarget);
+
+	RotateGizmo->ShouldUseCustomDestinationFunc = [this]() { return ShouldAlignDestination(); };
+	RotateGizmo->CustomDestinationFunc =
+		[this](const UEditorAxisAngleGizmo::FCustomDestinationParams& Params, FVector& OutputPoint) {
+		return DestinationAlignmentRayCaster(*Params.WorldRay, OutputPoint);
+	};
+
+	ActiveGizmos.Add(RotateGizmo);
+
+	return RotateGizmo;
+	*/
+	return nullptr;
+}
+
+
+UInteractiveGizmo* UEditorTransformGizmo::AddAxisScaleGizmo(
+	UGizmoArrowObject* InArrowObject,
+	IGizmoAxisSource* InGizmoAxisSource, IGizmoAxisSource* InParameterAxisSource,
+	IGizmoTransformSource* InTransformSource,
+	IGizmoStateTarget* InStateTarget,
+	EAxisList::Type InAxisType,
+	const FLinearColor InAxisColor)
+{
+	check(InArrowObject);
+	check(InArrowObject->CylinderObject);
+	check(InArrowObject->BoxObject);
+
+	InArrowObject->CylinderObject->Length = AXIS_LENGTH_SCALE;
+	InArrowObject->CylinderObject->Radius = AXIS_RADIUS;
+	InArrowObject->CylinderObject->Offset = AXIS_LENGTH_SCALE_OFFSET;
+	InArrowObject->bHasConeHead = false;
+
+	InArrowObject->BoxObject->Dimensions = FVector(AXIS_CUBE_SIZE, AXIS_CUBE_SIZE, AXIS_CUBE_SIZE);
+	InArrowObject->BoxObject->Offset = AXIS_LENGTH_SCALE + AXIS_LENGTH_SCALE_OFFSET + AXIS_CUBE_HEAD_OFFSET;
+	if (InAxisType == EAxisList::X)
+	{
+		InArrowObject->BoxObject->UpDirection = FVector(1, 0, 0);
+		InArrowObject->BoxObject->SideDirection = FVector(0, 1, 0);
+	}
+	else if (InAxisType == EAxisList::Y)
+	{
+		InArrowObject->BoxObject->UpDirection = FVector(0, 1, 0);
+		InArrowObject->BoxObject->SideDirection = FVector(0, 0, 1);
+	}
+	else
+	{
+		InArrowObject->BoxObject->UpDirection = FVector(0, 0, 1);
+		InArrowObject->BoxObject->SideDirection = FVector(1, 0, 0);
+	}
+
+	// create axis-position gizmo, axis-position parameter will drive scale
+	UAxisPositionGizmo* ScaleGizmo = Cast<UAxisPositionGizmo>(GetGizmoManager()->CreateGizmo(
+		UInteractiveGizmoManager::DefaultAxisPositionBuilderIdentifier));
+	ScaleGizmo->bEnableSignedAxis = true;
+	check(ScaleGizmo);
+
+	// arrow object provides the render capability
+	ScaleGizmo->GizmoObject = InArrowObject;
+
+	// axis source provides the translation axis
+	ScaleGizmo->AxisSource = Cast<UObject>(InGizmoAxisSource);
+
+	// parameter source maps axis-parameter-change to translation of TransformSource's transform
+	UGizmoAxisScaleParameterSource* ParamSource = UGizmoAxisScaleParameterSource::Construct(InParameterAxisSource, InTransformSource, this);
+	ParamSource->bClampToZero = bDisallowNegativeScaling;
+	ScaleGizmo->ParameterSource = ParamSource;
+
+	UGizmoObjectHitTarget* HitTarget = UGizmoObjectHitTarget::Construct(InArrowObject, this);
+
+	ScaleGizmo->HitTarget = HitTarget;
+	ScaleGizmo->StateTarget = Cast<UObject>(StateTarget);
+
+	ActiveGizmos.Add(ScaleGizmo);
+	return ScaleGizmo;
+}
+
+
+UInteractiveGizmo* UEditorTransformGizmo::AddPlaneScaleGizmo(
+	IGizmoAxisSource* InGizmoAxisSource, IGizmoAxisSource* InParameterAxisSource,
+	IGizmoTransformSource* InTransformSource,
+	IGizmoStateTarget* InStateTarget)
+{
+	/* @todo
+	// create axis-position gizmo, axis-position parameter will drive scale
+	UEditorPlanePositionGizmo* ScaleGizmo = Cast<UEditorPlanePositionGizmo>(GetGizmoManager()->CreateGizmo(
+		UEditorInteractiveGizmoManager::DefaultEditorPlanePositionBuilderIdentifier));
+	ScaleGizmo->bEnableSignedAxis = true;
+	check(ScaleGizmo);
+
+	// axis source provides the translation axis
+	ScaleGizmo->AxisSource = Cast<UObject>(InGizmoAxisSource);
+
+	// parameter source maps axis-parameter-change to translation of TransformSource's transform
+	UGizmoPlaneScaleParameterSource* ParamSource = UGizmoPlaneScaleParameterSource::Construct(InParameterAxisSource, InTransformSource, this);
+	ParamSource->bClampToZero = bDisallowNegativeScaling;
+	ParamSource->bUseEqualScaling = true;
+	ScaleGizmo->ParameterSource = ParamSource;
+
+	// sub-component provides hit target
+	UGizmoComponentHitTarget* HitTarget = UGizmoComponentHitTarget::Construct(InAxisComponent, this);
+	if (this->UpdateHoverFunctionComponent)
+	{
+		HitTarget->UpdateHoverFunction = [InAxisComponent, this](bool bHovering) { this->UpdateHoverFunctionComponent(InAxisComponent, bHovering); };
+	}
+	ScaleGizmo->HitTarget = HitTarget;
+
+	ScaleGizmo->StateTarget = Cast<UObject>(StateTarget);
+
+	ActiveGizmos.Add(ScaleGizmo);
+	return ScaleGizmo;
+	*/
+	return nullptr;
+}
+
+
+UInteractiveGizmo* UEditorTransformGizmo::AddUniformScaleGizmo(
+	IGizmoAxisSource* InGizmoAxisSource, IGizmoAxisSource* InParameterAxisSource,
+	IGizmoTransformSource* InTransformSource,
+	IGizmoStateTarget* InStateTarget)
+{
+	/* @todo
+	// create plane-position gizmo, plane-position parameter will drive scale
+	UEditorPlanePositionGizmo* ScaleGizmo = Cast<UEditorPlanePositionGizmo>(GetGizmoManager()->CreateGizmo(
+		UEditorInteractiveGizmoManager::DefaultEditorPlanePositionBuilderIdentifier));
+	check(ScaleGizmo);
+
+	// axis source provides the translation plane
+	ScaleGizmo->AxisSource = Cast<UObject>(InGizmoAxisSource);
+
+	// parameter source maps axis-parameter-change to translation of TransformSource's transform
+	UGizmoUniformScaleParameterSource* ParamSource = UGizmoUniformScaleParameterSource::Construct(InParameterAxisSource, InTransformSource, this);
+	//ParamSource->PositionConstraintFunction = [this](const FVector& Pos, FVector& Snapped) { return PositionSnapFunction(Pos, Snapped); };
+	ScaleGizmo->ParameterSource = ParamSource;
+
+	// sub-component provides hit target
+	UGizmoComponentHitTarget* HitTarget = UGizmoComponentHitTarget::Construct(InScaleComponent, this);
+	if (this->UpdateHoverFunctionComponent)
+	{
+		HitTarget->UpdateHoverFunction = [InScaleComponent, this](bool bHovering) { this->UpdateHoverFunctionComponent(InScaleComponent, bHovering); };
+	}
+	ScaleGizmo->HitTarget = HitTarget;
+
+	ScaleGizmo->StateTarget = Cast<UObject>(InStateTarget);
+
+	ActiveGizmos.Add(ScaleGizmo);
+	return ScaleGizmo;
+	*/
+	return nullptr;
+}
+
+
+void UEditorTransformGizmo::ClearActiveTarget()
+{
+	for (UInteractiveGizmo* Gizmo : ActiveGizmos)
+	{
+		GetGizmoManager()->DestroyGizmo(Gizmo);
+	}
+	ActiveGizmos.SetNum(0);
+	ActiveObjects.SetNum(0);
+	NonuniformScaleObjects.SetNum(0);
+
+	CameraAxisSource = nullptr;
+	GizmoGroupObject = nullptr;
+	AxisXSource = nullptr;
+	AxisYSource = nullptr;
+	AxisZSource = nullptr;
+	AxisXObject = nullptr;
+	AxisYObject = nullptr;
+	AxisZObject = nullptr;
+	UnitAxisXSource = nullptr;
+	UnitAxisYSource = nullptr;
+	UnitAxisZSource = nullptr;
+	StateTarget = nullptr;
+
+	ActiveTarget = nullptr;
+}
+
+
+bool UEditorTransformGizmo::PositionSnapFunction(const FVector& WorldPosition, FVector& SnappedPositionOut) const
+{
+	SnappedPositionOut = WorldPosition;
+
+	// only snap if we want snapping obvs
+	if (bSnapToWorldGrid == false)
+	{
+		return false;
+	}
+
+	// only snap to world grid when using world axes
+	if (GetGizmoManager()->GetContextQueriesAPI()->GetCurrentCoordinateSystem() != EToolContextCoordinateSystem::World)
+	{
+		return false;
+	}
+
+	FSceneSnapQueryRequest Request;
+	Request.RequestType = ESceneSnapQueryType::Position;
+	Request.TargetTypes = ESceneSnapQueryTargetType::Grid;
+	Request.Position = WorldPosition;
+	if ( bGridSizeIsExplicit )
+	{
+		Request.GridSize = ExplicitGridSize;
+	}
+	TArray<FSceneSnapQueryResult> Results;
+	if (GetGizmoManager()->GetContextQueriesAPI()->ExecuteSceneSnapQuery(Request, Results))
+	{
+		SnappedPositionOut = Results[0].Position;
+		return true;
+	};
+
+	return false;
+}
+
+
+FQuat UEditorTransformGizmo::RotationSnapFunction(const FQuat& DeltaRotation) const
+{
+	FQuat SnappedDeltaRotation = DeltaRotation;
+
+	// only snap if we want snapping 
+	if (bSnapToWorldRotGrid)
+	{
+		FSceneSnapQueryRequest Request;
+		Request.RequestType   = ESceneSnapQueryType::Rotation;
+		Request.TargetTypes   = ESceneSnapQueryTargetType::Grid;
+		Request.DeltaRotation = DeltaRotation;
+		if ( bRotationGridSizeIsExplicit )
+		{
+			Request.RotGridSize = ExplicitRotationGridSize;
+		}
+		TArray<FSceneSnapQueryResult> Results;
+		if (GetGizmoManager()->GetContextQueriesAPI()->ExecuteSceneSnapQuery(Request, Results))
+		{
+			SnappedDeltaRotation = Results[0].DeltaRotation;
+		};
+	}
+	return SnappedDeltaRotation;
+}
+
+#undef LOCTEXT_NAMESPACE
diff --git a/Engine/Source/Editor/Experimental/EditorInteractiveToolsFramework/Public/EditorGizmos/EditorAxisSources.h b/Engine/Source/Editor/Experimental/EditorInteractiveToolsFramework/Public/EditorGizmos/EditorAxisSources.h
new file mode 100644
--- /dev/null
+++ b/Engine/Source/Editor/Experimental/EditorInteractiveToolsFramework/Public/EditorGizmos/EditorAxisSources.h
@@ -0,0 +1,43 @@
+// Copyright Epic Games, Inc. All Rights Reserved.
+
+#pragma once
+
+#include "CoreMinimal.h"
+#include "BaseGizmos/GizmoInterfaces.h"
+#include "EditorAxisSources.generated.h"
+
+/**
+ * UGizmoEditorAxisSource is an IGizmoAxisSource implementation that provides one of the
+ * X/Y/Z axes of the Level Editor local coordinate system, mapped to World, based on an integer AxisIndex in range [0,2].
+ * The Axis Origin is the Level Editor's pivot location. 
+ */
+
+UCLASS()
+class EDITORINTERACTIVETOOLSFRAMEWORK_API UGizmoEditorAxisSource : public UObject, public IGizmoAxisSource
+{
+	GENERATED_BODY()
+public:
+	virtual FVector GetOrigin() const;
+
+	virtual FVector GetDirection() const;
+
+	/** Clamped internally to 0,1,2 */
+	UPROPERTY()
+	int AxisIndex = 2;
+
+	/** If false, returns World axes */
+	UPROPERTY()
+	bool bLocalAxes = true;
+
+public:
+	static UGizmoEditorAxisSource* Construct(
+		int LocalAxisIndex,
+		bool bUseLocalAxes,
+		UObject* Outer = (UObject*)GetTransientPackage())
+	{
+		UGizmoEditorAxisSource* NewSource = NewObject<UGizmoEditorAxisSource>(Outer);
+		NewSource->AxisIndex = LocalAxisIndex;
+		NewSource->bLocalAxes = bUseLocalAxes;
+		return NewSource;
+	}
+};
\ No newline at end of file
diff --git a/Engine/Source/Editor/Experimental/EditorInteractiveToolsFramework/Public/EditorGizmos/EditorParameterToTransformAdapters.h b/Engine/Source/Editor/Experimental/EditorInteractiveToolsFramework/Public/EditorGizmos/EditorParameterToTransformAdapters.h
new file mode 100644
--- /dev/null
+++ b/Engine/Source/Editor/Experimental/EditorInteractiveToolsFramework/Public/EditorGizmos/EditorParameterToTransformAdapters.h
@@ -0,0 +1,130 @@
+// Copyright Epic Games, Inc. All Rights Reserved.
+
+#pragma once
+
+#include "CoreMinimal.h"
+#include "BaseGizmos/GizmoInterfaces.h"
+#include "BaseGizmos/GizmoMath.h"
+#include "BaseGizmos/ParameterSourcesFloat.h"
+#include "BaseGizmos/ParameterSourcesVec2.h"
+#include "BaseGizmos/ParameterToTransformAdapters.h"
+#include "EditorModeManager.h"
+#include "EditorParameterToTransformAdapters.generated.h"
+
+//
+// Various 1D and 2D ParameterSource converters intended to be used to create 3D transformation gizmos.
+// Based on base classes in ParameterSourcesFloat.h and ParameterSourcesVec2.h
+// 
+
+/**
+ * UGizmoEditorAxisTranslationParameterSource is an IGizmoFloatParameterSource implementation that
+ * interprets the float value as the parameter of a line equation, and maps this parameter to a 3D translation 
+ * along a line with origin/direction given by an IGizmoAxisSource. This translation is applied to the mode tools 
+ * manager's widget location and its super methods are called to also apply it to an IGizmoTransformSource.
+ *
+ * This ParameterSource is intended to be used to create Editor 3D Axis Translation Gizmos.
+ */
+
+//@todo rename this UGizmoEditorAxisParameterSource
+ 
+UCLASS()
+class EDITORINTERACTIVETOOLSFRAMEWORK_API UGizmoEditorAxisTranslationParameterSource : public UGizmoBaseFloatParameterSource
+{
+	GENERATED_BODY()
+public:
+
+	virtual float GetParameter() const override
+	{
+		return AxisTranslationParameterSource->GetParameter();
+	}
+
+	virtual void SetParameter(float NewValue) override
+	{
+		Parameter = NewValue;
+		LastChange.CurrentValue = NewValue;
+
+		// construct translation as delta from initial position
+		FVector Translation = LastChange.GetChangeDelta() * CurTranslationAxis;
+
+		// @todo should this be handled in a parameter source instead? to properly handle world transformation
+		FEditorModeTools& EditorModeTools = GLevelEditorModeTools();
+		EditorModeTools.PivotLocation = LastPivotLocation + Translation;
+		EditorModeTools.SnappedLocation = LastPivotLocation + Translation;
+
+		AxisTranslationParameterSource->SetParameter(Parameter);
+	}
+
+	virtual void BeginModify()
+	{
+		check(AxisSource);
+		LastChange = FGizmoFloatParameterChange(Parameter);
+		CurTranslationAxis = AxisSource->GetDirection();
+		CurTranslationOrigin = AxisSource->GetOrigin();
+
+		FEditorModeTools& EditorModeTools = GLevelEditorModeTools();
+		LastPivotLocation = EditorModeTools.PivotLocation;
+		LastSnappedLocation = EditorModeTools.SnappedLocation;
+
+		AxisTranslationParameterSource->BeginModify();
+	}
+
+	virtual void EndModify()
+	{
+		AxisTranslationParameterSource->EndModify();
+
+	}
+
+public:
+	/** The Parameter line-equation value is converted to a 3D Translation along this Axis */
+	UPROPERTY()
+	TScriptInterface<IGizmoAxisSource> AxisSource;
+
+	/** Wrapped parameter source */
+	TObjectPtr<UGizmoAxisTranslationParameterSource> AxisTranslationParameterSource;
+
+public:
+	/** Parameter is the line-equation parameter that this FloatParameterSource provides */
+	UPROPERTY()
+	float Parameter = 0.0f;
+
+	/** Active parameter change (only valid between BeginModify/EndModify) */
+	UPROPERTY()
+	FGizmoFloatParameterChange LastChange;
+
+	/** translation axis for current parameter edit (only valid between BeginModify/EndModify) */
+	UPROPERTY()
+	FVector CurTranslationAxis;
+
+	/** translation origin for current parameter edit (only valid between BeginModify/EndModify) */
+	UPROPERTY()
+	FVector CurTranslationOrigin;
+
+	/** Last pivot location before translation began (only valid between BeginModify/EndModify) */
+	UPROPERTY()
+	FVector LastPivotLocation = FVector::ZeroVector;
+
+	/** Last snapped location before translation began (only valid between BeginModify/EndModify) */
+	UPROPERTY()
+	FVector LastSnappedLocation = FVector::ZeroVector;
+
+
+public:
+	/**
+	 * Create a standard instance of this ParameterSource, with the given AxisSource and TransformSource
+	 */
+	static UGizmoEditorAxisTranslationParameterSource* Construct(
+		IGizmoAxisSource* InAxisSource,
+		IGizmoTransformSource* InTransformSource,
+		UObject* Outer = (UObject*)GetTransientPackage())
+	{
+		UGizmoEditorAxisTranslationParameterSource* NewSource = NewObject<UGizmoEditorAxisTranslationParameterSource>(Outer);
+		
+		UGizmoAxisTranslationParameterSource* ParameterSource = UGizmoAxisTranslationParameterSource::Construct(InAxisSource, InTransformSource, Outer);
+		NewSource->AxisTranslationParameterSource = ParameterSource;
+		NewSource->AxisSource = Cast<UObject>(InAxisSource);
+		return NewSource;
+	}
+};
+
+
+
diff --git a/Engine/Source/Editor/Experimental/EditorInteractiveToolsFramework/Public/EditorGizmos/EditorTransformGizmo.h b/Engine/Source/Editor/Experimental/EditorInteractiveToolsFramework/Public/EditorGizmos/EditorTransformGizmo.h
new file mode 100644
--- /dev/null
+++ b/Engine/Source/Editor/Experimental/EditorInteractiveToolsFramework/Public/EditorGizmos/EditorTransformGizmo.h
@@ -0,0 +1,350 @@
+// Copyright Epic Games, Inc. All Rights Reserved.
+
+#pragma once
+
+#include "CoreMinimal.h"
+#include "BaseGizmos/GizmoArrowObject.h"
+#include "BaseGizmos/GizmoBaseObject.h"
+#include "BaseGizmos/GizmoGroupObject.h"
+#include "BaseGizmos/TransformProxy.h"
+#include "EditorGizmos/EditorAxisSources.h"
+#include "InteractiveGizmo.h"
+#include "InteractiveGizmoBuilder.h"
+#include "InteractiveToolObjects.h"
+#include "InteractiveToolChange.h"
+#include "Materials/MaterialInstanceDynamic.h"
+
+#include "EditorTransformGizmo.generated.h"
+
+class UInteractiveGizmoManager;
+class IGizmoAxisSource;
+class IGizmoTransformSource;
+class IGizmoStateTarget;
+class UGizmoConstantFrameAxisSource;
+class UGizmoTransformChangeStateTarget;
+class FEditorTransformGizmoTransformChange;
+class FRay;
+
+UCLASS()
+class EDITORINTERACTIVETOOLSFRAMEWORK_API UEditorTransformGizmoBuilder : public UInteractiveGizmoBuilder
+{
+	GENERATED_BODY()
+
+public:
+
+	virtual UInteractiveGizmo* BuildGizmo(const FToolBuilderState& SceneState) const override;
+};
+
+
+/**
+ * UEditorTransformGizmo provides standard Transformation Gizmo interactions,
+ * applied to a UTransformProxy target object. By default the Gizmo will be
+ * a standard XYZ translate/rotate Gizmo (axis and plane translation).
+ */
+UCLASS()
+class EDITORINTERACTIVETOOLSFRAMEWORK_API UEditorTransformGizmo : public UInteractiveGizmo
+{
+	GENERATED_BODY()
+
+public:
+	constexpr static float AXIS_LENGTH = 35.0f;
+	constexpr static float AXIS_RADIUS = 1.2f;
+	constexpr static float AXIS_CONE_ANGLE = 15.7f;
+	constexpr static float AXIS_CONE_HEIGHT = 13;
+	constexpr static float AXIS_CONE_HEAD_OFFSET = 12;
+	constexpr static float AXIS_CUBE_SIZE = 4;
+	constexpr static float AXIS_CUBE_HEAD_OFFSET= 3;
+	constexpr static float TRANSLATE_ROTATE_AXIS_CIRCLE_RADIUS = 20.0f;
+	constexpr static float TWOD_AXIS_CIRCLE_RADIUS = 10.0f;
+	constexpr static float INNER_AXIS_CIRCLE_RADIUS = 48.0f;
+	constexpr static float OUTER_AXIS_CIRCLE_RADIUS = 56.0f;
+	constexpr static float ROTATION_TEXT_RADIUS = 75.0f;
+	constexpr static int32 AXIS_CIRCLE_SIDES = 24;
+	constexpr static float ARCALL_RELATIVE_INNER_SIZE = 0.75f;
+	constexpr static float AXIS_LENGTH_SCALE = 25.0f;
+	constexpr static float AXIS_LENGTH_SCALE_OFFSET = 5.0f;
+
+	constexpr static FLinearColor AxisColorX = FLinearColor(0.594f, 0.0197f, 0.0f);
+	constexpr static FLinearColor AxisColorY = FLinearColor(0.1349f, 0.3959f, 0.0f);
+	constexpr static FLinearColor AxisColorZ = FLinearColor(0.0251f, 0.207f, 0.85f);
+	constexpr static FLinearColor ScreenAxisColor = FLinearColor(0.76, 0.72, 0.14f);
+	constexpr static FColor PlaneColorXY = FColor(255, 255, 0); // FColor::Yellow
+	constexpr static FColor ArcBallColor = FColor(128, 128, 128, 6);
+	constexpr static FColor ScreenSpaceColor = FColor(196, 196, 196);
+	constexpr static FColor CurrentColor = FColor(255, 255, 0); // FColor::Yellow
+
+public:
+
+	virtual void SetWorld(UWorld* World);
+	virtual void SetElements(ETransformGizmoSubElements InEnableElements);
+
+	/**
+	 * By default, non-uniform scaling handles appear (assuming they exist in the gizmo to begin with), 
+	 * when CurrentCoordinateSystem == EToolContextCoordinateSystem::Local, since components can only be
+	 * locally scaled. However, this can be changed to a custom check here, perhaps to hide them in extra
+	 * conditions or to always show them (if the gizmo is not scaling a component).
+	 */
+	virtual void SetIsNonUniformScaleAllowedFunction(
+		TUniqueFunction<bool()>&& IsNonUniformScaleAllowed
+	);
+
+	/**
+	 * By default, the nonuniform scale components can scale negatively. However, they can be made to clamp
+	 * to zero instead by passing true here. This is useful for using the gizmo to flatten geometry.
+	 *
+	 * TODO: Should this affect uniform scaling too?
+	 */
+	virtual void SetDisallowNegativeScaling(bool bDisallow);
+
+	// UInteractiveGizmo overrides
+	virtual void Setup() override;
+	virtual void Shutdown() override;
+	virtual void Tick(float DeltaTime) override;
+
+
+	/**
+	 * Set the active target object for the Gizmo
+	 * @param Target active target
+	 * @param TransactionProvider optional IToolContextTransactionProvider implementation to use - by default uses GizmoManager
+	 */
+	virtual void SetActiveTarget(UTransformProxy* Target, IToolContextTransactionProvider* TransactionProvider = nullptr);
+
+	/**
+	 * Clear the active target object for the Gizmo
+	 */
+	virtual void ClearActiveTarget();
+
+	/** The active target object for the Gizmo */
+	UPROPERTY()
+	UTransformProxy* ActiveTarget;
+
+	/**
+	 * Repositions the gizmo without issuing undo/redo changes, triggering callbacks, 
+	 * or moving any components. Useful for resetting the gizmo to a new location without
+	 * it being viewed as a gizmo manipulation.
+	 */
+	void ReinitializeGizmoTransform(const FTransform& NewTransform);
+
+	/**
+	 * Set a new position for the Gizmo. This is done via the same mechanisms as the sub-gizmos,
+	 * so it generates the same Change/Modify() events, and hence works with Undo/Redo
+	 */
+	virtual void SetNewGizmoTransform(const FTransform& NewTransform);
+
+	/**
+	 * Explicitly set the child scale. Mainly useful to "reset" the child scale to (1,1,1) when re-using Gizmo across multiple transform actions.
+	 * @warning does not generate change/modify events!!
+	 */
+	virtual void SetNewChildScale(const FVector& NewChildScale);
+
+	/**
+	 * Set visibility for this Gizmo
+	 */
+	virtual void SetVisibility(bool bVisible);
+
+	/**
+	 * If true, then when using world frame, Axis and Plane translation snap to the world grid via the ContextQueriesAPI (in PositionSnapFunction)
+	 */
+	UPROPERTY()
+	bool bSnapToWorldGrid = false;
+
+	/**
+	 * Optional grid size which overrides the Context Grid
+	 */
+	UPROPERTY()
+	bool bGridSizeIsExplicit = false;
+	UPROPERTY()
+	FVector ExplicitGridSize;
+
+	/**
+	 * Optional grid size which overrides the Context Rotation Grid
+	 */
+	UPROPERTY()
+	bool bRotationGridSizeIsExplicit = false;
+	UPROPERTY()
+	FRotator ExplicitRotationGridSize;
+
+	/**
+	 * If true, then when using world frame, Axis and Plane translation snap to the world grid via the ContextQueriesAPI (in RotationSnapFunction)
+	 */
+	UPROPERTY()
+	bool bSnapToWorldRotGrid = false;
+
+	/**
+	 * Whether to use the World/Local coordinate system provided by the context via the ContextyQueriesAPI.
+	 */
+	UPROPERTY()
+	bool bUseContextCoordinateSystem = true;
+
+	/**
+	 * Current coordinate system in use. If bUseContextCoordinateSystem is true, this value will be updated internally every Tick()
+	 * by querying the ContextyQueriesAPI, otherwise the default is Local and the client can change it as necessary
+	 */
+	UPROPERTY()
+	EToolContextCoordinateSystem CurrentCoordinateSystem = EToolContextCoordinateSystem::Local;
+
+
+protected:
+	// TSharedPtr<FEditorTransformGizmoActorFactory> GizmoEditorActorBuilder;
+	/** Only these parts of the UEditorTransformGizmo will be initialized */
+	ETransformGizmoSubElements EnableElements =
+		ETransformGizmoSubElements::TranslateAllAxes |
+		ETransformGizmoSubElements::TranslateAllPlanes |
+		ETransformGizmoSubElements::RotateAllAxes |
+		ETransformGizmoSubElements::ScaleAllAxes |
+		ETransformGizmoSubElements::ScaleAllPlanes |
+		ETransformGizmoSubElements::ScaleUniform;
+
+	/** List of current-active gizmo objects */
+	UPROPERTY()
+	TArray<TObjectPtr<UGizmoBaseObject>> ActiveObjects;
+
+	/**
+	 * List of nonuniform scale objects. Subset of of ActiveObjects. These are tracked separately so they can
+	 * be hidden when Gizmo is not configured to use local axes, because UE only supports local nonuniform scaling
+	 * on Components
+	 */
+	UPROPERTY()
+	TArray<TObjectPtr<UGizmoBaseObject>> NonuniformScaleObjects;
+
+	/** list of currently-active child gizmos */
+	UPROPERTY()
+	TArray<TObjectPtr<UInteractiveGizmo>> ActiveGizmos;
+
+	/** GizmoActors will be spawned in this World */
+	UWorld* World;
+
+	//
+	// Axis Sources
+	//
+	/** Axis that points towards camera, X/Y plane tangents aligned to right/up. Shared across Gizmos, and created internally during SetActiveTarget() */
+	UPROPERTY()
+	TObjectPtr<UGizmoConstantFrameAxisSource> CameraAxisSource;
+
+	// internal function that updates CameraAxisSource by getting current view state from GizmoManager
+	void UpdateCameraAxisSource();
+
+	/** Gizmo group object contains all gizmos created internally during SetActiveTarget() */
+	UPROPERTY()
+	TObjectPtr<UGizmoGroupObject> GizmoGroupObject;
+
+	// @todo - these properties can probably be removed since they are contained within the group now.
+	/** X-axis source is shared across Gizmos, and created internally during SetActiveTarget() */
+	UPROPERTY()
+	TObjectPtr<UGizmoEditorAxisSource> AxisXSource;
+
+	/** Y-axis source is shared across Gizmos, and created internally during SetActiveTarget() */
+	UPROPERTY()
+	TObjectPtr<UGizmoEditorAxisSource> AxisYSource;
+
+	/** Z-axis source is shared across Gizmos, and created internally during SetActiveTarget() */
+	UPROPERTY()
+	TObjectPtr<UGizmoEditorAxisSource> AxisZSource;
+
+	/** X-axis arrow object is shared across Gizmos, and created internally during SetActiveTarget() */
+	UPROPERTY()
+	TObjectPtr<UGizmoArrowObject> AxisXObject;
+
+	/** Y-axis arrow object is shared across Gizmos, and created internally during SetActiveTarget() */
+	UPROPERTY()
+	TObjectPtr<UGizmoArrowObject> AxisYObject;
+
+	/** Z-axis arrow object is shared across Gizmos, and created internally during SetActiveTarget() */
+	UPROPERTY()
+	TObjectPtr<UGizmoArrowObject> AxisZObject;
+	//
+	// Scaling support. 
+	// UE Components only support scaling in local coordinates, so we have to create separate sources for that.
+	//
+
+	/** Local X-axis source (ie 1,0,0) is shared across Scale Gizmos, and created internally during SetActiveTarget() */
+	UPROPERTY()
+	TObjectPtr<UGizmoEditorAxisSource> UnitAxisXSource;
+
+	/** Y-axis source (ie 0,1,0) is shared across Scale Gizmos, and created internally during SetActiveTarget() */
+	UPROPERTY()
+	TObjectPtr<UGizmoEditorAxisSource> UnitAxisYSource;
+
+	/** Z-axis source (ie 0,0,1) is shared across Scale Gizmos, and created internally during SetActiveTarget() */
+	UPROPERTY()
+	TObjectPtr<UGizmoEditorAxisSource> UnitAxisZSource;
+
+	/** 
+	 * State target is shared across gizmos, and created internally during SetActiveTarget(). 
+	 * Several FChange providers are registered with this StateTarget, including the UEditorTransformGizmo
+	 * itself (IToolCommandChangeSource implementation above is called)
+	 */
+	UPROPERTY()
+	TObjectPtr<UGizmoObjectTransformChangeStateTarget> StateTarget;
+
+	/**
+	 * These are used to let the translation subgizmos use raycasts into the scene to align the gizmo with scene geometry.
+	 * See comment for SetWorldAlignmentFunctions().
+	 */
+	TUniqueFunction<bool()> ShouldAlignDestination = []() { return false; };
+	TUniqueFunction<bool(const FRay&, FVector&)> DestinationAlignmentRayCaster = [](const FRay&, FVector&) {return false; };
+
+	TUniqueFunction<bool()> IsNonUniformScaleAllowed = [this]() { return CurrentCoordinateSystem == EToolContextCoordinateSystem::Local; };
+
+	bool bDisallowNegativeScaling = false;
+protected:
+
+
+	/** @return a new instance of the standard axis-translation Gizmo */
+	virtual UInteractiveGizmo* AddAxisTranslationGizmo(
+		UGizmoArrowObject* InArrowObject,
+		IGizmoAxisSource* InAxisSource,
+		IGizmoTransformSource* InTransformSource,
+		IGizmoStateTarget* InStateTarget,
+		EAxisList::Type InAxisType,
+		const FLinearColor InAxisColor);
+
+	/** @return a new instance of the standard plane-translation Gizmo */
+	virtual UInteractiveGizmo* AddPlaneTranslationGizmo(
+		IGizmoAxisSource* InAxisSource,
+		IGizmoTransformSource* InTransformSource,
+		IGizmoStateTarget* InStateTarget);
+
+	/** @return a new instance of the standard axis-rotation Gizmo */
+	virtual UInteractiveGizmo* AddAxisRotationGizmo(
+		IGizmoAxisSource* InAxisSource,
+		IGizmoTransformSource* InTransformSource,
+		IGizmoStateTarget* InStateTarget,
+		EAxisList::Type InAxisType,
+		const FLinearColor InAxisColor);
+
+	/** @return a new instance of the standard axis-scaling Gizmo */
+	virtual UInteractiveGizmo* AddAxisScaleGizmo(
+		UGizmoArrowObject* InArrowObject,
+		IGizmoAxisSource* InGizmoAxisSource, IGizmoAxisSource* InParameterAxisSource,
+		IGizmoTransformSource* InTransformSource,
+		IGizmoStateTarget* InStateTarget,
+		EAxisList::Type InAxisType,
+		const FLinearColor InAxisColor);
+
+	/** @return a new instance of the standard plane-scaling Gizmo */
+	virtual UInteractiveGizmo* AddPlaneScaleGizmo(
+		IGizmoAxisSource* InGizmoAxisSource, IGizmoAxisSource* InParameterAxisSource,
+		IGizmoTransformSource* InTransformSource,
+		IGizmoStateTarget* InStateTarget);
+
+	/** @return a new instance of the standard plane-scaling Gizmo */
+	virtual UInteractiveGizmo* AddUniformScaleGizmo(
+		IGizmoAxisSource* InGizmoAxisSource, IGizmoAxisSource* InParameterAxisSource,
+		IGizmoTransformSource* InTransformSource,
+		IGizmoStateTarget* InStateTarget);
+
+	// Axis and Plane TransformSources use this function to execute worldgrid snap queries
+	bool PositionSnapFunction(const FVector& WorldPosition, FVector& SnappedPositionOut) const;
+	FQuat RotationSnapFunction(const FQuat& DeltaRotation) const;
+
+	/** Materials and colors to be used when drawing the items for each axis */
+	TObjectPtr<UMaterialInterface> TransparentPlaneMaterialXY;
+	TObjectPtr<UMaterialInterface> GridMaterial;
+
+	TObjectPtr<UMaterialInstanceDynamic> AxisMaterialX;
+	TObjectPtr<UMaterialInstanceDynamic> AxisMaterialY;
+	TObjectPtr<UMaterialInstanceDynamic> AxisMaterialZ;
+	TObjectPtr<UMaterialInstanceDynamic> CurrentAxisMaterial;
+	TObjectPtr<UMaterialInstanceDynamic> OpaquePlaneMaterialXY;
+};
