Add a skeletal mesh component tool target that allows editing skeletal meshes with the regular modeling tools.

diff --git a/Engine/Source/Editor/Experimental/EditorInteractiveToolsFramework/Private/ToolTargets/SkeletalMeshComponentToolTarget.cpp b/Engine/Source/Editor/Experimental/EditorInteractiveToolsFramework/Private/ToolTargets/SkeletalMeshComponentToolTarget.cpp
new file mode 100644
--- /dev/null
+++ b/Engine/Source/Editor/Experimental/EditorInteractiveToolsFramework/Private/ToolTargets/SkeletalMeshComponentToolTarget.cpp
@@ -0,0 +1,225 @@
+﻿// Copyright Epic Games, Inc. All Rights Reserved.
+
+#include "ToolTargets/SkeletalMeshComponentToolTarget.h"
+
+#include "Components/SkinnedMeshComponent.h"
+#include "Rendering/SkeletalMeshModel.h"
+
+
+namespace USkeletalMeshComponentToolTargetLocals
+{
+	int32 LODIndex = 0;
+}
+
+
+int32 USkeletalMeshComponentToolTarget::GetNumMaterials() const
+{
+	return ensure(IsValid()) ? Component->GetNumMaterials() : 0;
+}
+
+UMaterialInterface* USkeletalMeshComponentToolTarget::GetMaterial(int32 MaterialIndex) const
+{
+	return ensure(IsValid()) ? Component->GetMaterial(MaterialIndex) : nullptr;
+}
+
+void USkeletalMeshComponentToolTarget::GetMaterialSet(FComponentMaterialSet& MaterialSetOut, bool bPreferAssetMaterials) const
+{
+	if (!ensure(IsValid())) return;
+
+	if (bPreferAssetMaterials)
+	{
+		const USkeletalMesh* SkeletalMesh = Cast<USkinnedMeshComponent>(Component)->SkeletalMesh;
+		const TArray<FSkeletalMaterial>& Materials = SkeletalMesh->GetMaterials(); 
+		MaterialSetOut.Materials.SetNum(Materials.Num());
+		for (int32 k = 0; k < Materials.Num(); ++k)
+		{
+			MaterialSetOut.Materials[k] = Materials[k].MaterialInterface;
+		}
+	}
+	else
+	{
+		int32 NumMaterials = Component->GetNumMaterials();
+		MaterialSetOut.Materials.SetNum(NumMaterials);
+		for (int32 k = 0; k < NumMaterials; ++k)
+		{
+			MaterialSetOut.Materials[k] = Component->GetMaterial(k);
+		}
+	}
+}
+
+bool USkeletalMeshComponentToolTarget::CommitMaterialSetUpdate(const FComponentMaterialSet& MaterialSet, bool bApplyToAsset)
+{
+	if (!ensure(IsValid())) return false;
+
+	if (bApplyToAsset)
+	{
+		USkeletalMesh* SkeletalMesh = Cast<USkinnedMeshComponent>(Component)->SkeletalMesh;
+
+		if (SkeletalMesh->GetPathName().StartsWith(TEXT("/Engine/")))
+		{
+			UE_LOG(LogTemp, Warning, TEXT("CANNOT MODIFY BUILT-IN ENGINE ASSET %s"), *SkeletalMesh->GetPathName());
+			return false;
+		}
+
+		// flush any pending rendering commands, which might touch this component while we are rebuilding its mesh
+		FlushRenderingCommands();
+
+		// unregister the component while we update it's static mesh
+		TUniquePtr<FComponentReregisterContext> ComponentReregisterContext = MakeUnique<FComponentReregisterContext>(Component);
+
+		// make sure transactional flag is on
+		SkeletalMesh->SetFlags(RF_Transactional);
+
+		SkeletalMesh->Modify();
+
+		const int NewNumMaterials = MaterialSet.Materials.Num();
+		TArray<FSkeletalMaterial> &SkeletalMaterials = SkeletalMesh->GetMaterials(); 
+		if (NewNumMaterials != SkeletalMaterials.Num())
+		{
+			SkeletalMaterials.SetNum(NewNumMaterials);
+		}
+		for (int k = 0; k < NewNumMaterials; ++k)
+		{
+			if (SkeletalMaterials[k].MaterialInterface != MaterialSet.Materials[k])
+			{
+				SkeletalMaterials[k].MaterialInterface = MaterialSet.Materials[k];
+				if (SkeletalMaterials[k].MaterialSlotName.IsNone())
+				{
+					SkeletalMaterials[k].MaterialSlotName = MaterialSet.Materials[k]->GetFName();
+				}
+			}
+		}
+
+		SkeletalMesh->PostEditChange();
+	}
+	else
+	{
+		const int32 NumMaterialsNeeded = Component->GetNumMaterials();
+		const int32 NumMaterialsGiven = MaterialSet.Materials.Num();
+
+		// We wrote the below code to support a mismatch in the number of materials.
+		// However, it is not yet clear whether this might be desirable, and we don't
+		// want to inadvertantly hide bugs in the meantime. So, we keep this ensure here
+		// for now, and we can remove it if we decide that we want the ability.
+		ensure(NumMaterialsNeeded == NumMaterialsGiven);
+
+		check(NumMaterialsGiven > 0);
+
+		for (int32 i = 0; i < NumMaterialsNeeded; ++i)
+		{
+			const int32 MaterialToUseIndex = FMath::Min(i, NumMaterialsGiven - 1);
+			Component->SetMaterial(i, MaterialSet.Materials[MaterialToUseIndex]);
+		}
+	}
+
+	return true;
+}
+
+FMeshDescription* USkeletalMeshComponentToolTarget::GetMeshDescription()
+{
+	if (!ensure(IsValid()))
+	{
+		return nullptr;
+	}
+
+	if (!CachedMeshDescription.IsValid())
+	{
+		using namespace USkeletalMeshComponentToolTargetLocals;
+
+		CachedMeshDescription = MakeUnique<FMeshDescription>();
+		const USkeletalMesh* SkeletalMesh = Cast<USkinnedMeshComponent>(Component)->SkeletalMesh;
+
+		// Check first if we have bulk data available.
+		if (SkeletalMesh->IsLODImportedDataBuildAvailable(LODIndex))
+		{
+			FSkeletalMeshImportData SkeletalMeshImportData;
+			SkeletalMesh->LoadLODImportedData(LODIndex, SkeletalMeshImportData);
+			SkeletalMeshImportData.GetMeshDescription(*CachedMeshDescription);
+		}
+		else
+		{
+			// Fall back on the LOD model directly if no bulk data exists. When we commit
+			// the mesh description, we override using the bulk data. This can happen for older
+			// skeletal meshes, from UE 4.24 and earlier.
+			const FSkeletalMeshModel* SkeletalMeshModel = SkeletalMesh->GetImportedModel();
+			if (SkeletalMeshModel && SkeletalMeshModel->LODModels.IsValidIndex(LODIndex))
+			{
+				SkeletalMeshModel->LODModels[LODIndex].GetMeshDescription(*CachedMeshDescription, SkeletalMesh);
+			}			
+		}
+	}
+	
+	return CachedMeshDescription.Get();
+}
+
+void USkeletalMeshComponentToolTarget::CommitMeshDescription(const FCommitter& Committer)
+{
+	using namespace USkeletalMeshComponentToolTargetLocals;
+	if (ensure(IsValid()) == false) return;
+
+	USkeletalMesh* SkeletalMesh = Cast<USkinnedMeshComponent>(Component)->SkeletalMesh;
+
+	if (SkeletalMesh->GetPathName().StartsWith(TEXT("/Engine/")))
+	{
+		const FString DebugMessage = FString::Printf(TEXT("CANNOT MODIFY BUILT-IN ENGINE ASSET %s"), *SkeletalMesh->GetPathName());
+		if (GAreScreenMessagesEnabled)
+		{
+			GEngine->AddOnScreenDebugMessage(INDEX_NONE, 10.0f, FColor::Red, DebugMessage);
+		}
+		UE_LOG(LogTemp, Warning, TEXT("%s"), *DebugMessage);
+		return;
+	}
+
+	// flush any pending rendering commands, which might touch this component while we are rebuilding it's mesh
+	FlushRenderingCommands();
+
+	// unregister the component while we update its skeletal mesh
+	FComponentReregisterContext ComponentReregisterContext(Component);
+
+	// make sure transactional flag is on for this asset
+	SkeletalMesh->SetFlags(RF_Transactional);
+
+	verify(SkeletalMesh->Modify());
+
+	FCommitterParams CommitterParams;
+
+	CommitterParams.MeshDescriptionOut = GetMeshDescription();
+
+	Committer(CommitterParams);
+
+	FSkeletalMeshImportData SkeletalMeshImportData = FSkeletalMeshImportData::CreateFromMeshDescription(*CommitterParams.MeshDescriptionOut);
+	SkeletalMesh->SaveLODImportedData(LODIndex, SkeletalMeshImportData);
+
+	// Make sure the mesh builder knows it's the latest variety, so that the render data gets
+	// properly rebuilt.
+	SkeletalMesh->SetLODImportedDataVersions(LODIndex, ESkeletalMeshGeoImportVersions::LatestVersion, ESkeletalMeshSkinningImportVersions::LatestVersion);
+	SkeletalMesh->SetUseLegacyMeshDerivedDataKey(false);
+
+	SkeletalMesh->PostEditChange();
+
+	// this rebuilds physics, but it doesn't undo!
+	Component->RecreatePhysicsState();
+}
+
+USkeletalMesh* USkeletalMeshComponentToolTarget::GetSkeletalMesh() const
+{
+	return IsValid() ? Cast<USkinnedMeshComponent>(Component)->SkeletalMesh : nullptr;
+}
+
+
+// Factory
+
+bool USkeletalMeshComponentToolTargetFactory::CanBuildTarget(UObject* SourceObject, const FToolTargetTypeRequirements& Requirements) const
+{
+	return Cast<USkinnedMeshComponent>(SourceObject) && Cast<USkinnedMeshComponent>(SourceObject)->SkeletalMesh &&
+		Requirements.AreSatisfiedBy(USkeletalMeshComponentToolTarget::StaticClass());
+}
+
+UToolTarget* USkeletalMeshComponentToolTargetFactory::BuildTarget(UObject* SourceObject, const FToolTargetTypeRequirements& Requirements)
+{
+	USkeletalMeshComponentToolTarget* Target = NewObject<USkeletalMeshComponentToolTarget>();
+	Target->Component = Cast<USkinnedMeshComponent>(SourceObject);
+	check(Target->Component && Requirements.AreSatisfiedBy(Target));
+
+	return Target;
+}
diff --git a/Engine/Source/Editor/Experimental/EditorInteractiveToolsFramework/Public/ToolTargets/SkeletalMeshComponentToolTarget.h b/Engine/Source/Editor/Experimental/EditorInteractiveToolsFramework/Public/ToolTargets/SkeletalMeshComponentToolTarget.h
new file mode 100644
--- /dev/null
+++ b/Engine/Source/Editor/Experimental/EditorInteractiveToolsFramework/Public/ToolTargets/SkeletalMeshComponentToolTarget.h
@@ -0,0 +1,69 @@
+﻿// Copyright Epic Games, Inc. All Rights Reserved.
+
+#pragma once
+
+#include "CoreMinimal.h"
+
+#include "TargetInterfaces/MaterialProvider.h"
+#include "TargetInterfaces/MeshDescriptionCommitter.h"
+#include "TargetInterfaces/MeshDescriptionProvider.h"
+#include "TargetInterfaces/SkeletalMeshBackedTarget.h"
+#include "ToolTargets/PrimitiveComponentToolTarget.h"
+
+#include "SkeletalMeshComponentToolTarget.generated.h"
+
+
+class USkeletalMesh;
+
+/**
+ * A tool target backed by a skeletal mesh component that can provide and take a mesh
+ * description.
+ */
+UCLASS(Transient)
+class EDITORINTERACTIVETOOLSFRAMEWORK_API USkeletalMeshComponentToolTarget :
+	public UPrimitiveComponentToolTarget,
+	public IMeshDescriptionCommitter,
+	public IMeshDescriptionProvider,
+	public IMaterialProvider,
+	public ISkeletalMeshBackedTarget
+{
+	GENERATED_BODY()
+
+public:
+	// IMeshDescriptionProvider implementation
+	FMeshDescription* GetMeshDescription() override;
+
+	// IMeshDescritpionCommitter implementation
+	void CommitMeshDescription(const FCommitter& Committer) override;
+
+	// IMaterialProvider implementation
+	int32 GetNumMaterials() const override;
+	UMaterialInterface* GetMaterial(int32 MaterialIndex) const override;
+	void GetMaterialSet(FComponentMaterialSet& MaterialSetOut, bool bPreferAssetMaterials) const override;
+	bool CommitMaterialSetUpdate(const FComponentMaterialSet& MaterialSet, bool bApplyToAsset) override;	
+
+	// ISkeletalMeshBackedTarget implementation
+	USkeletalMesh* GetSkeletalMesh() const override;
+
+protected:
+	// So that the tool target factory can poke into Component.
+	friend class USkeletalMeshComponentToolTargetFactory;
+
+private:
+	// Until USkeletalMesh stores its internal representation as FMeshDescription, we need to
+	// retain the storage here to cover the lifetime of the pointer returned by GetMeshDescription(). 
+	TUniquePtr<FMeshDescription> CachedMeshDescription;	
+};
+
+/** Factory for USkeletalMeshComponentToolTarget to be used by the target manager. */
+UCLASS(Transient)
+class EDITORINTERACTIVETOOLSFRAMEWORK_API USkeletalMeshComponentToolTargetFactory : public UToolTargetFactory
+{
+	GENERATED_BODY()
+
+public:
+
+	bool CanBuildTarget(UObject* SourceObject, const FToolTargetTypeRequirements& TargetTypeInfo) const override;
+
+	UToolTarget* BuildTarget(UObject* SourceObject, const FToolTargetTypeRequirements& TargetTypeInfo) override;
+};
diff --git a/Engine/Source/Runtime/Experimental/InteractiveToolsFramework/Public/TargetInterfaces/SkeletalMeshBackedTarget.h b/Engine/Source/Runtime/Experimental/InteractiveToolsFramework/Public/TargetInterfaces/SkeletalMeshBackedTarget.h
new file mode 100644
--- /dev/null
+++ b/Engine/Source/Runtime/Experimental/InteractiveToolsFramework/Public/TargetInterfaces/SkeletalMeshBackedTarget.h
@@ -0,0 +1,35 @@
+﻿// Copyright Epic Games, Inc. All Rights Reserved.
+
+#pragma once
+
+#include "CoreMinimal.h"
+
+#include "Engine/SkeletalMesh.h"
+#include "TargetInterfaces/AssetBackedTarget.h"
+
+#include "SkeletalMeshBackedTarget.generated.h"
+
+UINTERFACE()
+class INTERACTIVETOOLSFRAMEWORK_API USkeletalMeshBackedTarget : public UAssetBackedTarget
+{
+	GENERATED_BODY()
+};
+
+class INTERACTIVETOOLSFRAMEWORK_API ISkeletalMeshBackedTarget : public IAssetBackedTarget
+{
+	GENERATED_BODY()
+
+	public:
+	/**
+	* @return the underlying source asset for this Target.
+	*/
+	virtual UObject* GetSourceData() const override
+	{
+		return GetSkeletalMesh();
+	}
+
+	/**
+	* @return the underlying source USkeletalMesh asset for this Target
+	*/
+	virtual USkeletalMesh* GetSkeletalMesh() const = 0;
+};
