MeshModelingTools: Fix crash bug when selecting a tool that requires a static mesh but there is no static mesh available (e.g. after a failed load).

#rnx

diff --git a/Engine/Source/Editor/Experimental/EditorInteractiveToolsFramework/Private/ToolTargets/StaticMeshComponentToolTarget.cpp b/Engine/Source/Editor/Experimental/EditorInteractiveToolsFramework/Private/ToolTargets/StaticMeshComponentToolTarget.cpp
--- a/Engine/Source/Editor/Experimental/EditorInteractiveToolsFramework/Private/ToolTargets/StaticMeshComponentToolTarget.cpp
+++ b/Engine/Source/Editor/Experimental/EditorInteractiveToolsFramework/Private/ToolTargets/StaticMeshComponentToolTarget.cpp
@@ -169,7 +169,9 @@ UStaticMesh* UStaticMeshComponentToolTarget::GetStaticMesh() const
 
 bool UStaticMeshComponentToolTargetFactory::CanBuildTarget(UObject* SourceObject, const FToolTargetTypeRequirements& Requirements) const
 {
-	return Cast<UStaticMeshComponent>(SourceObject) && Requirements.AreSatisfiedBy(UStaticMeshComponentToolTarget::StaticClass());
+	const UStaticMeshComponent* Component = Cast<UStaticMeshComponent>(SourceObject);
+	return Component && !Component->IsPendingKillOrUnreachable() && Component->IsValidLowLevel() && Component->GetStaticMesh()
+		&& Requirements.AreSatisfiedBy(UStaticMeshComponentToolTarget::StaticClass());
 }
 
 UToolTarget* UStaticMeshComponentToolTargetFactory::BuildTarget(UObject* SourceObject, const FToolTargetTypeRequirements& Requirements)
