MeshModelingTools: Cancel active tools on world tear down. This addresses issues where changes are saved but we cannot undo them, and it also fixes issues where accepting tool edits during world tear down breaks assumptions made in the code.

#rnx

diff --git a/Engine/Source/Editor/Experimental/EditorInteractiveToolsFramework/Private/EdModeInteractiveToolsContext.cpp b/Engine/Source/Editor/Experimental/EditorInteractiveToolsFramework/Private/EdModeInteractiveToolsContext.cpp
--- a/Engine/Source/Editor/Experimental/EditorInteractiveToolsFramework/Private/EdModeInteractiveToolsContext.cpp
+++ b/Engine/Source/Editor/Experimental/EditorInteractiveToolsFramework/Private/EdModeInteractiveToolsContext.cpp
@@ -590,7 +590,7 @@ void UEdModeInteractiveToolsContext::Shutdown()
 		GEditor->OnViewportClientListChanged().Remove(ViewportClientListChangedHandle);
 
 		// auto-accept any in-progress tools
-		DeactivateAllActiveTools();
+		DeactivateAllActiveTools(EToolShutdownType::Accept);
 
 		UInteractiveToolsContext::Shutdown();
 	}
@@ -652,15 +652,15 @@ void UEdModeInteractiveToolsContext::ShutdownContext()
 
 void UEdModeInteractiveToolsContext::TerminateActiveToolsOnPIEStart()
 {
-	DeactivateAllActiveTools();
+	DeactivateAllActiveTools(EToolShutdownType::Accept);
 }
 void UEdModeInteractiveToolsContext::TerminateActiveToolsOnSaveWorld()
 {
-	DeactivateAllActiveTools();
+	DeactivateAllActiveTools(EToolShutdownType::Accept);
 }
 void UEdModeInteractiveToolsContext::TerminateActiveToolsOnWorldTearDown()
 {
-	DeactivateAllActiveTools();
+	DeactivateAllActiveTools(EToolShutdownType::Cancel);
 }
 
 void UEdModeInteractiveToolsContext::PostInvalidation()
@@ -1213,9 +1213,9 @@ void UEdModeInteractiveToolsContext::DeactivateActiveTool(EToolSide WhichSide, E
 	RestoreEditorState();
 }
 
-void UEdModeInteractiveToolsContext::DeactivateAllActiveTools()
+void UEdModeInteractiveToolsContext::DeactivateAllActiveTools(EToolShutdownType ShutdownType)
 {
-	UInteractiveToolsContext::DeactivateAllActiveTools();
+	UInteractiveToolsContext::DeactivateAllActiveTools(ShutdownType);
 	RestoreEditorState();
 }
 
diff --git a/Engine/Source/Editor/Experimental/EditorInteractiveToolsFramework/Public/EdModeInteractiveToolsContext.h b/Engine/Source/Editor/Experimental/EditorInteractiveToolsFramework/Public/EdModeInteractiveToolsContext.h
--- a/Engine/Source/Editor/Experimental/EditorInteractiveToolsFramework/Public/EdModeInteractiveToolsContext.h
+++ b/Engine/Source/Editor/Experimental/EditorInteractiveToolsFramework/Public/EdModeInteractiveToolsContext.h
@@ -97,7 +97,7 @@ protected:
 	virtual void Shutdown() override;
 
 	virtual void DeactivateActiveTool(EToolSide WhichSide, EToolShutdownType ShutdownType);
-	virtual void DeactivateAllActiveTools();
+	virtual void DeactivateAllActiveTools(EToolShutdownType ShutdownType);
 
 public:
 	UPROPERTY()
diff --git a/Engine/Source/Runtime/Experimental/InteractiveToolsFramework/Private/InteractiveToolsContext.cpp b/Engine/Source/Runtime/Experimental/InteractiveToolsFramework/Private/InteractiveToolsContext.cpp
--- a/Engine/Source/Runtime/Experimental/InteractiveToolsFramework/Private/InteractiveToolsContext.cpp
+++ b/Engine/Source/Runtime/Experimental/InteractiveToolsFramework/Private/InteractiveToolsContext.cpp
@@ -49,20 +49,20 @@ void UInteractiveToolsContext::DeactivateActiveTool(EToolSide WhichSide, EToolSh
 	ToolManager->DeactivateTool(WhichSide, ShutdownType);
 }
 
-void UInteractiveToolsContext::DeactivateAllActiveTools()
-{
-	if (ToolManager->HasActiveTool(EToolSide::Left))
-	{
-		EToolShutdownType ShutdownType = ToolManager->CanAcceptActiveTool(EToolSide::Left) ?
-			EToolShutdownType::Accept : EToolShutdownType::Cancel;
-		ToolManager->DeactivateTool(EToolSide::Left, ShutdownType);
-	}
-	if (ToolManager->HasActiveTool(EToolSide::Right))
-	{
-		EToolShutdownType ShutdownType = ToolManager->CanAcceptActiveTool(EToolSide::Right) ?
-			EToolShutdownType::Accept : EToolShutdownType::Cancel;
-		ToolManager->DeactivateTool(EToolSide::Right, ShutdownType);
-	}
+void UInteractiveToolsContext::DeactivateAllActiveTools(EToolShutdownType ShutdownType)
+{
+	auto DeactivateTool = [this, ShutdownType](EToolSide WhichSide) {
+		if (ToolManager->HasActiveTool(WhichSide))
+		{
+			const EToolShutdownType AcceptOrCancel =
+				ShutdownType != EToolShutdownType::Cancel && ToolManager->CanAcceptActiveTool(WhichSide)
+				? EToolShutdownType::Accept : EToolShutdownType::Cancel;
+			ToolManager->DeactivateTool(WhichSide, AcceptOrCancel);
+		}
+	};
+	
+	DeactivateTool(EToolSide::Left);
+	DeactivateTool(EToolSide::Right);
 }
 
 
diff --git a/Engine/Source/Runtime/Experimental/InteractiveToolsFramework/Public/InteractiveToolsContext.h b/Engine/Source/Runtime/Experimental/InteractiveToolsFramework/Public/InteractiveToolsContext.h
--- a/Engine/Source/Runtime/Experimental/InteractiveToolsFramework/Public/InteractiveToolsContext.h
+++ b/Engine/Source/Runtime/Experimental/InteractiveToolsFramework/Public/InteractiveToolsContext.h
@@ -33,7 +33,7 @@ public:
 	virtual void Shutdown();
 
 	virtual void DeactivateActiveTool(EToolSide WhichSide, EToolShutdownType ShutdownType);
-	virtual void DeactivateAllActiveTools();
+	virtual void DeactivateAllActiveTools(EToolShutdownType ShutdownType);
 
 	bool CanStartTool(EToolSide WhichSide, const FString& ToolTypeIdentifier) const;
 	bool HasActiveTool(EToolSide WhichSide) const;
