ModeingMode: restore EditingLOD functionality lost in ToolTarget transition. Port EditingLOD functionality from FStaticMeshComponentTarget/Factory to UStaticMeshComponentToolTarget/Factory. Add UToolTargetManager::FindFirstFactoryByPredicate() and FindFirstFactoryByType(). ModelingToolsEditorModeToolkit now also looks up the StaticMeshFactory and updates the EditingLOD when it changes in the UI.

EStaticMeshEditingLOD enum moved to ComponentSourceInterfaces.h to make it more widely available.

Add ToolBuilderUtil::EnumerateComponents() and ToolTargetManager::EnumerateSelectedAndTargetableComponents(), this allows ToolBuilders to do additional checks on the valid Targets without having to make local arrays/etc. Fix SetCollisionGeometryTool to not build ToolTargets every frame just to check if one is a StaticMeshComponent.

#rnx

diff --git a/Engine/Plugins/Experimental/MeshModelingToolset/Source/MeshModelingTools/Private/Physics/SetCollisionGeometryTool.cpp b/Engine/Plugins/Experimental/MeshModelingToolset/Source/MeshModelingTools/Private/Physics/SetCollisionGeometryTool.cpp
--- a/Engine/Plugins/Experimental/MeshModelingToolset/Source/MeshModelingTools/Private/Physics/SetCollisionGeometryTool.cpp
+++ b/Engine/Plugins/Experimental/MeshModelingToolset/Source/MeshModelingTools/Private/Physics/SetCollisionGeometryTool.cpp
@@ -46,8 +46,10 @@ const FToolTargetTypeRequirements& USetCollisionGeometryToolBuilder::GetTargetRe
 
 bool USetCollisionGeometryToolBuilder::CanBuildTool(const FToolBuilderState& SceneState) const
 {
-	TArray<TObjectPtr<UToolTarget>> Targets = SceneState.TargetManager->BuildAllSelectedTargetable(SceneState, GetTargetRequirements());
-	return (Targets.Num() > 0 && Cast<IStaticMeshBackedTarget>(Targets[Targets.Num() - 1]) != nullptr);
+	UActorComponent* LastValidTarget = nullptr;
+	SceneState.TargetManager->EnumerateSelectedAndTargetableComponents(SceneState, GetTargetRequirements(),
+		[&](UActorComponent* Component) { LastValidTarget = Component; });
+	return (LastValidTarget != nullptr && Cast<UStaticMeshComponent>(LastValidTarget) != nullptr);
 }
 
 
diff --git a/Engine/Plugins/Experimental/ModelingToolsEditorMode/Source/ModelingToolsEditorMode/Private/ModelingToolsEditorModeToolkit.cpp b/Engine/Plugins/Experimental/ModelingToolsEditorMode/Source/ModelingToolsEditorMode/Private/ModelingToolsEditorModeToolkit.cpp
--- a/Engine/Plugins/Experimental/ModelingToolsEditorMode/Source/ModelingToolsEditorMode/Private/ModelingToolsEditorModeToolkit.cpp
+++ b/Engine/Plugins/Experimental/ModelingToolsEditorMode/Source/ModelingToolsEditorMode/Private/ModelingToolsEditorModeToolkit.cpp
@@ -28,6 +28,8 @@
 // for LOD setting
 #include "EditorInteractiveToolsFrameworkModule.h"
 #include "Tools/EditorComponentSourceFactory.h"
+#include "ToolTargetManager.h"
+#include "ToolTargets/StaticMeshComponentToolTarget.h"
 
 #define LOCTEXT_NAMESPACE "FModelingToolsEditorModeToolkit"
 
@@ -259,33 +261,44 @@ TSharedPtr<SWidget> FModelingToolsEditorModeToolkit::MakeAssetConfigPanel()
 		.OptionsSource(&AssetLODModes)
 		.OnSelectionChanged_Lambda([&](TSharedPtr<FString> String, ESelectInfo::Type)
 	{
+		EStaticMeshEditingLOD NewSelectedLOD = EStaticMeshEditingLOD::LOD0;
+		if (*String == *AssetLODModes[0])
+		{
+			NewSelectedLOD = EStaticMeshEditingLOD::MaxQuality;
+		}
+		else if (*String == *AssetLODModes[1])
+		{
+			NewSelectedLOD = EStaticMeshEditingLOD::HiResSource;
+		}
+		else
+		{
+			for (int32 k = 2; k < AssetLODModes.Num(); ++k)
+			{
+				if (*String == *AssetLODModes[k])
+				{
+					NewSelectedLOD = (EStaticMeshEditingLOD)(k - 2);
+					break;
+				}
+			}
+		}
+
 		if (FEditorInteractiveToolsFrameworkGlobals::RegisteredStaticMeshTargetFactoryKey >= 0)
 		{
 			FComponentTargetFactory* Factory = FindComponentTargetFactoryByKey(FEditorInteractiveToolsFrameworkGlobals::RegisteredStaticMeshTargetFactoryKey);
 			if (Factory != nullptr)
 			{
 				FStaticMeshComponentTargetFactory* StaticMeshFactory = static_cast<FStaticMeshComponentTargetFactory*>(Factory);
-				if (*String == *AssetLODModes[0])
-				{
-					StaticMeshFactory->CurrentEditingLOD = EStaticMeshEditingLOD::MaxQuality;
-				}
-				else if (*String == *AssetLODModes[1])
-				{
-					StaticMeshFactory->CurrentEditingLOD = EStaticMeshEditingLOD::HiResSource;
-				}
-				else
-				{
-					for (int32 k = 2; k < AssetLODModes.Num(); ++k)
-					{
-						if (*String == *AssetLODModes[k])
-						{
-							StaticMeshFactory->CurrentEditingLOD = (EStaticMeshEditingLOD)(k - 2);
-							break;
-						}
-					}
-				}
+				StaticMeshFactory->CurrentEditingLOD = NewSelectedLOD;
 			}
 		}
+
+		TObjectPtr<UToolTargetManager> TargetManager = GetScriptableEditorMode()->GetInteractiveToolsContext()->TargetManager;
+		UStaticMeshComponentToolTargetFactory* StaticMeshTargetFactory = TargetManager->FindFirstFactoryByType<UStaticMeshComponentToolTargetFactory>();
+		if (StaticMeshTargetFactory)
+		{
+			StaticMeshTargetFactory->SetActiveEditingLOD(NewSelectedLOD);
+		}
+
 	});
 
 
diff --git a/Engine/Plugins/Experimental/ModelingToolsEditorMode/Source/ModelingToolsEditorMode/Private/StaticMeshComponentBuilder.h b/Engine/Plugins/Experimental/ModelingToolsEditorMode/Source/ModelingToolsEditorMode/Private/StaticMeshComponentBuilder.h
--- a/Engine/Plugins/Experimental/ModelingToolsEditorMode/Source/ModelingToolsEditorMode/Private/StaticMeshComponentBuilder.h
+++ b/Engine/Plugins/Experimental/ModelingToolsEditorMode/Source/ModelingToolsEditorMode/Private/StaticMeshComponentBuilder.h
@@ -45,6 +45,7 @@ public:
 		NewStaticMesh->SetNumSourceModels(1);
 		NewStaticMesh->GetSourceModel(0).BuildSettings.bRecomputeNormals = false;
 		NewStaticMesh->GetSourceModel(0).BuildSettings.bRecomputeTangents = true;
+		NewStaticMesh->GetSourceModel(0).BuildSettings.bRemoveDegenerates = false;
 		this->MeshDescription = NewStaticMesh->CreateMeshDescription(0);
 
 		if (NewStaticMesh->GetBodySetup() == nullptr)
diff --git a/Engine/Source/Editor/Experimental/EditorInteractiveToolsFramework/Private/ToolTargets/StaticMeshComponentToolTarget.cpp b/Engine/Source/Editor/Experimental/EditorInteractiveToolsFramework/Private/ToolTargets/StaticMeshComponentToolTarget.cpp
--- a/Engine/Source/Editor/Experimental/EditorInteractiveToolsFramework/Private/ToolTargets/StaticMeshComponentToolTarget.cpp
+++ b/Engine/Source/Editor/Experimental/EditorInteractiveToolsFramework/Private/ToolTargets/StaticMeshComponentToolTarget.cpp
@@ -7,11 +7,90 @@
 #include "Engine/StaticMesh.h"
 #include "RenderingThread.h"
 
-namespace UStaticMeshComponentToolTargetLocals
+
+static void DisplayCriticalWarningMessage(const FString& Message)
+{
+	if (GAreScreenMessagesEnabled)
+	{
+		GEngine->AddOnScreenDebugMessage(INDEX_NONE, 10.0f, FColor::Red, Message);
+	}
+	UE_LOG(LogTemp, Warning, TEXT("%s"), *Message);
+}
+
+
+void UStaticMeshComponentToolTarget::SetEditingLOD(EStaticMeshEditingLOD RequestedEditingLOD)
+{
+	EStaticMeshEditingLOD ValidEditingLOD = EStaticMeshEditingLOD::LOD0;
+
+	UStaticMeshComponent* StaticMeshComponent = Cast<UStaticMeshComponent>(Component);
+	if (ensure(StaticMeshComponent != nullptr))
+	{
+		UStaticMesh* StaticMeshAsset = StaticMeshComponent->GetStaticMesh();
+		if (ensure(StaticMeshAsset != nullptr))
+		{
+			if (RequestedEditingLOD == EStaticMeshEditingLOD::MaxQuality)
+			{
+				ValidEditingLOD = StaticMeshAsset->IsHiResMeshDescriptionValid() ? EStaticMeshEditingLOD::HiResSource : EStaticMeshEditingLOD::LOD0;
+			}
+			else if (RequestedEditingLOD == EStaticMeshEditingLOD::HiResSource)
+			{
+				ValidEditingLOD = StaticMeshAsset->IsHiResMeshDescriptionValid() ? EStaticMeshEditingLOD::HiResSource : EStaticMeshEditingLOD::LOD0;
+				if (ValidEditingLOD != EStaticMeshEditingLOD::HiResSource)
+				{
+					DisplayCriticalWarningMessage(FString(TEXT("HiRes Source selected but not available - Falling Back to LOD0")));
+				}
+			}
+			else
+			{
+				ValidEditingLOD = RequestedEditingLOD;
+				int32 MaxExistingLOD = StaticMeshAsset->GetNumSourceModels() - 1;
+				if ((int32)ValidEditingLOD > MaxExistingLOD)
+				{
+					DisplayCriticalWarningMessage(FString::Printf(TEXT("LOD%d Requested but not available - Falling Back to LOD%d"), (int32)ValidEditingLOD, MaxExistingLOD));
+					ValidEditingLOD = (EStaticMeshEditingLOD)MaxExistingLOD;
+				}
+			}
+		}
+	}
+
+	EditingLOD = ValidEditingLOD;
+}
+
+
+bool UStaticMeshComponentToolTarget::IsValid() const
 {
-	int32 LODIndex = 0;
+	if (!UPrimitiveComponentToolTarget::IsValid())
+	{
+		return false;
+	}
+	UStaticMeshComponent* StaticMeshComponent = Cast<UStaticMeshComponent>(Component);
+	if (StaticMeshComponent == nullptr)
+	{
+		return false;
+	}
+	UStaticMesh* StaticMesh = StaticMeshComponent->GetStaticMesh();
+	if (StaticMesh == nullptr)
+	{
+		return false;
+	}
+
+	if (EditingLOD == EStaticMeshEditingLOD::HiResSource)
+	{
+		if (StaticMesh->IsHiResMeshDescriptionValid() == false)
+		{
+			return false;
+		}
+	}
+	else if ((int32)EditingLOD >= StaticMesh->GetNumSourceModels())
+	{
+		return false;
+	}
+
+	return true;
 }
 
+
+
 int32 UStaticMeshComponentToolTarget::GetNumMaterials() const
 {
 	return ensure(IsValid()) ? Component->GetNumMaterials() : 0;
@@ -110,28 +189,28 @@ bool UStaticMeshComponentToolTarget::CommitMaterialSetUpdate(const FComponentMat
 	return true;
 }
 
+
 FMeshDescription* UStaticMeshComponentToolTarget::GetMeshDescription()
 {
-	using namespace UStaticMeshComponentToolTargetLocals;
-
-	return IsValid() ? Cast<UStaticMeshComponent>(Component)->GetStaticMesh()->GetMeshDescription(LODIndex) : nullptr;
+	if (ensure(IsValid()))
+	{
+		UStaticMesh* StaticMesh = Cast<UStaticMeshComponent>(Component)->GetStaticMesh();
+		return (EditingLOD == EStaticMeshEditingLOD::HiResSource) ?
+			StaticMesh->GetHiResMeshDescription() : StaticMesh->GetMeshDescription((int32)EditingLOD);
+	}
+	return nullptr;
 }
 
+
 void UStaticMeshComponentToolTarget::CommitMeshDescription(const FCommitter& Committer)
 {
-	using namespace UStaticMeshComponentToolTargetLocals;
 	if (ensure(IsValid()) == false) return;
 
 	UStaticMesh* StaticMesh = Cast<UStaticMeshComponent>(Component)->GetStaticMesh();
 
 	if (StaticMesh->GetPathName().StartsWith(TEXT("/Engine/")))
 	{
-		FString DebugMessage = FString::Printf(TEXT("CANNOT MODIFY BUILT-IN ENGINE ASSET %s"), *StaticMesh->GetPathName());
-		if (GAreScreenMessagesEnabled)
-		{
-			GEngine->AddOnScreenDebugMessage(INDEX_NONE, 10.0f, FColor::Red, DebugMessage);
-		}
-		UE_LOG(LogTemp, Warning, TEXT("%s"), *DebugMessage);
+		DisplayCriticalWarningMessage(FString::Printf(TEXT("CANNOT MODIFY BUILT-IN ENGINE ASSET %s"), *StaticMesh->GetPathName()));
 		return;
 	}
 
@@ -145,14 +224,29 @@ void UStaticMeshComponentToolTarget::CommitMeshDescription(const FCommitter& Com
 	StaticMesh->SetFlags(RF_Transactional);
 
 	verify(StaticMesh->Modify());
-	verify(StaticMesh->ModifyMeshDescription(LODIndex));
+	if (EditingLOD == EStaticMeshEditingLOD::HiResSource)
+	{
+		verify(StaticMesh->ModifyHiResMeshDescription());
+	}
+	else
+	{
+		verify(StaticMesh->ModifyMeshDescription((int32)EditingLOD));
+	}
 
 	FCommitterParams CommitterParams;
-	CommitterParams.MeshDescriptionOut = StaticMesh->GetMeshDescription(LODIndex);
+	CommitterParams.MeshDescriptionOut = GetMeshDescription();
 
 	Committer(CommitterParams);
 
-	StaticMesh->CommitMeshDescription(LODIndex);
+	if (EditingLOD == EStaticMeshEditingLOD::HiResSource)
+	{
+		StaticMesh->CommitHiResMeshDescription();
+	}
+	else
+	{
+		StaticMesh->CommitMeshDescription((int32)EditingLOD);
+	}
+
 	StaticMesh->PostEditChange();
 
 	// this rebuilds physics, but it doesn't undo!
@@ -171,6 +265,7 @@ bool UStaticMeshComponentToolTargetFactory::CanBuildTarget(UObject* SourceObject
 {
 	const UStaticMeshComponent* Component = Cast<UStaticMeshComponent>(SourceObject);
 	return Component && !Component->IsPendingKillOrUnreachable() && Component->IsValidLowLevel() && Component->GetStaticMesh()
+		&& (Component->GetStaticMesh()->GetNumSourceModels() > 0)
 		&& Requirements.AreSatisfiedBy(UStaticMeshComponentToolTarget::StaticClass());
 }
 
@@ -178,7 +273,14 @@ UToolTarget* UStaticMeshComponentToolTargetFactory::BuildTarget(UObject* SourceO
 {
 	UStaticMeshComponentToolTarget* Target = NewObject<UStaticMeshComponentToolTarget>();// TODO: Should we set an outer here?
 	Target->Component = Cast<UStaticMeshComponent>(SourceObject);
+	Target->SetEditingLOD(EditingLOD);
 	check(Target->Component && Requirements.AreSatisfiedBy(Target));
 
 	return Target;
+}
+
+
+void UStaticMeshComponentToolTargetFactory::SetActiveEditingLOD(EStaticMeshEditingLOD NewEditingLOD)
+{
+	EditingLOD = NewEditingLOD;
 }
\ No newline at end of file
diff --git a/Engine/Source/Editor/Experimental/EditorInteractiveToolsFramework/Public/ToolTargets/StaticMeshComponentToolTarget.h b/Engine/Source/Editor/Experimental/EditorInteractiveToolsFramework/Public/ToolTargets/StaticMeshComponentToolTarget.h
--- a/Engine/Source/Editor/Experimental/EditorInteractiveToolsFramework/Public/ToolTargets/StaticMeshComponentToolTarget.h
+++ b/Engine/Source/Editor/Experimental/EditorInteractiveToolsFramework/Public/ToolTargets/StaticMeshComponentToolTarget.h
@@ -9,6 +9,7 @@
 #include "TargetInterfaces/MeshDescriptionProvider.h"
 #include "TargetInterfaces/StaticMeshBackedTarget.h"
 #include "ToolTargets/PrimitiveComponentToolTarget.h"
+#include "ComponentSourceInterfaces.h"  // for EStaticMeshEditingLOD
 
 #include "StaticMeshComponentToolTarget.generated.h"
 
@@ -25,7 +26,20 @@ class EDITORINTERACTIVETOOLSFRAMEWORK_API UStaticMeshComponentToolTarget : publi
 	GENERATED_BODY()
 
 public:
+	/**
+	 * Configure active LOD to edit. Can only call this after Component is configured in base UPrimitiveComponentToolTarget.
+	 * If requested LOD does not exist, fall back to one that does.
+	 */
+	virtual void SetEditingLOD(EStaticMeshEditingLOD RequestedEditingLOD);
 
+	/** @return current editing LOD */
+	virtual EStaticMeshEditingLOD GetEditingLOD() const { return EditingLOD; }
+
+public:
+	virtual bool IsValid() const override;
+
+
+public:
 	// IMeshDescriptionProvider implementation
 	FMeshDescription* GetMeshDescription() override;
 
@@ -43,11 +57,15 @@ public:
 
 	// Rest provided by parent class
 
+public:
+
 protected:
+	EStaticMeshEditingLOD EditingLOD = EStaticMeshEditingLOD::LOD0;
 
 	friend class UStaticMeshComponentToolTargetFactory;
 };
 
+
 /** Factory for UStaticMeshComponentToolTarget to be used by the target manager. */
 UCLASS(Transient)
 class EDITORINTERACTIVETOOLSFRAMEWORK_API UStaticMeshComponentToolTargetFactory : public UToolTargetFactory
@@ -59,4 +77,14 @@ public:
 	virtual bool CanBuildTarget(UObject* SourceObject, const FToolTargetTypeRequirements& TargetTypeInfo) const override;
 
 	virtual UToolTarget* BuildTarget(UObject* SourceObject, const FToolTargetTypeRequirements& TargetTypeInfo) override;
+
+
+
+public:
+	virtual EStaticMeshEditingLOD GetActiveEditingLOD() const { return EditingLOD; }
+	virtual void SetActiveEditingLOD(EStaticMeshEditingLOD NewEditingLOD);
+
+protected:
+	// LOD to edit, default is to edit LOD0
+	EStaticMeshEditingLOD EditingLOD = EStaticMeshEditingLOD::LOD0;
 };
\ No newline at end of file
diff --git a/Engine/Source/Editor/UnrealEd/Public/Tools/EditorComponentSourceFactory.h b/Engine/Source/Editor/UnrealEd/Public/Tools/EditorComponentSourceFactory.h
--- a/Engine/Source/Editor/UnrealEd/Public/Tools/EditorComponentSourceFactory.h
+++ b/Engine/Source/Editor/UnrealEd/Public/Tools/EditorComponentSourceFactory.h
@@ -8,28 +8,6 @@
 #include "Components/StaticMeshComponent.h"
 
 
-/**
- * EStaticMeshEditingLOD specifies which LOD of a StaticMesh the FStaticMeshComponentTarget
- * refers to. FStaticMeshComponentTarget will remap requested value to the closest available LOD.
- */
-enum class EStaticMeshEditingLOD
-{
-	LOD0 = 0,
-	LOD1 = 1,
-	LOD2 = 2,
-	LOD3 = 3,
-	LOD4 = 4,
-	LOD5 = 5,
-	LOD6 = 6,
-	LOD7 = 7,
-
-	HiResSource = 20,	// HiRes source mesh is optional, and will fall back to LOD0 if not stored
-	MaxQuality = 21		// use HiRes source mesh if available, or LOD0 otherwise
-};
-
-
-
-
 
 class UNREALED_API FStaticMeshComponentTargetFactory : public FComponentTargetFactory
 {
diff --git a/Engine/Source/Runtime/Experimental/InteractiveToolsFramework/Private/ToolBuilderUtil.cpp b/Engine/Source/Runtime/Experimental/InteractiveToolsFramework/Private/ToolBuilderUtil.cpp
--- a/Engine/Source/Runtime/Experimental/InteractiveToolsFramework/Private/ToolBuilderUtil.cpp
+++ b/Engine/Source/Runtime/Experimental/InteractiveToolsFramework/Private/ToolBuilderUtil.cpp
@@ -83,6 +83,26 @@ TArray<UActorComponent*> ToolBuilderUtil::FindAllComponents(const FToolBuilderSt
 }
 
 
+void ToolBuilderUtil::EnumerateComponents(const FToolBuilderState& InputState, TFunctionRef<void(UActorComponent*)> ComponentFunc)
+{
+	if (InputState.SelectedComponents.Num() > 0)
+	{
+		for (UActorComponent* Component : InputState.SelectedComponents)
+		{
+			ComponentFunc(Component);
+		}
+	}
+	else
+	{
+		for (AActor* Actor : InputState.SelectedActors)
+		{
+			Actor->ForEachComponent(true,
+				[&ComponentFunc](UActorComponent* Component) {
+				ComponentFunc(Component);
+			});
+		}
+	}
+}
 
 
 
diff --git a/Engine/Source/Runtime/Experimental/InteractiveToolsFramework/Private/ToolTargetManager.cpp b/Engine/Source/Runtime/Experimental/InteractiveToolsFramework/Private/ToolTargetManager.cpp
--- a/Engine/Source/Runtime/Experimental/InteractiveToolsFramework/Private/ToolTargetManager.cpp
+++ b/Engine/Source/Runtime/Experimental/InteractiveToolsFramework/Private/ToolTargetManager.cpp
@@ -30,6 +30,14 @@ void UToolTargetManager::AddTargetFactory(UToolTargetFactory* Factory)
 	Factories.Add(Factory);
 }
 
+
+UToolTargetFactory* UToolTargetManager::FindFirstFactoryByPredicate(TFunctionRef<bool(UToolTargetFactory*)> Predicate)
+{
+	TObjectPtr<UToolTargetFactory>* Found = Factories.FindByPredicate(Predicate);
+	return (Found) ? Found->Get() : nullptr;
+}
+
+
 bool UToolTargetManager::CanBuildTarget(UObject* SourceObject, const FToolTargetTypeRequirements& TargetType) const
 {
 	for (UToolTargetFactory* Factory : Factories)
@@ -66,6 +74,23 @@ int32 UToolTargetManager::CountSelectedAndTargetable(const FToolBuilderState& Sc
 		});
 }
 
+
+void UToolTargetManager::EnumerateSelectedAndTargetableComponents(const FToolBuilderState& SceneState,
+	const FToolTargetTypeRequirements& TargetRequirements,
+	TFunctionRef<void(UActorComponent*)> ComponentFunc) const
+{
+	ToolBuilderUtil::EnumerateComponents(SceneState, [&](UActorComponent* Component)
+	{
+		if (CanBuildTarget(Component, TargetRequirements))
+		{
+			ComponentFunc(Component);
+		}
+	});
+}
+
+
+
+
 UToolTarget* UToolTargetManager::BuildFirstSelectedTargetable(const FToolBuilderState& SceneState, const FToolTargetTypeRequirements& TargetType)
 {
 	return BuildTarget(
diff --git a/Engine/Source/Runtime/Experimental/InteractiveToolsFramework/Public/ComponentSourceInterfaces.h b/Engine/Source/Runtime/Experimental/InteractiveToolsFramework/Public/ComponentSourceInterfaces.h
--- a/Engine/Source/Runtime/Experimental/InteractiveToolsFramework/Public/ComponentSourceInterfaces.h
+++ b/Engine/Source/Runtime/Experimental/InteractiveToolsFramework/Public/ComponentSourceInterfaces.h
@@ -16,6 +16,28 @@ class UPrimitiveComponent;
 struct FMeshDescription;
 class UMaterialInterface;
 
+
+/**
+ * EStaticMeshEditingLOD specifies which LOD of a StaticMesh the FStaticMeshComponentTarget
+ * refers to. FStaticMeshComponentTarget will remap requested value to the closest available LOD.
+ */
+enum class EStaticMeshEditingLOD
+{
+	LOD0 = 0,
+	LOD1 = 1,
+	LOD2 = 2,
+	LOD3 = 3,
+	LOD4 = 4,
+	LOD5 = 5,
+	LOD6 = 6,
+	LOD7 = 7,
+
+	HiResSource = 20,	// HiRes source mesh is optional, and will fall back to LOD0 if not stored
+	MaxQuality = 21		// use HiRes source mesh if available, or LOD0 otherwise
+};
+
+
+
 /**
  * @deprecated Use tool targets instead, such as UPrimitiveComponentToolTarget
  *
diff --git a/Engine/Source/Runtime/Experimental/InteractiveToolsFramework/Public/ToolBuilderUtil.h b/Engine/Source/Runtime/Experimental/InteractiveToolsFramework/Public/ToolBuilderUtil.h
--- a/Engine/Source/Runtime/Experimental/InteractiveToolsFramework/Public/ToolBuilderUtil.h
+++ b/Engine/Source/Runtime/Experimental/InteractiveToolsFramework/Public/ToolBuilderUtil.h
@@ -35,6 +35,8 @@ namespace ToolBuilderUtil
 	template<typename ComponentType>
 	ComponentType* FindFirstComponentOfType(const FToolBuilderState& InputState);
 
+	INTERACTIVETOOLSFRAMEWORK_API
+	void EnumerateComponents(const FToolBuilderState& InputState, TFunctionRef<void(UActorComponent*)> ComponentFunc);
 
 	/** Count number of selected Actors that pass predicate. */
 	INTERACTIVETOOLSFRAMEWORK_API
@@ -55,7 +57,6 @@ namespace ToolBuilderUtil
 	/** Find first first available Actor of given type, or return nullptr if not found */
 	template<typename ActorType>
 	ActorType* FindFirstActorOfType(const FToolBuilderState& InputState);
-
 }
 
 /*
diff --git a/Engine/Source/Runtime/Experimental/InteractiveToolsFramework/Public/ToolTargetManager.h b/Engine/Source/Runtime/Experimental/InteractiveToolsFramework/Public/ToolTargetManager.h
--- a/Engine/Source/Runtime/Experimental/InteractiveToolsFramework/Public/ToolTargetManager.h
+++ b/Engine/Source/Runtime/Experimental/InteractiveToolsFramework/Public/ToolTargetManager.h
@@ -31,8 +31,28 @@ public:
 	 */
 	bool IsActive() const { return bIsActive; }
 
+	/**
+	 * Add a new ToolTargetFactory
+	 */
 	virtual void AddTargetFactory(UToolTargetFactory* Factory);
 
+	/**
+	 * Find the first ToolTargetFactory that passes the Predicate function
+	 */
+	virtual UToolTargetFactory* FindFirstFactoryByPredicate(TFunctionRef<bool(UToolTargetFactory*)> Predicate);
+
+	/**
+	 * Find the first ToolTargetFactory of a given type
+	 */
+	template<typename CastToType>
+	CastToType* FindFirstFactoryByType()
+	{
+		UToolTargetFactory* Found = FindFirstFactoryByPredicate([](UToolTargetFactory* Factory) { return Cast<CastToType>(Factory) != nullptr; });
+		return (Found != nullptr) ? Cast<CastToType>(Found) : nullptr;
+	}
+
+
+
 	/** Examines stored target factories to see if one can build the requested type of target. */
 	virtual bool CanBuildTarget(UObject* SourceObject, const FToolTargetTypeRequirements& TargetRequirements) const;
 
@@ -60,6 +80,15 @@ public:
 	 */
 	virtual int32 CountSelectedAndTargetable(const FToolBuilderState& SceneState,
 		const FToolTargetTypeRequirements& TargetRequirements) const;
+
+	/**
+	 * Looks through the currently selected components and actors and counts the number of
+	 * inputs that could be used to create qualifying tool targets.
+	 */
+	virtual void EnumerateSelectedAndTargetableComponents(const FToolBuilderState& SceneState,
+		const FToolTargetTypeRequirements& TargetRequirements,
+		TFunctionRef<void(UActorComponent*)> ComponentFunc) const;
+
 	
 	/**
 	 * Looks through the currently selected components and actors and builds a target out of
