HordeExecutor: Remove testing config

diff --git a/Engine/Source/Editor/Experimental/HordeExecutor/Private/HordeExecutorSettings.cpp b/Engine/Source/Editor/Experimental/HordeExecutor/Private/HordeExecutorSettings.cpp
--- a/Engine/Source/Editor/Experimental/HordeExecutor/Private/HordeExecutorSettings.cpp
+++ b/Engine/Source/Editor/Experimental/HordeExecutor/Private/HordeExecutorSettings.cpp
@@ -6,6 +6,4 @@
 UHordeExecutorSettings::UHordeExecutorSettings(const FObjectInitializer& ObjectInitializer)
 	: Super(ObjectInitializer)
 {
-	Target = TEXT("localhost:5001");
-	PemRootCertificates = TEXT("[REDACTED_USER_HOME]\\Downloads\\localhost2.cer");
 }
