Rename HordeExecutor -> BazelExecutor

diff --git a/Engine/Source/Editor/Experimental/BazelExecutor/BazelExecutor.Build.cs b/Engine/Source/Editor/Experimental/BazelExecutor/BazelExecutor.Build.cs
new file mode 100644
--- /dev/null
+++ b/Engine/Source/Editor/Experimental/BazelExecutor/BazelExecutor.Build.cs
@@ -0,0 +1,37 @@
+// Copyright Epic Games, Inc. All Rights Reserved.
+
+using UnrealBuildTool;
+using System.IO;
+
+namespace UnrealBuildTool.Rules
+{
+	public class BazelExecutor : ModuleRules
+	{
+		public BazelExecutor(ReadOnlyTargetRules Target) : base(Target)
+		{
+			if (Target.Platform != UnrealTargetPlatform.Win64)
+			{
+				// Grpc is currently only supported for Win64, set to external to prevent compiling for any other platform
+				Type = ModuleType.External;
+				PrecompileForTargets = PrecompileTargetsType.None;
+				return;
+			}
+
+			bUseUnity = false;
+
+			PrivateDefinitions.Add("GPR_FORBID_UNREACHABLE_CODE=0");
+
+			PrivateIncludePaths.Add(Path.Combine(ModuleDirectory, "Generated"));
+
+			PrivateDependencyModuleNames.AddRange(
+				new string[] {
+				"Core",
+				"CoreUObject",
+				"Settings",
+				"Grpc",
+				"RemoteExecution",
+				}
+			);
+		}
+	}
+}
diff --git a/Engine/Source/Editor/Experimental/BazelExecutor/Generate.bat b/Engine/Source/Editor/Experimental/BazelExecutor/Generate.bat
new file mode 100755
--- /dev/null
+++ b/Engine/Source/Editor/Experimental/BazelExecutor/Generate.bat
@@ -0,0 +1,33 @@
+@echo off
+setlocal
+
+SET vcpkg_root=%~dp0..\..\..\ThirdParty\vcpkg\Win64\x64-windows-static-md-v142
+SET protoc=%vcpkg_root%\protobuf_x64-windows-static-md-v142\tools\protobuf\protoc.exe
+SET include=%vcpkg_root%\protobuf_x64-windows-static-md-v142\include
+SET grpc_cpp_plugin=%vcpkg_root%\grpc_x64-windows-static-md-v142\tools\grpc\grpc_cpp_plugin.exe
+
+pushd %~dp0
+mkdir .\Generated
+
+for /R ".\Schemas" %%f in (*.proto) do (
+	%protoc% -I=%include% --proto_path=%CD%\Schemas --cpp_out=%CD%\Generated %%f
+	%protoc% -I=%include% --proto_path=%CD%\Schemas --plugin=protoc-gen-grpc=%grpc_cpp_plugin% --grpc_out=%CD%\Generated %%f
+)
+
+for /R ".\Generated" %%f in (*.pb.cc) do (
+	type SchemaAutogenHeader.txt >%%f.new
+	type SchemaHeader.txt >>%%f.new
+	type %%f >>%%f.new
+	type SchemaFooter.txt >>%%f.new
+
+	move /y %%f.new %%f >NUL
+)
+
+for /R ".\Generated" %%f in (*.pb.h) do (
+	type SchemaAutogenHeader.txt >%%f.new
+	type %%f >>%%f.new
+
+	move /y %%f.new %%f >NUL
+)
+
+popd
diff --git a/Engine/Source/Editor/Experimental/BazelExecutor/Generate.sh b/Engine/Source/Editor/Experimental/BazelExecutor/Generate.sh
new file mode 100644
--- /dev/null
+++ b/Engine/Source/Editor/Experimental/BazelExecutor/Generate.sh
@@ -0,0 +1,30 @@
+#!/bin/bash
+
+vcpkg_root=../../../../ThirdParty/vcpkg/Linux/x64-linux
+protoc=${vcpkg_root}/protobuf_x64-linux/tools/protobuf/protoc-3.13.0.0
+include=${vcpkg_root}/protobuf_x64-linux/include
+grpc_cpp_plugin=${vcpkg_root}/grpc_x64-linux/tools/grpc/grpc_cpp_plugin
+
+mkdir -p ./Generated
+pushd ./Schemas
+for f in $(find ./ -name '*.proto'); do
+	${protoc} -I=${include} --proto_path=./ --cpp_out=../Generated $f
+	${protoc} -I=${include} --proto_path=./ --plugin=protoc-gen-grpc=${grpc_cpp_plugin} --grpc_out=../Generated $f
+done
+popd
+
+for f in $(find ./Generated -name '*.pb.cc'); do
+	cat SchemaAutogenHeader.txt >$f.new
+	cat SchemaHeader.txt >>$f.new
+	cat $f >>$f.new
+	cat SchemaFooter.txt >>$f.new
+
+	mv -f $f.new $f
+done
+
+for f in $(find ./Generated -name '*.pb.h'); do
+	cat SchemaAutogenHeader.txt >>$f.new
+	cat $f >>$f.new
+
+	mv -f $f.new $f
+done
diff --git a/Engine/Source/Editor/Experimental/BazelExecutor/Private/BazelExecutor.cpp b/Engine/Source/Editor/Experimental/BazelExecutor/Private/BazelExecutor.cpp
new file mode 100644
--- /dev/null
+++ b/Engine/Source/Editor/Experimental/BazelExecutor/Private/BazelExecutor.cpp
@@ -0,0 +1,59 @@
+// Copyright Epic Games, Inc. All Rights Reserved.
+
+#include "BazelExecutor.h"
+
+#include "Execution.h"
+#include "ContentAddressableStorage.h"
+
+THIRD_PARTY_INCLUDES_START
+UE_PUSH_MACRO("TEXT")
+#undef TEXT
+#include <grpcpp/grpcpp.h>
+#include "build\bazel\remote\execution\v2\remote_execution.pb.h"
+#include "build\bazel\remote\execution\v2\remote_execution.grpc.pb.h"
+UE_POP_MACRO("TEXT");
+THIRD_PARTY_INCLUDES_END
+
+#include <memory>
+
+#define LOCTEXT_NAMESPACE "BazelExecutor"
+
+
+void FBazelExecutor::Initialize(const FString& Target, const FSslCredentialsOptions& SslCredentialsOptions)
+{
+	ContentAddressableStorage.Reset();
+	Execution.Reset();
+
+	grpc::SslCredentialsOptions SslOptions;
+	SslOptions.pem_cert_chain.assign(TCHAR_TO_UTF8(*SslCredentialsOptions.PemCertChain));
+	SslOptions.pem_private_key.assign(TCHAR_TO_UTF8(*SslCredentialsOptions.PemPrivateKey));
+	SslOptions.pem_root_certs.assign(TCHAR_TO_UTF8(*SslCredentialsOptions.PemRootCerts));
+
+	// Additional configuration could be set with grpc::CreateCustomChannel()
+	std::shared_ptr<grpc::Channel> Channel = grpc::CreateChannel(TCHAR_TO_UTF8(*Target), grpc::SslCredentials(SslOptions));
+
+	ContentAddressableStorage.Reset(new FContentAddressableStorage(Channel));
+	Execution.Reset(new FExecution(Channel));
+}
+
+FName FBazelExecutor::GetFName() const
+{
+	return FName("Horde");
+}
+
+FText FBazelExecutor::GetNameText() const
+{
+	return LOCTEXT("DefaultDisplayName", "Horde");
+}
+
+FText FBazelExecutor::GetDescriptionText() const
+{
+	return LOCTEXT("DefaultDisplayDesc", "Horde remote execution.");
+}
+
+bool FBazelExecutor::CanRemoteExecute() const
+{
+	return ContentAddressableStorage.IsValid() && Execution.IsValid();
+}
+
+#undef LOCTEXT_NAMESPACE
diff --git a/Engine/Source/Editor/Experimental/BazelExecutor/Private/BazelExecutor.h b/Engine/Source/Editor/Experimental/BazelExecutor/Private/BazelExecutor.h
new file mode 100644
--- /dev/null
+++ b/Engine/Source/Editor/Experimental/BazelExecutor/Private/BazelExecutor.h
@@ -0,0 +1,38 @@
+// Copyright Epic Games, Inc. All Rights Reserved.
+
+#pragma once
+
+#include "CoreMinimal.h"
+
+#include "IRemoteExecutor.h"
+#include "IContentAddressableStorage.h"
+#include "IExecution.h"
+#include "Messages.h"
+
+
+class FBazelExecutor : public IRemoteExecutor
+{
+public:
+	struct FSslCredentialsOptions
+	{
+		FString PemCertChain;
+		FString PemPrivateKey;
+		FString PemRootCerts;
+	};
+
+private:
+	TUniquePtr<IContentAddressableStorage> ContentAddressableStorage;
+	TUniquePtr<IExecution> Execution;
+
+public:
+	void Initialize(const FString& Target, const FSslCredentialsOptions& SslCredentialsOptions);
+
+	virtual FName GetFName() const override;
+	virtual FText GetNameText() const override;
+	virtual FText GetDescriptionText() const override;
+
+	virtual bool CanRemoteExecute() const override;
+
+	IContentAddressableStorage* GetContentAddressableStorage() const override { return ContentAddressableStorage.Get(); }
+	IExecution* GetExecution() const override { return Execution.Get(); }
+};
\ No newline at end of file
diff --git a/Engine/Source/Editor/Experimental/BazelExecutor/Private/BazelExecutorModule.cpp b/Engine/Source/Editor/Experimental/BazelExecutor/Private/BazelExecutorModule.cpp
new file mode 100644
--- /dev/null
+++ b/Engine/Source/Editor/Experimental/BazelExecutor/Private/BazelExecutorModule.cpp
@@ -0,0 +1,55 @@
+// Copyright Epic Games, Inc. All Rights Reserved.
+
+#include "BazelExecutorModule.h"
+#include "ISettingsModule.h"
+#include "BazelExecutorSettings.h"
+#include "Misc/FileHelper.h"
+#include "Misc/Paths.h"
+#include "Messages.h"
+
+IMPLEMENT_MODULE(FBazelExecutorModule, EditorConfig);
+
+#define LOCTEXT_NAMESPACE "BazelExecutorModule"
+
+DEFINE_LOG_CATEGORY(LogBazelExecutor);
+
+
+void FBazelExecutorModule::StartupModule()
+{
+	GetMutableDefault<UBazelExecutorSettings>()->LoadConfig();
+
+	const FString Target = GetDefault<UBazelExecutorSettings>()->Target;
+	const FString PemCertificateChain = GetDefault<UBazelExecutorSettings>()->PemCertificateChain;
+	const FString PemPrivateKey = GetDefault<UBazelExecutorSettings>()->PemPrivateKey;
+	const FString PemRootCertificates = GetDefault<UBazelExecutorSettings>()->PemRootCertificates;
+
+	FBazelExecutor::FSslCredentialsOptions SslCredentialsOptions;
+	if (!PemCertificateChain.IsEmpty() && FPaths::FileExists(PemCertificateChain))
+	{
+		FFileHelper::LoadFileToString(SslCredentialsOptions.PemCertChain, *PemCertificateChain);
+	}
+	if (!PemPrivateKey.IsEmpty() && FPaths::FileExists(PemPrivateKey))
+	{
+		FFileHelper::LoadFileToString(SslCredentialsOptions.PemPrivateKey, *PemPrivateKey);
+	}
+	if (!PemRootCertificates.IsEmpty() && FPaths::FileExists(PemRootCertificates))
+	{
+		FFileHelper::LoadFileToString(SslCredentialsOptions.PemRootCerts, *PemRootCertificates);
+	}
+
+	HordeExecution.Initialize(Target, SslCredentialsOptions);
+
+	IModularFeatures::Get().RegisterModularFeature(TEXT("RemoteExecution"), &HordeExecution);
+}
+
+void FBazelExecutorModule::ShutdownModule()
+{
+	IModularFeatures::Get().UnregisterModularFeature(TEXT("RemoteExecution"), &HordeExecution);
+}
+
+bool FBazelExecutorModule::SupportsDynamicReloading()
+{
+	return true;
+}
+
+#undef LOCTEXT_NAMESPACE
diff --git a/Engine/Source/Editor/Experimental/BazelExecutor/Private/BazelExecutorModule.h b/Engine/Source/Editor/Experimental/BazelExecutor/Private/BazelExecutorModule.h
new file mode 100644
--- /dev/null
+++ b/Engine/Source/Editor/Experimental/BazelExecutor/Private/BazelExecutorModule.h
@@ -0,0 +1,26 @@
+// Copyright Epic Games, Inc. All Rights Reserved.
+
+#pragma once
+
+#include "CoreMinimal.h"
+#include "Modules/ModuleInterface.h"
+#include "Modules/ModuleManager.h"
+#include "Features/IModularFeatures.h"
+
+#include "BazelExecutor.h"
+
+BAZELEXECUTOR_API DECLARE_LOG_CATEGORY_EXTERN(LogBazelExecutor, Display, All);
+
+class FBazelExecutorModule : public IModuleInterface
+{
+public:
+	/** IModuleInterface implementation */
+	virtual void StartupModule() override;
+
+	virtual void ShutdownModule() override;
+
+	virtual bool SupportsDynamicReloading() override;
+
+private:
+	FBazelExecutor HordeExecution;
+};
diff --git a/Engine/Source/Editor/Experimental/BazelExecutor/Private/BazelExecutorSettings.cpp b/Engine/Source/Editor/Experimental/BazelExecutor/Private/BazelExecutorSettings.cpp
new file mode 100644
--- /dev/null
+++ b/Engine/Source/Editor/Experimental/BazelExecutor/Private/BazelExecutorSettings.cpp
@@ -0,0 +1,9 @@
+// Copyright Epic Games, Inc. All Rights Reserved.
+
+#include "BazelExecutorSettings.h"
+#include "Misc/ConfigCacheIni.h"
+
+UBazelExecutorSettings::UBazelExecutorSettings(const FObjectInitializer& ObjectInitializer)
+	: Super(ObjectInitializer)
+{
+}
diff --git a/Engine/Source/Editor/Experimental/BazelExecutor/Private/BazelExecutorSettings.h b/Engine/Source/Editor/Experimental/BazelExecutor/Private/BazelExecutorSettings.h
new file mode 100644
--- /dev/null
+++ b/Engine/Source/Editor/Experimental/BazelExecutor/Private/BazelExecutorSettings.h
@@ -0,0 +1,32 @@
+// Copyright Epic Games, Inc. All Rights Reserved.
+
+#pragma once
+
+
+#include "CoreMinimal.h"
+#include "UObject/Object.h"
+#include "BazelExecutorSettings.generated.h"
+
+UCLASS(config = EditorSettings)
+class UBazelExecutorSettings : public UObject
+{
+	GENERATED_UCLASS_BODY()
+
+public:
+	/** The horde server address. */
+	UPROPERTY(Config, EditAnywhere, Category = "Horde", meta = (DisplayName = "Target", ConfigRestartRequired = true))
+	FString Target;
+	
+	/** The horde PEM certificate chain". */
+	UPROPERTY(Config, EditAnywhere, Category = "Horde", meta = (DisplayName = "PEM Certificate Chain", ConfigRestartRequired = true))
+	FString PemCertificateChain;
+
+	/** The horde PEM private key". */
+	UPROPERTY(Config, EditAnywhere, Category = "Horde", meta = (DisplayName = "PEM Private Key", ConfigRestartRequired = true))
+	FString PemPrivateKey;
+
+	/** The horde pem root certificates. */
+	UPROPERTY(Config, EditAnywhere, Category = "Horde", meta = (DisplayName = "PEM Root Certificates", ConfigRestartRequired = true))
+	FString PemRootCertificates;
+
+};
diff --git a/Engine/Source/Editor/Experimental/BazelExecutor/Private/ContentAddressableStorage.cpp b/Engine/Source/Editor/Experimental/BazelExecutor/Private/ContentAddressableStorage.cpp
new file mode 100644
--- /dev/null
+++ b/Engine/Source/Editor/Experimental/BazelExecutor/Private/ContentAddressableStorage.cpp
@@ -0,0 +1,70 @@
+// Copyright Epic Games, Inc. All Rights Reserved.
+
+#include "ContentAddressableStorage.h"
+#include "BazelExecutorModule.h"
+#include "ProtoConverter.h"
+
+THIRD_PARTY_INCLUDES_START
+UE_PUSH_MACRO("TEXT")
+#undef TEXT
+#include <grpcpp/grpcpp.h>
+UE_POP_MACRO("TEXT");
+THIRD_PARTY_INCLUDES_END
+
+
+FContentAddressableStorage::FContentAddressableStorage(const std::shared_ptr<grpc::Channel>& Channel) :
+	Stub(build::bazel::remote::execution::v2::ContentAddressableStorage::NewStub(Channel))
+{
+}
+
+FContentAddressableStorage::~FContentAddressableStorage()
+{
+}
+
+FStatus FContentAddressableStorage::FindMissingBlobs(const FFindMissingBlobsRequest& Request, FFindMissingBlobsResponse& Response)
+{
+	grpc::ClientContext Context;
+
+	build::bazel::remote::execution::v2::FindMissingBlobsRequest ProtoRequest;
+	ProtoConverter::ToProto(Request, ProtoRequest);
+
+	build::bazel::remote::execution::v2::FindMissingBlobsResponse ProtoResponse;
+	grpc::Status ProtoStatus = Stub->FindMissingBlobs(&Context, ProtoRequest, &ProtoResponse);
+	FStatus Status;
+	ProtoConverter::FromProto(ProtoStatus, Status);
+	ProtoConverter::FromProto(ProtoResponse, Response);
+
+	return MoveTemp(Status);
+}
+
+FStatus FContentAddressableStorage::BatchUpdateBlobs(const FBatchUpdateBlobsRequest& Request, FBatchUpdateBlobsResponse& Response)
+{
+	grpc::ClientContext Context;
+
+	build::bazel::remote::execution::v2::BatchUpdateBlobsRequest ProtoRequest;
+	ProtoConverter::ToProto(Request, ProtoRequest);
+
+	build::bazel::remote::execution::v2::BatchUpdateBlobsResponse ProtoResponse;
+	grpc::Status ProtoStatus = Stub->BatchUpdateBlobs(&Context, ProtoRequest, &ProtoResponse);
+	FStatus Status;
+	ProtoConverter::FromProto(ProtoStatus, Status);
+	ProtoConverter::FromProto(ProtoResponse, Response);
+
+	return MoveTemp(Status);
+}
+
+FStatus FContentAddressableStorage::BatchReadBlobs(const FBatchReadBlobsRequest& Request, FBatchReadBlobsResponse& Response)
+{
+	grpc::ClientContext Context;
+
+	build::bazel::remote::execution::v2::BatchReadBlobsRequest ProtoRequest;
+	ProtoConverter::ToProto(Request, ProtoRequest);
+
+	build::bazel::remote::execution::v2::BatchReadBlobsResponse ProtoResponse;
+	grpc::Status ProtoStatus = Stub->BatchReadBlobs(&Context, ProtoRequest, &ProtoResponse);
+	FStatus Status;
+	ProtoConverter::FromProto(ProtoStatus, Status);
+	ProtoConverter::FromProto(ProtoResponse, Response);
+
+	return MoveTemp(Status);
+}
diff --git a/Engine/Source/Editor/Experimental/BazelExecutor/Private/ContentAddressableStorage.h b/Engine/Source/Editor/Experimental/BazelExecutor/Private/ContentAddressableStorage.h
new file mode 100644
--- /dev/null
+++ b/Engine/Source/Editor/Experimental/BazelExecutor/Private/ContentAddressableStorage.h
@@ -0,0 +1,37 @@
+// Copyright Epic Games, Inc. All Rights Reserved.
+
+#pragma once
+
+#include "CoreMinimal.h"
+#include "IContentAddressableStorage.h"
+#include "Messages.h"
+
+#include <memory>
+
+
+THIRD_PARTY_INCLUDES_START
+UE_PUSH_MACRO("TEXT")
+#undef TEXT
+#include "build\bazel\remote\execution\v2\remote_execution.grpc.pb.h"
+UE_POP_MACRO("TEXT");
+THIRD_PARTY_INCLUDES_END
+
+namespace grpc
+{
+	class Channel;
+}
+
+
+class FContentAddressableStorage : public IContentAddressableStorage
+{
+private:
+	std::unique_ptr<build::bazel::remote::execution::v2::ContentAddressableStorage::Stub> Stub;
+
+public:
+	FContentAddressableStorage(const std::shared_ptr<grpc::Channel>& Channel);
+	~FContentAddressableStorage();
+
+	FStatus FindMissingBlobs(const FFindMissingBlobsRequest& Request, FFindMissingBlobsResponse& Response) override;
+	FStatus BatchUpdateBlobs(const FBatchUpdateBlobsRequest& Request, FBatchUpdateBlobsResponse& Response) override;
+	FStatus BatchReadBlobs(const FBatchReadBlobsRequest& Request, FBatchReadBlobsResponse& Response) override;
+};
diff --git a/Engine/Source/Editor/Experimental/BazelExecutor/Private/Execution.cpp b/Engine/Source/Editor/Experimental/BazelExecutor/Private/Execution.cpp
new file mode 100644
--- /dev/null
+++ b/Engine/Source/Editor/Experimental/BazelExecutor/Private/Execution.cpp
@@ -0,0 +1,61 @@
+// Copyright Epic Games, Inc. All Rights Reserved.
+
+#include "Execution.h"
+#include "BazelExecutorModule.h"
+#include "ProtoConverter.h"
+
+THIRD_PARTY_INCLUDES_START
+UE_PUSH_MACRO("TEXT")
+#undef TEXT
+#include <grpcpp/grpcpp.h>
+UE_POP_MACRO("TEXT");
+THIRD_PARTY_INCLUDES_END
+
+
+FExecution::FExecution(const std::shared_ptr<grpc::Channel>& Channel) :
+	Stub(build::bazel::remote::execution::v2::Execution::NewStub(Channel))
+{
+}
+
+FExecution::~FExecution()
+{
+}
+
+bool FExecution::Execute(const FExecuteRequest& Request, FExecuteResponse& Response)
+{
+	grpc::ClientContext Context;
+
+	build::bazel::remote::execution::v2::ExecuteRequest ProtoRequest;
+	ProtoConverter::ToProto(Request, ProtoRequest);
+
+	std::unique_ptr<grpc::ClientReader<google::longrunning::Operation>> Call = Stub->Execute(&Context, ProtoRequest);
+	Call->WaitForInitialMetadata();
+	while (true)
+	{
+		google::longrunning::Operation Operation;
+		Call->Read(&Operation);
+
+		build::bazel::remote::execution::v2::ExecuteOperationMetadata Metadata;
+		Operation.metadata().UnpackTo(&Metadata);
+
+		UE_LOG(LogBazelExecutor, Display, TEXT("Execution state: %s"), UTF8_TO_TCHAR(build::bazel::remote::execution::v2::ExecutionStage_Value_Name(Metadata.stage()).c_str()));
+
+		if (Operation.done())
+		{
+			build::bazel::remote::execution::v2::ExecuteResponse ProtoResponse;
+			if (!Operation.response().UnpackTo(&ProtoResponse))
+			{
+				UE_LOG(LogBazelExecutor, Error, TEXT("Unable to decode response"));
+				break;
+			}
+			if ((grpc::StatusCode)ProtoResponse.status().code() != grpc::StatusCode::OK) {
+				UE_LOG(LogBazelExecutor, Error, TEXT("%s"), UTF8_TO_TCHAR(ProtoResponse.status().message().c_str()));
+				break;
+			}
+			ProtoConverter::FromProto(ProtoResponse, Response);
+			return true;
+		}
+	}
+
+	return false;
+}
diff --git a/Engine/Source/Editor/Experimental/BazelExecutor/Private/Execution.h b/Engine/Source/Editor/Experimental/BazelExecutor/Private/Execution.h
new file mode 100644
--- /dev/null
+++ b/Engine/Source/Editor/Experimental/BazelExecutor/Private/Execution.h
@@ -0,0 +1,33 @@
+// Copyright Epic Games, Inc. All Rights Reserved.
+
+#pragma once
+
+#include "CoreMinimal.h"
+#include "IExecution.h"
+#include "Messages.h"
+
+THIRD_PARTY_INCLUDES_START
+UE_PUSH_MACRO("TEXT")
+#undef TEXT
+#include <memory>
+#include "build\bazel\remote\execution\v2\remote_execution.grpc.pb.h"
+UE_POP_MACRO("TEXT");
+THIRD_PARTY_INCLUDES_END
+
+namespace grpc
+{
+	class Channel;
+}
+
+
+class FExecution : public IExecution
+{
+private:
+	std::unique_ptr<build::bazel::remote::execution::v2::Execution::Stub> Stub;
+
+public:
+	FExecution(const std::shared_ptr<grpc::Channel>& Channel);
+	~FExecution();
+
+	bool Execute(const FExecuteRequest& Request, FExecuteResponse& Response) override;
+};
diff --git a/Engine/Source/Editor/Experimental/BazelExecutor/Private/ProtoConverter.cpp b/Engine/Source/Editor/Experimental/BazelExecutor/Private/ProtoConverter.cpp
new file mode 100644
--- /dev/null
+++ b/Engine/Source/Editor/Experimental/BazelExecutor/Private/ProtoConverter.cpp
@@ -0,0 +1,450 @@
+// Copyright Epic Games, Inc. All Rights Reserved.
+
+#include "ProtoConverter.h"
+#include "Misc/Paths.h"
+
+THIRD_PARTY_INCLUDES_START
+#include <openssl/sha.h>
+THIRD_PARTY_INCLUDES_END
+
+
+void ProtoConverter::ToProto(const FString& In, std::string& Out)
+{
+	Out.assign(TCHAR_TO_UTF8(*In));
+}
+
+void ProtoConverter::ToProto(const FDuration& In, google::protobuf::Duration& Out)
+{
+	Out.set_seconds(In.Seconds);
+	Out.set_nanos(In.Nanos);
+}
+
+void ProtoConverter::ToProto(const FTimestamp& In, google::protobuf::Timestamp& Out)
+{
+	Out.set_seconds(In.Seconds);
+	Out.set_nanos(In.Nanos);
+}
+
+void ProtoConverter::ToProto(const FDigest& In, build::bazel::remote::execution::v2::Digest& Out)
+{
+	ToProto(In.Hash, *Out.mutable_hash());
+	Out.set_size_bytes(In.SizeBytes);
+}
+
+void ProtoConverter::ToProto(const FNodeProperty& In, build::bazel::remote::execution::v2::NodeProperty& Out)
+{
+	ToProto(In.Name, *Out.mutable_name());
+	ToProto(In.Value, *Out.mutable_value());
+}
+
+void ProtoConverter::ToProto(const FNodeProperties& In, build::bazel::remote::execution::v2::NodeProperties& Out)
+{
+	for (const FNodeProperty& InProperty : In.Properties)
+	{
+		ToProto(InProperty, *Out.add_properties());
+	}
+	if (In.ModifiedTime.Nanos || In.ModifiedTime.Seconds)
+	{
+		ToProto(In.ModifiedTime, *Out.mutable_mtime());
+	}
+	if (In.UnixMode)
+	{
+		Out.mutable_unix_mode()->set_value(In.UnixMode);
+	}
+}
+
+void ProtoConverter::ToProto(const FFileNode& In, build::bazel::remote::execution::v2::FileNode& Out)
+{
+	ToProto(In.Name, *Out.mutable_name());
+	ToProto(In.Digest, *Out.mutable_digest());
+	Out.set_is_executable(In.IsExecutable);
+	if (In.NodeProperties.ModifiedTime.Seconds || In.NodeProperties.ModifiedTime.Nanos || !In.NodeProperties.Properties.IsEmpty() || In.NodeProperties.UnixMode)
+	{
+		ToProto(In.NodeProperties, *Out.mutable_node_properties());
+	}
+}
+
+void ProtoConverter::ToProto(const FDirectoryNode& In, build::bazel::remote::execution::v2::DirectoryNode& Out)
+{
+	ToProto(In.Name, *Out.mutable_name());
+	ToProto(In.Digest, *Out.mutable_digest());
+}
+
+void ProtoConverter::ToProto(const FSymlinkNode& In, build::bazel::remote::execution::v2::SymlinkNode& Out)
+{
+	ToProto(In.Name, *Out.mutable_name());
+	ToProto(In.Target, *Out.mutable_target());
+	if (In.NodeProperties.ModifiedTime.Seconds || In.NodeProperties.ModifiedTime.Nanos || !In.NodeProperties.Properties.IsEmpty() || In.NodeProperties.UnixMode)
+	{
+		ToProto(In.NodeProperties, *Out.mutable_node_properties());
+	}
+}
+
+void ProtoConverter::ToProto(const FDirectory& In, build::bazel::remote::execution::v2::Directory& Out)
+{
+	for (const FFileNode& InFile : In.Files)
+	{
+		ToProto(InFile, *Out.add_files());
+	}
+
+	for (const FDirectoryNode& InDirectory : In.Directories)
+	{
+		ToProto(InDirectory, *Out.add_directories());
+	}
+
+	for (const FSymlinkNode& InSymlink : In.Symlinks)
+	{
+		ToProto(InSymlink, *Out.add_symlinks());
+	}
+	if (In.NodeProperties.ModifiedTime.Seconds || In.NodeProperties.ModifiedTime.Nanos || !In.NodeProperties.Properties.IsEmpty() || In.NodeProperties.UnixMode)
+	{
+		ToProto(In.NodeProperties, *Out.mutable_node_properties());
+	}
+}
+
+void ProtoConverter::ToProto(const FPlatform& In, build::bazel::remote::execution::v2::Platform& Out)
+{
+	for (const FPlatform::FProperty& InProperty : In.Properties)
+	{
+		build::bazel::remote::execution::v2::Platform_Property* OutProperty = Out.add_properties();
+		ToProto(InProperty.Name, *OutProperty->mutable_name());
+		ToProto(InProperty.Value, *OutProperty->mutable_value());
+	}
+}
+
+void ProtoConverter::ToProto(const FAction& In, build::bazel::remote::execution::v2::Action& Out)
+{
+	ToProto(In.CommandDigest, *Out.mutable_command_digest());
+	ToProto(In.InputRootDigest, *Out.mutable_input_root_digest());
+	if (In.Timeout.Seconds || In.Timeout.Nanos)
+	{
+		ToProto(In.Timeout, *Out.mutable_timeout());
+	}
+	Out.set_do_not_cache(In.DoNotCache);
+	if (!In.Salt.IsEmpty())
+	{
+		ToProto(In.Salt, *Out.mutable_salt());
+	}
+	if (!In.Platform.Properties.IsEmpty())
+	{
+		ToProto(In.Platform, *Out.mutable_platform());
+	}
+}
+
+void ProtoConverter::ToProto(const FCommand& In, build::bazel::remote::execution::v2::Command& Out)
+{
+	for (const FString& InArgument : In.Arguments)
+	{
+		ToProto(InArgument, *Out.add_arguments());
+	}
+
+	for (const FCommand::FEnvironmentVariable& InEnvironmentVariable : In.EnvironmentVariables)
+	{
+		build::bazel::remote::execution::v2::Command_EnvironmentVariable* OutEnvironmentVariable = Out.add_environment_variables();
+		ToProto(InEnvironmentVariable.Name, *OutEnvironmentVariable->mutable_name());
+		ToProto(InEnvironmentVariable.Value, *OutEnvironmentVariable->mutable_value());
+	}
+
+	for (const FString& InOutputPath : In.OutputPaths)
+	{
+		ToProto(InOutputPath, *Out.add_output_paths());
+	}
+
+	ToProto(In.WorkingDirectory, *Out.mutable_working_directory());
+
+	for (const FString& InOutputNodeProperty : In.OutputNodeProperties)
+	{
+		ToProto(InOutputNodeProperty, *Out.add_output_node_properties());
+	}
+}
+
+void ProtoConverter::FromProto(const std::string& In, FString& Out)
+{
+	Out = UTF8_TO_TCHAR(In.c_str());
+}
+
+void ProtoConverter::FromProto(const std::string& In, TArray<char>& Out)
+{
+	Out.Empty();
+	if (In.size() > 0)
+	{
+		Out.AddUninitialized(In.size());
+		memcpy(Out.GetData(), In.c_str(), In.size());
+	}
+}
+
+void ProtoConverter::FromProto(const google::protobuf::Timestamp& In, FTimestamp& Out)
+{
+	Out.Seconds = In.seconds();
+	Out.Nanos = In.nanos();
+}
+
+void ProtoConverter::FromProto(const build::bazel::remote::execution::v2::Digest& In, FDigest& Out)
+{
+	FromProto(In.hash(), Out.Hash);
+	Out.SizeBytes = In.size_bytes();
+}
+
+void ProtoConverter::FromProto(const build::bazel::remote::execution::v2::NodeProperty& In, FNodeProperty& Out)
+{
+	FromProto(In.name(), Out.Name);
+	FromProto(In.value(), Out.Value);
+}
+
+void ProtoConverter::FromProto(const build::bazel::remote::execution::v2::NodeProperties& In, FNodeProperties& Out)
+{
+	for (const build::bazel::remote::execution::v2::NodeProperty& InProperty : In.properties())
+	{
+		FNodeProperty OutProperty;
+		FromProto(InProperty, OutProperty);
+		Out.Properties.Add(MoveTemp(OutProperty));
+	}
+	FromProto(In.mtime(), Out.ModifiedTime);
+	Out.UnixMode = In.unix_mode().value();
+}
+
+void ProtoConverter::FromProto(const build::bazel::remote::execution::v2::LogFile& In, FLogFile& Out)
+{
+	FromProto(In.digest(), Out.Digest);
+	Out.HumanReadable = In.human_readable();
+}
+
+void ProtoConverter::FromProto(const build::bazel::remote::execution::v2::OutputFile& In, FOutputFile& Out)
+{
+	FromProto(In.path(), Out.Path);
+	FromProto(In.digest(), Out.Digest);
+	Out.IsExecutable = In.is_executable();
+	FromProto(In.contents(), Out.Contents);
+	FromProto(In.node_properties(), Out.NodeProperties);
+}
+
+void ProtoConverter::FromProto(const build::bazel::remote::execution::v2::OutputSymlink& In, FOutputSymlink& Out)
+{
+	FromProto(In.path(), Out.Path);
+	FromProto(In.target(), Out.Target);
+	FromProto(In.node_properties(), Out.NodeProperties);
+}
+
+void ProtoConverter::FromProto(const build::bazel::remote::execution::v2::OutputDirectory& In, FOutputDirectory& Out)
+{
+	FromProto(In.path(), Out.Path);
+	FromProto(In.tree_digest(), Out.TreeDigest);
+}
+
+void ProtoConverter::FromProto(const build::bazel::remote::execution::v2::ExecutedActionMetadata& In, FExecutedActionMetadata& Out)
+{
+	FromProto(In.worker(), Out.Worker);
+	FromProto(In.queued_timestamp(), Out.QueuedTimestamp);
+	FromProto(In.worker_start_timestamp(), Out.WorkerStartTimestamp);
+	FromProto(In.worker_completed_timestamp(), Out.WorkerCompletedTimestamp);
+	FromProto(In.input_fetch_start_timestamp(), Out.InputFetchStartTimestamp);
+	FromProto(In.input_fetch_completed_timestamp(), Out.InputFetchCompletedTimestamp);
+	FromProto(In.execution_start_timestamp(), Out.ExecutionStartTimestamp);
+	FromProto(In.execution_completed_timestamp(), Out.ExecutionCompletedTimestamp);
+	FromProto(In.output_upload_start_timestamp(), Out.OutputUploadStartTimestamp);
+	FromProto(In.output_upload_completed_timestamp(), Out.OutputUploadCompletedTimestamp);
+}
+
+void ProtoConverter::FromProto(const build::bazel::remote::execution::v2::ActionResult& In, FActionResult& Out)
+{
+	for (const build::bazel::remote::execution::v2::OutputFile& InOutputFile : In.output_files())
+	{
+		FOutputFile OutOutputFile;
+		FromProto(InOutputFile, OutOutputFile);
+		Out.OutputFiles.Add(MoveTemp(OutOutputFile));
+	}
+
+	for (const build::bazel::remote::execution::v2::OutputSymlink& InOutputSymlink : In.output_symlinks())
+	{
+		FOutputSymlink OutOutputSymlink;
+		FromProto(InOutputSymlink, OutOutputSymlink);
+		Out.OutputSymlinks.Add(MoveTemp(OutOutputSymlink));
+	}
+
+	for (const build::bazel::remote::execution::v2::OutputDirectory& InOutputDirectory : In.output_directories())
+	{
+		FOutputDirectory OutOutputDirectory;
+		FromProto(InOutputDirectory, OutOutputDirectory);
+		Out.OutputDirectories.Add(MoveTemp(OutOutputDirectory));
+	}
+
+	Out.ExitCode = In.exit_code();
+	FromProto(In.stdout_raw(), Out.StdoutRaw);
+	FromProto(In.stdout_digest(), Out.StdoutDigest);
+	FromProto(In.stderr_raw(), Out.StderrRaw);
+	FromProto(In.stderr_digest(), Out.StderrDigest);
+	FromProto(In.execution_metadata(), Out.ExecutionMetadata);
+}
+
+void ProtoConverter::ToProto(const FExecuteRequest& In, build::bazel::remote::execution::v2::ExecuteRequest& Out)
+{
+	if (!In.InstanceName.IsEmpty())
+	{
+		ToProto(In.InstanceName, *Out.mutable_instance_name());
+	}
+	Out.set_skip_cache_lookup(In.SkipCacheLookup);
+	ToProto(In.ActionDigest, *Out.mutable_action_digest());
+	if (In.ExecutionPolicy.Priority)
+	{
+		Out.mutable_execution_policy()->set_priority(In.ExecutionPolicy.Priority);
+	}
+	if (In.ResultsCachePolicy.Priority)
+	{
+		Out.mutable_results_cache_policy()->set_priority(In.ResultsCachePolicy.Priority);
+	}
+}
+
+void ProtoConverter::ToProto(const FFindMissingBlobsRequest& In, build::bazel::remote::execution::v2::FindMissingBlobsRequest& Out)
+{
+	if (!In.InstanceName.IsEmpty())
+	{
+		ToProto(In.InstanceName, *Out.mutable_instance_name());
+	}
+	for	(const FDigest& InDigest : In.BlobDigests)
+	{
+		ToProto(InDigest, *Out.add_blob_digests());
+	}
+}
+
+void ProtoConverter::ToProto(const FBatchUpdateBlobsRequest& In, build::bazel::remote::execution::v2::BatchUpdateBlobsRequest& Out)
+{
+	if (!In.InstanceName.IsEmpty())
+	{
+		ToProto(In.InstanceName, *Out.mutable_instance_name());
+	}
+	for (const FBatchUpdateBlobsRequest::FRequest& InRequest : In.Requests)
+	{
+		build::bazel::remote::execution::v2::BatchUpdateBlobsRequest_Request* OutRequest = Out.add_requests();
+		ToProto(InRequest.Digest, *OutRequest->mutable_digest());
+		OutRequest->mutable_data()->assign(InRequest.Data.GetData(), InRequest.Data.Num());
+	}
+}
+
+void ProtoConverter::ToProto(const FBatchReadBlobsRequest& In, build::bazel::remote::execution::v2::BatchReadBlobsRequest& Out)
+{
+	if (!In.InstanceName.IsEmpty())
+	{
+		ToProto(In.InstanceName, *Out.mutable_instance_name());
+	}
+	for (const FDigest& InDigest : In.Digests)
+	{
+		ToProto(InDigest, *Out.add_digests());
+	}
+}
+
+void ProtoConverter::FromProto(const google::rpc::Status& In, FStatus& Out)
+{
+	Out.Code = (EStatusCode)In.code();
+	FromProto(In.message(), Out.Message);
+	Out.Details.Empty();
+}
+
+void ProtoConverter::FromProto(const grpc::Status& In, FStatus& Out)
+{
+	Out.Code = (EStatusCode)In.error_code();
+	FromProto(In.error_message(), Out.Message);
+	FromProto(In.error_details(), Out.Details);
+}
+
+void ProtoConverter::FromProto(const build::bazel::remote::execution::v2::FindMissingBlobsResponse& In, FFindMissingBlobsResponse& Out)
+{
+	for (const build::bazel::remote::execution::v2::Digest& InDigest : In.missing_blob_digests())
+	{
+		FDigest OutDigest;
+		FromProto(InDigest, OutDigest);
+		Out.MissingBlobDigests.Add(MoveTemp(OutDigest));
+	}
+}
+
+void ProtoConverter::FromProto(const build::bazel::remote::execution::v2::BatchUpdateBlobsResponse& In, FBatchUpdateBlobsResponse& Out)
+{
+	for (const build::bazel::remote::execution::v2::BatchUpdateBlobsResponse_Response& InResponse : In.responses())
+	{
+		FBatchUpdateBlobsResponse::FResponse OutResponse;
+		FromProto(InResponse.digest(), OutResponse.Digest);
+		FromProto(InResponse.status(), OutResponse.Status);
+	}
+}
+
+void ProtoConverter::FromProto(const build::bazel::remote::execution::v2::BatchReadBlobsResponse& In, FBatchReadBlobsResponse& Out)
+{
+	for (const build::bazel::remote::execution::v2::BatchReadBlobsResponse_Response& InResponse : In.responses())
+	{
+		FBatchReadBlobsResponse::FResponse OutResponse;
+		FromProto(InResponse.digest(), OutResponse.Digest);
+		OutResponse.Data.Empty();
+		OutResponse.Data.Append(InResponse.data().c_str(), InResponse.data().size());
+		FromProto(InResponse.status(), OutResponse.Status);
+		Out.Responses.Add(MoveTemp(OutResponse));
+	}
+}
+
+void ProtoConverter::FromProto(const build::bazel::remote::execution::v2::ExecuteResponse& In, FExecuteResponse& Out)
+{
+	FromProto(In.result(), Out.Result);
+	Out.CachedResult = In.cached_result();
+	FromProto(In.status(), Out.Status);
+
+	for (const google::protobuf::MapPair<std::string, build::bazel::remote::execution::v2::LogFile>& InServerLog : In.server_logs())
+	{
+		FString Key;
+		FLogFile Value;
+		FromProto(InServerLog.first, Key);
+		FromProto(InServerLog.second, Value);
+		Out.ServerLogs.Add(MoveTemp(Key), MoveTemp(Value));
+	}
+
+	FromProto(In.message(), Out.Message);
+}
+
+bool ProtoConverter::ToDigest(const TArray<char>& Data, FDigest& OutDigest)
+{
+	unsigned char Sha256Buffer[SHA256_DIGEST_LENGTH];
+
+	SHA256_CTX sha256;
+	SHA256_Init(&sha256);
+	SHA256_Update(&sha256, Data.GetData(), Data.Num());
+	SHA256_Final(Sha256Buffer, &sha256);
+
+	// GRPC requires lowercase hex
+	OutDigest.Hash = FString::FromHexBlob(Sha256Buffer, SHA256_DIGEST_LENGTH).ToLower();
+	OutDigest.SizeBytes = Data.Num();
+	return true;
+}
+
+bool ProtoConverter::ToBlob(const google::protobuf::Message& InMessage, TArray<char>& OutData, FDigest& OutDigest)
+{
+	OutData.Empty();
+	OutData.AddUninitialized(InMessage.ByteSizeLong());
+	if (!InMessage.SerializeToArray(OutData.GetData(), OutData.Num()))
+	{
+		OutData.Empty();
+		return false;
+	}
+	return ToDigest(OutData, OutDigest);
+}
+
+bool ProtoConverter::ToBlob(const FDirectory& InDirectory, TArray<char>& OutData, FDigest& OutDigest)
+{
+	// Convert to temporary proto Directory
+	build::bazel::remote::execution::v2::Directory OutDirectory;
+	ToProto(InDirectory, OutDirectory);
+	return ToBlob(OutDirectory, OutData, OutDigest);
+}
+
+bool ProtoConverter::ToBlob(const FCommand& InCommand, TArray<char>& OutData, FDigest& OutDigest)
+{
+	// Convert to temporary proto Action
+	build::bazel::remote::execution::v2::Command OutCommand;
+	ToProto(InCommand, OutCommand);
+	return ToBlob(OutCommand, OutData, OutDigest);
+}
+
+bool ProtoConverter::ToBlob(const FAction& InAction, TArray<char>& OutData, FDigest& OutDigest)
+{
+	// Convert to temporary proto Action
+	build::bazel::remote::execution::v2::Action OutAction;
+	ToProto(InAction, OutAction);
+	return ToBlob(OutAction, OutData, OutDigest);
+}
+
diff --git a/Engine/Source/Editor/Experimental/BazelExecutor/Private/ProtoConverter.h b/Engine/Source/Editor/Experimental/BazelExecutor/Private/ProtoConverter.h
new file mode 100644
--- /dev/null
+++ b/Engine/Source/Editor/Experimental/BazelExecutor/Private/ProtoConverter.h
@@ -0,0 +1,78 @@
+// Copyright Epic Games, Inc. All Rights Reserved.
+
+#pragma once
+
+#include "CoreMinimal.h"
+#include "Messages.h"
+
+THIRD_PARTY_INCLUDES_START
+UE_PUSH_MACRO("TEXT")
+#undef TEXT
+#include <grpcpp/grpcpp.h>
+#include "build\bazel\remote\execution\v2\remote_execution.pb.h"
+#include "build\bazel\remote\execution\v2\remote_execution.grpc.pb.h"
+UE_POP_MACRO("TEXT");
+THIRD_PARTY_INCLUDES_END
+
+
+class ProtoConverter
+{
+private:
+	static void ToProto(const FString& In, std::string& Out);
+
+	static void ToProto(const FDuration& In, google::protobuf::Duration& Out);
+	static void ToProto(const FTimestamp& In, google::protobuf::Timestamp& Out);
+
+	static void ToProto(const FDigest& In, build::bazel::remote::execution::v2::Digest& Out);
+
+	static void ToProto(const FNodeProperty& In, build::bazel::remote::execution::v2::NodeProperty& Out);
+	static void ToProto(const FNodeProperties& In, build::bazel::remote::execution::v2::NodeProperties& Out);
+
+	static void ToProto(const FFileNode& In, build::bazel::remote::execution::v2::FileNode& Out);
+	static void ToProto(const FDirectoryNode& In, build::bazel::remote::execution::v2::DirectoryNode& Out);
+	static void ToProto(const FSymlinkNode& In, build::bazel::remote::execution::v2::SymlinkNode& Out);
+	static void ToProto(const FDirectory& In, build::bazel::remote::execution::v2::Directory& Out);
+
+	static void ToProto(const FPlatform& In, build::bazel::remote::execution::v2::Platform& Out);
+	static void ToProto(const FAction& In, build::bazel::remote::execution::v2::Action& Out);
+	static void ToProto(const FCommand& In, build::bazel::remote::execution::v2::Command& Out);
+
+	static void FromProto(const std::string& In, FString& Out);
+	static void FromProto(const std::string& In, TArray<char>& Out);
+
+	static void FromProto(const google::protobuf::Timestamp& In, FTimestamp& Out);
+
+	static void FromProto(const build::bazel::remote::execution::v2::Digest& In, FDigest& Out);
+
+	static void FromProto(const build::bazel::remote::execution::v2::NodeProperty& In, FNodeProperty& Out);
+	static void FromProto(const build::bazel::remote::execution::v2::NodeProperties& In, FNodeProperties& Out);
+
+	static void FromProto(const build::bazel::remote::execution::v2::LogFile& In, FLogFile& Out);
+	static void FromProto(const build::bazel::remote::execution::v2::OutputFile& In, FOutputFile& Out);
+	static void FromProto(const build::bazel::remote::execution::v2::OutputSymlink& In, FOutputSymlink& Out);
+	static void FromProto(const build::bazel::remote::execution::v2::OutputDirectory& In, FOutputDirectory& Out);
+	static void FromProto(const build::bazel::remote::execution::v2::ExecutedActionMetadata& In, FExecutedActionMetadata& Out);
+	static void FromProto(const build::bazel::remote::execution::v2::ActionResult& In, FActionResult& Out);
+
+public:
+	static void ToProto(const FExecuteRequest& In, build::bazel::remote::execution::v2::ExecuteRequest& Out);
+	static void ToProto(const FFindMissingBlobsRequest& In, build::bazel::remote::execution::v2::FindMissingBlobsRequest& Out);
+	static void ToProto(const FBatchUpdateBlobsRequest& In, build::bazel::remote::execution::v2::BatchUpdateBlobsRequest& Out);
+	static void ToProto(const FBatchReadBlobsRequest& In, build::bazel::remote::execution::v2::BatchReadBlobsRequest& Out);
+
+	static void FromProto(const google::rpc::Status& In, FStatus& Out);
+	static void FromProto(const grpc::Status& In, FStatus& Out);
+	static void FromProto(const build::bazel::remote::execution::v2::FindMissingBlobsResponse& In, FFindMissingBlobsResponse& Out);
+	static void FromProto(const build::bazel::remote::execution::v2::BatchUpdateBlobsResponse& In, FBatchUpdateBlobsResponse& Out);
+	static void FromProto(const build::bazel::remote::execution::v2::BatchReadBlobsResponse& In, FBatchReadBlobsResponse& Out);
+	static void FromProto(const build::bazel::remote::execution::v2::ExecuteResponse& In, FExecuteResponse& Out);
+
+	static bool ToDigest(const TArray<char>& Data, FDigest& OutDigest);
+
+	static bool ToBlob(const google::protobuf::Message& InMessage, TArray<char>& OutData, FDigest& OutDigest);
+
+	static bool ToBlob(const FDirectory& InDirectory, TArray<char>& OutData, FDigest& OutDigest);
+	static bool ToBlob(const FCommand& InCommand, TArray<char>& OutData, FDigest& OutDigest);
+	static bool ToBlob(const FAction& InAction, TArray<char>& OutData, FDigest& OutDigest);
+};
+
diff --git a/Engine/Source/Editor/Experimental/BazelExecutor/SchemaFooter.txt b/Engine/Source/Editor/Experimental/BazelExecutor/SchemaFooter.txt
new file mode 100644
--- /dev/null
+++ b/Engine/Source/Editor/Experimental/BazelExecutor/SchemaFooter.txt
@@ -0,0 +1,3 @@
+
+UE_POP_MACRO("TEXT")
+THIRD_PARTY_INCLUDES_END
diff --git a/Engine/Source/Editor/Experimental/BazelExecutor/SchemaHeader.txt b/Engine/Source/Editor/Experimental/BazelExecutor/SchemaHeader.txt
new file mode 100644
--- /dev/null
+++ b/Engine/Source/Editor/Experimental/BazelExecutor/SchemaHeader.txt
@@ -0,0 +1,8 @@
+//-V::568,590
+#include "CoreMinimal.h"
+#include "Misc/CoreMiscDefines.h"
+
+THIRD_PARTY_INCLUDES_START
+UE_PUSH_MACRO("TEXT")
+#undef TEXT
+
diff --git a/Engine/Source/Editor/Experimental/HordeExecutor/Generate.bat b/Engine/Source/Editor/Experimental/HordeExecutor/Generate.bat
deleted file mode 100755
diff --git a/Engine/Source/Editor/Experimental/HordeExecutor/Generate.sh b/Engine/Source/Editor/Experimental/HordeExecutor/Generate.sh
deleted file mode 100644
diff --git a/Engine/Source/Editor/Experimental/HordeExecutor/HordeExecutor.Build.cs b/Engine/Source/Editor/Experimental/HordeExecutor/HordeExecutor.Build.cs
deleted file mode 100644
diff --git a/Engine/Source/Editor/Experimental/HordeExecutor/Private/ContentAddressableStorage.cpp b/Engine/Source/Editor/Experimental/HordeExecutor/Private/ContentAddressableStorage.cpp
deleted file mode 100644
diff --git a/Engine/Source/Editor/Experimental/HordeExecutor/Private/ContentAddressableStorage.h b/Engine/Source/Editor/Experimental/HordeExecutor/Private/ContentAddressableStorage.h
deleted file mode 100644
diff --git a/Engine/Source/Editor/Experimental/HordeExecutor/Private/Execution.cpp b/Engine/Source/Editor/Experimental/HordeExecutor/Private/Execution.cpp
deleted file mode 100644
diff --git a/Engine/Source/Editor/Experimental/HordeExecutor/Private/Execution.h b/Engine/Source/Editor/Experimental/HordeExecutor/Private/Execution.h
deleted file mode 100644
diff --git a/Engine/Source/Editor/Experimental/HordeExecutor/Private/HordeExecutor.cpp b/Engine/Source/Editor/Experimental/HordeExecutor/Private/HordeExecutor.cpp
deleted file mode 100644
diff --git a/Engine/Source/Editor/Experimental/HordeExecutor/Private/HordeExecutor.h b/Engine/Source/Editor/Experimental/HordeExecutor/Private/HordeExecutor.h
deleted file mode 100644
diff --git a/Engine/Source/Editor/Experimental/HordeExecutor/Private/HordeExecutorModule.cpp b/Engine/Source/Editor/Experimental/HordeExecutor/Private/HordeExecutorModule.cpp
deleted file mode 100644
diff --git a/Engine/Source/Editor/Experimental/HordeExecutor/Private/HordeExecutorModule.h b/Engine/Source/Editor/Experimental/HordeExecutor/Private/HordeExecutorModule.h
deleted file mode 100644
diff --git a/Engine/Source/Editor/Experimental/HordeExecutor/Private/HordeExecutorSettings.cpp b/Engine/Source/Editor/Experimental/HordeExecutor/Private/HordeExecutorSettings.cpp
deleted file mode 100644
diff --git a/Engine/Source/Editor/Experimental/HordeExecutor/Private/HordeExecutorSettings.h b/Engine/Source/Editor/Experimental/HordeExecutor/Private/HordeExecutorSettings.h
deleted file mode 100644
diff --git a/Engine/Source/Editor/Experimental/HordeExecutor/Private/ProtoConverter.cpp b/Engine/Source/Editor/Experimental/HordeExecutor/Private/ProtoConverter.cpp
deleted file mode 100644
diff --git a/Engine/Source/Editor/Experimental/HordeExecutor/Private/ProtoConverter.h b/Engine/Source/Editor/Experimental/HordeExecutor/Private/ProtoConverter.h
deleted file mode 100644
diff --git a/Engine/Source/Editor/Experimental/HordeExecutor/SchemaAutogenHeader.txt b/Engine/Source/Editor/Experimental/HordeExecutor/SchemaAutogenHeader.txt
deleted file mode 100644
diff --git a/Engine/Source/Editor/Experimental/HordeExecutor/SchemaFooter.txt b/Engine/Source/Editor/Experimental/HordeExecutor/SchemaFooter.txt
deleted file mode 100644
diff --git a/Engine/Source/Editor/Experimental/HordeExecutor/SchemaHeader.txt b/Engine/Source/Editor/Experimental/HordeExecutor/SchemaHeader.txt
deleted file mode 100644
