A few missed HordeExecutor -> Bazel rename (strings, comments)

diff --git a/Engine/Source/Editor/Experimental/BazelExecutor/Private/BazelExecutor.cpp b/Engine/Source/Editor/Experimental/BazelExecutor/Private/BazelExecutor.cpp
--- a/Engine/Source/Editor/Experimental/BazelExecutor/Private/BazelExecutor.cpp
+++ b/Engine/Source/Editor/Experimental/BazelExecutor/Private/BazelExecutor.cpp
@@ -38,17 +38,17 @@ void FBazelExecutor::Initialize(const FString& Target, const FSslCredentialsOpti
 
 FName FBazelExecutor::GetFName() const
 {
-	return FName("Horde");
+	return FName("Bazel");
 }
 
 FText FBazelExecutor::GetNameText() const
 {
-	return LOCTEXT("DefaultDisplayName", "Horde");
+	return LOCTEXT("DefaultDisplayName", "Bazel");
 }
 
 FText FBazelExecutor::GetDescriptionText() const
 {
-	return LOCTEXT("DefaultDisplayDesc", "Horde remote execution.");
+	return LOCTEXT("DefaultDisplayDesc", "Bazel remote execution.");
 }
 
 bool FBazelExecutor::CanRemoteExecute() const
diff --git a/Engine/Source/Editor/Experimental/BazelExecutor/Private/BazelExecutorModule.cpp b/Engine/Source/Editor/Experimental/BazelExecutor/Private/BazelExecutorModule.cpp
--- a/Engine/Source/Editor/Experimental/BazelExecutor/Private/BazelExecutorModule.cpp
+++ b/Engine/Source/Editor/Experimental/BazelExecutor/Private/BazelExecutorModule.cpp
@@ -37,14 +37,14 @@ void FBazelExecutorModule::StartupModule()
 		FFileHelper::LoadFileToString(SslCredentialsOptions.PemRootCerts, *PemRootCertificates);
 	}
 
-	HordeExecution.Initialize(Target, SslCredentialsOptions);
+	BazelExecution.Initialize(Target, SslCredentialsOptions);
 
-	IModularFeatures::Get().RegisterModularFeature(TEXT("RemoteExecution"), &HordeExecution);
+	IModularFeatures::Get().RegisterModularFeature(TEXT("RemoteExecution"), &BazelExecution);
 }
 
 void FBazelExecutorModule::ShutdownModule()
 {
-	IModularFeatures::Get().UnregisterModularFeature(TEXT("RemoteExecution"), &HordeExecution);
+	IModularFeatures::Get().UnregisterModularFeature(TEXT("RemoteExecution"), &BazelExecution);
 }
 
 bool FBazelExecutorModule::SupportsDynamicReloading()
diff --git a/Engine/Source/Editor/Experimental/BazelExecutor/Private/BazelExecutorModule.h b/Engine/Source/Editor/Experimental/BazelExecutor/Private/BazelExecutorModule.h
--- a/Engine/Source/Editor/Experimental/BazelExecutor/Private/BazelExecutorModule.h
+++ b/Engine/Source/Editor/Experimental/BazelExecutor/Private/BazelExecutorModule.h
@@ -22,5 +22,5 @@ public:
 	virtual bool SupportsDynamicReloading() override;
 
 private:
-	FBazelExecutor HordeExecution;
+	FBazelExecutor BazelExecution;
 };
diff --git a/Engine/Source/Editor/Experimental/BazelExecutor/Private/BazelExecutorSettings.h b/Engine/Source/Editor/Experimental/BazelExecutor/Private/BazelExecutorSettings.h
--- a/Engine/Source/Editor/Experimental/BazelExecutor/Private/BazelExecutorSettings.h
+++ b/Engine/Source/Editor/Experimental/BazelExecutor/Private/BazelExecutorSettings.h
@@ -13,20 +13,20 @@ class UBazelExecutorSettings : public UObject
 	GENERATED_UCLASS_BODY()
 
 public:
-	/** The horde server address. */
-	UPROPERTY(Config, EditAnywhere, Category = "Horde", meta = (DisplayName = "Target", ConfigRestartRequired = true))
+	/** The Bazel server address. */
+	UPROPERTY(Config, EditAnywhere, Category = "Bazel", meta = (DisplayName = "Target", ConfigRestartRequired = true))
 	FString Target;
 	
-	/** The horde PEM certificate chain". */
-	UPROPERTY(Config, EditAnywhere, Category = "Horde", meta = (DisplayName = "PEM Certificate Chain", ConfigRestartRequired = true))
+	/** The Bazel PEM certificate chain". */
+	UPROPERTY(Config, EditAnywhere, Category = "Bazel", meta = (DisplayName = "PEM Certificate Chain", ConfigRestartRequired = true))
 	FString PemCertificateChain;
 
-	/** The horde PEM private key". */
-	UPROPERTY(Config, EditAnywhere, Category = "Horde", meta = (DisplayName = "PEM Private Key", ConfigRestartRequired = true))
+	/** The Bazel PEM private key". */
+	UPROPERTY(Config, EditAnywhere, Category = "Bazel", meta = (DisplayName = "PEM Private Key", ConfigRestartRequired = true))
 	FString PemPrivateKey;
 
-	/** The horde pem root certificates. */
-	UPROPERTY(Config, EditAnywhere, Category = "Horde", meta = (DisplayName = "PEM Root Certificates", ConfigRestartRequired = true))
+	/** The Bazel pem root certificates. */
+	UPROPERTY(Config, EditAnywhere, Category = "Bazel", meta = (DisplayName = "PEM Root Certificates", ConfigRestartRequired = true))
 	FString PemRootCertificates;
 
 };
