BazelExecutor: Expose ToBlob functions in IContentAddressableStorage
Update Digest hash, Timestamp, and Duration to use UE types
Forward declare as much as possible

diff --git a/Engine/Source/Editor/Experimental/BazelExecutor/Private/BazelExecutor.h b/Engine/Source/Editor/Experimental/BazelExecutor/Private/BazelExecutor.h
--- a/Engine/Source/Editor/Experimental/BazelExecutor/Private/BazelExecutor.h
+++ b/Engine/Source/Editor/Experimental/BazelExecutor/Private/BazelExecutor.h
@@ -7,7 +7,6 @@
 #include "IRemoteExecutor.h"
 #include "IContentAddressableStorage.h"
 #include "IExecution.h"
-#include "Messages.h"
 
 
 class FBazelExecutor : public IRemoteExecutor
diff --git a/Engine/Source/Editor/Experimental/BazelExecutor/Private/ContentAddressableStorage.cpp b/Engine/Source/Editor/Experimental/BazelExecutor/Private/ContentAddressableStorage.cpp
--- a/Engine/Source/Editor/Experimental/BazelExecutor/Private/ContentAddressableStorage.cpp
+++ b/Engine/Source/Editor/Experimental/BazelExecutor/Private/ContentAddressableStorage.cpp
@@ -3,6 +3,7 @@
 #include "ContentAddressableStorage.h"
 #include "BazelExecutorModule.h"
 #include "ProtoConverter.h"
+#include "Messages.h"
 
 THIRD_PARTY_INCLUDES_START
 UE_PUSH_MACRO("TEXT")
@@ -21,6 +22,21 @@ FContentAddressableStorage::~FContentAddressableStorage()
 {
 }
 
+bool FContentAddressableStorage::ToBlob(const FDirectory& InDirectory, TArray<char>& OutData, FDigest& OutDigest)
+{
+	return ProtoConverter::ToBlob(InDirectory, OutData, OutDigest);
+}
+
+bool FContentAddressableStorage::ToBlob(const FCommand& InCommand, TArray<char>& OutData, FDigest& OutDigest)
+{
+	return ProtoConverter::ToBlob(InCommand, OutData, OutDigest);
+}
+
+bool FContentAddressableStorage::ToBlob(const FAction& InAction, TArray<char>& OutData, FDigest& OutDigest)
+{
+	return ProtoConverter::ToBlob(InAction, OutData, OutDigest);
+}
+
 FStatus FContentAddressableStorage::FindMissingBlobs(const FFindMissingBlobsRequest& Request, FFindMissingBlobsResponse& Response)
 {
 	grpc::ClientContext Context;
diff --git a/Engine/Source/Editor/Experimental/BazelExecutor/Private/ContentAddressableStorage.h b/Engine/Source/Editor/Experimental/BazelExecutor/Private/ContentAddressableStorage.h
--- a/Engine/Source/Editor/Experimental/BazelExecutor/Private/ContentAddressableStorage.h
+++ b/Engine/Source/Editor/Experimental/BazelExecutor/Private/ContentAddressableStorage.h
@@ -4,11 +4,9 @@
 
 #include "CoreMinimal.h"
 #include "IContentAddressableStorage.h"
-#include "Messages.h"
 
 #include <memory>
 
-
 THIRD_PARTY_INCLUDES_START
 UE_PUSH_MACRO("TEXT")
 #undef TEXT
@@ -31,6 +29,10 @@ public:
 	FContentAddressableStorage(const std::shared_ptr<grpc::Channel>& Channel);
 	~FContentAddressableStorage();
 
+	bool ToBlob(const FDirectory& InDirectory, TArray<char>& OutData, FDigest& OutDigest) override;
+	bool ToBlob(const FCommand& InCommand, TArray<char>& OutData, FDigest& OutDigest) override;
+	bool ToBlob(const FAction& InAction, TArray<char>& OutData, FDigest& OutDigest) override;
+
 	FStatus FindMissingBlobs(const FFindMissingBlobsRequest& Request, FFindMissingBlobsResponse& Response) override;
 	FStatus BatchUpdateBlobs(const FBatchUpdateBlobsRequest& Request, FBatchUpdateBlobsResponse& Response) override;
 	FStatus BatchReadBlobs(const FBatchReadBlobsRequest& Request, FBatchReadBlobsResponse& Response) override;
diff --git a/Engine/Source/Editor/Experimental/BazelExecutor/Private/Execution.cpp b/Engine/Source/Editor/Experimental/BazelExecutor/Private/Execution.cpp
--- a/Engine/Source/Editor/Experimental/BazelExecutor/Private/Execution.cpp
+++ b/Engine/Source/Editor/Experimental/BazelExecutor/Private/Execution.cpp
@@ -3,6 +3,7 @@
 #include "Execution.h"
 #include "BazelExecutorModule.h"
 #include "ProtoConverter.h"
+#include "Messages.h"
 
 THIRD_PARTY_INCLUDES_START
 UE_PUSH_MACRO("TEXT")
diff --git a/Engine/Source/Editor/Experimental/BazelExecutor/Private/Execution.h b/Engine/Source/Editor/Experimental/BazelExecutor/Private/Execution.h
--- a/Engine/Source/Editor/Experimental/BazelExecutor/Private/Execution.h
+++ b/Engine/Source/Editor/Experimental/BazelExecutor/Private/Execution.h
@@ -4,7 +4,6 @@
 
 #include "CoreMinimal.h"
 #include "IExecution.h"
-#include "Messages.h"
 
 THIRD_PARTY_INCLUDES_START
 UE_PUSH_MACRO("TEXT")
diff --git a/Engine/Source/Editor/Experimental/BazelExecutor/Private/ProtoConverter.cpp b/Engine/Source/Editor/Experimental/BazelExecutor/Private/ProtoConverter.cpp
--- a/Engine/Source/Editor/Experimental/BazelExecutor/Private/ProtoConverter.cpp
+++ b/Engine/Source/Editor/Experimental/BazelExecutor/Private/ProtoConverter.cpp
@@ -2,9 +2,17 @@
 
 #include "ProtoConverter.h"
 #include "Misc/Paths.h"
+#include "Messages.h"
+
+#include <string>
+
 
 THIRD_PARTY_INCLUDES_START
-#include <openssl/sha.h>
+UE_PUSH_MACRO("TEXT")
+#undef TEXT
+#include "build\bazel\remote\execution\v2\remote_execution.pb.h"
+#include "build\bazel\remote\execution\v2\remote_execution.grpc.pb.h"
+UE_POP_MACRO("TEXT");
 THIRD_PARTY_INCLUDES_END
 
 
@@ -13,16 +21,21 @@ void ProtoConverter::ToProto(const FString& In, std::string& Out)
 	Out.assign(TCHAR_TO_UTF8(*In));
 }
 
-void ProtoConverter::ToProto(const FDuration& In, google::protobuf::Duration& Out)
+void ProtoConverter::ToProto(const FIoHash& In, std::string& Out)
+{
+	Out.assign(TCHAR_TO_UTF8(*FString::FromHexBlob(In.GetBytes(), sizeof(FIoHash::ByteArray)).ToLower()));
+}
+
+void ProtoConverter::ToProto(const FTimespan& In, google::protobuf::Duration& Out)
 {
-	Out.set_seconds(In.Seconds);
-	Out.set_nanos(In.Nanos);
+	Out.set_seconds(In.GetTotalSeconds());
+	Out.set_nanos(In.GetFractionNano());
 }
 
-void ProtoConverter::ToProto(const FTimestamp& In, google::protobuf::Timestamp& Out)
+void ProtoConverter::ToProto(const FDateTime& In, google::protobuf::Timestamp& Out)
 {
-	Out.set_seconds(In.Seconds);
-	Out.set_nanos(In.Nanos);
+	Out.set_seconds(In.ToUnixTimestamp());
+	Out.set_nanos((In.GetTicks() % ETimespan::TicksPerSecond) * 100);
 }
 
 void ProtoConverter::ToProto(const FDigest& In, build::bazel::remote::execution::v2::Digest& Out)
@@ -43,7 +56,7 @@ void ProtoConverter::ToProto(const FNodeProperties& In, build::bazel::remote::ex
 	{
 		ToProto(InProperty, *Out.add_properties());
 	}
-	if (In.ModifiedTime.Nanos || In.ModifiedTime.Seconds)
+	if (In.ModifiedTime.GetTicks() != 0)
 	{
 		ToProto(In.ModifiedTime, *Out.mutable_mtime());
 	}
@@ -58,7 +71,7 @@ void ProtoConverter::ToProto(const FFileNode& In, build::bazel::remote::executio
 	ToProto(In.Name, *Out.mutable_name());
 	ToProto(In.Digest, *Out.mutable_digest());
 	Out.set_is_executable(In.IsExecutable);
-	if (In.NodeProperties.ModifiedTime.Seconds || In.NodeProperties.ModifiedTime.Nanos || !In.NodeProperties.Properties.IsEmpty() || In.NodeProperties.UnixMode)
+	if (In.NodeProperties.ModifiedTime != 0 || !In.NodeProperties.Properties.IsEmpty() || In.NodeProperties.UnixMode)
 	{
 		ToProto(In.NodeProperties, *Out.mutable_node_properties());
 	}
@@ -74,7 +87,7 @@ void ProtoConverter::ToProto(const FSymlinkNode& In, build::bazel::remote::execu
 {
 	ToProto(In.Name, *Out.mutable_name());
 	ToProto(In.Target, *Out.mutable_target());
-	if (In.NodeProperties.ModifiedTime.Seconds || In.NodeProperties.ModifiedTime.Nanos || !In.NodeProperties.Properties.IsEmpty() || In.NodeProperties.UnixMode)
+	if (In.NodeProperties.ModifiedTime.GetTicks() != 0 || !In.NodeProperties.Properties.IsEmpty() || In.NodeProperties.UnixMode)
 	{
 		ToProto(In.NodeProperties, *Out.mutable_node_properties());
 	}
@@ -96,7 +109,7 @@ void ProtoConverter::ToProto(const FDirectory& In, build::bazel::remote::executi
 	{
 		ToProto(InSymlink, *Out.add_symlinks());
 	}
-	if (In.NodeProperties.ModifiedTime.Seconds || In.NodeProperties.ModifiedTime.Nanos || !In.NodeProperties.Properties.IsEmpty() || In.NodeProperties.UnixMode)
+	if (In.NodeProperties.ModifiedTime.GetTicks() != 0 || !In.NodeProperties.Properties.IsEmpty() || In.NodeProperties.UnixMode)
 	{
 		ToProto(In.NodeProperties, *Out.mutable_node_properties());
 	}
@@ -116,7 +129,7 @@ void ProtoConverter::ToProto(const FAction& In, build::bazel::remote::execution:
 {
 	ToProto(In.CommandDigest, *Out.mutable_command_digest());
 	ToProto(In.InputRootDigest, *Out.mutable_input_root_digest());
-	if (In.Timeout.Seconds || In.Timeout.Nanos)
+	if (!In.Timeout.IsZero())
 	{
 		ToProto(In.Timeout, *Out.mutable_timeout());
 	}
@@ -163,6 +176,13 @@ void ProtoConverter::FromProto(const std::string& In, FString& Out)
 	Out = UTF8_TO_TCHAR(In.c_str());
 }
 
+void ProtoConverter::FromProto(const std::string& In, FIoHash& Out)
+{
+	FIoHash::ByteArray ByteArray;
+	FString::ToHexBlob(UTF8_TO_TCHAR(In.c_str()), ByteArray, sizeof(FIoHash::ByteArray));
+	Out = FIoHash(ByteArray);
+}
+
 void ProtoConverter::FromProto(const std::string& In, TArray<char>& Out)
 {
 	Out.Empty();
@@ -173,10 +193,10 @@ void ProtoConverter::FromProto(const std::string& In, TArray<char>& Out)
 	}
 }
 
-void ProtoConverter::FromProto(const google::protobuf::Timestamp& In, FTimestamp& Out)
+void ProtoConverter::FromProto(const google::protobuf::Timestamp& In, FDateTime& Out)
 {
-	Out.Seconds = In.seconds();
-	Out.Nanos = In.nanos();
+	Out.FromUnixTimestamp(In.seconds());
+	Out += FTimespan(In.nanos() / 100);
 }
 
 void ProtoConverter::FromProto(const build::bazel::remote::execution::v2::Digest& In, FDigest& Out)
@@ -399,15 +419,7 @@ void ProtoConverter::FromProto(const build::bazel::remote::execution::v2::Execut
 
 bool ProtoConverter::ToDigest(const TArray<char>& Data, FDigest& OutDigest)
 {
-	unsigned char Sha256Buffer[SHA256_DIGEST_LENGTH];
-
-	SHA256_CTX sha256;
-	SHA256_Init(&sha256);
-	SHA256_Update(&sha256, Data.GetData(), Data.Num());
-	SHA256_Final(Sha256Buffer, &sha256);
-
-	// GRPC requires lowercase hex
-	OutDigest.Hash = FString::FromHexBlob(Sha256Buffer, SHA256_DIGEST_LENGTH).ToLower();
+	OutDigest.Hash = FIoHash::HashBuffer(Data.GetData(), Data.Num());
 	OutDigest.SizeBytes = Data.Num();
 	return true;
 }
diff --git a/Engine/Source/Editor/Experimental/BazelExecutor/Private/ProtoConverter.h b/Engine/Source/Editor/Experimental/BazelExecutor/Private/ProtoConverter.h
--- a/Engine/Source/Editor/Experimental/BazelExecutor/Private/ProtoConverter.h
+++ b/Engine/Source/Editor/Experimental/BazelExecutor/Private/ProtoConverter.h
@@ -3,25 +3,104 @@
 #pragma once
 
 #include "CoreMinimal.h"
-#include "Messages.h"
 
-THIRD_PARTY_INCLUDES_START
-UE_PUSH_MACRO("TEXT")
-#undef TEXT
-#include <grpcpp/grpcpp.h>
-#include "build\bazel\remote\execution\v2\remote_execution.pb.h"
-#include "build\bazel\remote\execution\v2\remote_execution.grpc.pb.h"
-UE_POP_MACRO("TEXT");
-THIRD_PARTY_INCLUDES_END
+#include <string>
+
+namespace google
+{
+	namespace protobuf
+	{
+		class Duration;
+		class Message;
+		class Timestamp;
+	}
+	namespace rpc
+	{
+		class Status;
+	}
+}
+
+namespace grpc
+{
+	class Status;
+}
+
+namespace build
+{
+	namespace bazel
+	{
+		namespace remote
+		{
+			namespace execution
+			{
+				namespace v2
+				{
+					class Action;
+					class ActionResult;
+					class BatchReadBlobsRequest;
+					class BatchReadBlobsResponse;
+					class BatchUpdateBlobsRequest;
+					class BatchUpdateBlobsResponse;
+					class Command;
+					class Digest;
+					class Directory;
+					class DirectoryNode;
+					class ExecutedActionMetadata;
+					class ExecuteRequest;
+					class ExecuteResponse;
+					class FileNode;
+					class FindMissingBlobsRequest;
+					class FindMissingBlobsResponse;
+					class LogFile;
+					class NodeProperties;
+					class NodeProperty;
+					class OutputDirectory;
+					class OutputFile;
+					class OutputSymlink;
+					class Platform;
+					class SymlinkNode;
+				}
+			}
+		}
+	}
+}
+
+struct FAction;
+struct FActionResult;
+struct FBatchReadBlobsRequest;
+struct FBatchReadBlobsResponse;
+struct FBatchUpdateBlobsRequest;
+struct FBatchUpdateBlobsResponse;
+struct FCommand;
+struct FDigest;
+struct FDirectory;
+struct FDirectoryNode;
+struct FExecutedActionMetadata;
+struct FExecuteRequest;
+struct FExecuteResponse;
+struct FFileNode;
+struct FFindMissingBlobsRequest;
+struct FFindMissingBlobsResponse;
+struct FLogFile;
+struct FNodeProperties;
+struct FNodeProperty;
+struct FOutputDirectory;
+struct FOutputFile;
+struct FOutputSymlink;
+struct FPlatform;
+struct FSymlinkNode;
+struct FStatus;
+struct FIoHash;
 
 
 class ProtoConverter
 {
 private:
 	static void ToProto(const FString& In, std::string& Out);
+	static void ToProto(const FIoHash& In, std::string& Out);
 
-	static void ToProto(const FDuration& In, google::protobuf::Duration& Out);
-	static void ToProto(const FTimestamp& In, google::protobuf::Timestamp& Out);
+	static void ToProto(const FTimespan& In, google::protobuf::Duration& Out);
+	static void ToProto(const FDateTime& In, google::protobuf::Timestamp& Out);
 
 	static void ToProto(const FDigest& In, build::bazel::remote::execution::v2::Digest& Out);
 
@@ -38,9 +117,10 @@ private:
 	static void ToProto(const FCommand& In, build::bazel::remote::execution::v2::Command& Out);
 
 	static void FromProto(const std::string& In, FString& Out);
+	static void FromProto(const std::string& In, FIoHash& Out);
 	static void FromProto(const std::string& In, TArray<char>& Out);
 
-	static void FromProto(const google::protobuf::Timestamp& In, FTimestamp& Out);
+	static void FromProto(const google::protobuf::Timestamp& In, FDateTime& Out);
 
 	static void FromProto(const build::bazel::remote::execution::v2::Digest& In, FDigest& Out);
 
@@ -54,6 +134,8 @@ private:
 	static void FromProto(const build::bazel::remote::execution::v2::ExecutedActionMetadata& In, FExecutedActionMetadata& Out);
 	static void FromProto(const build::bazel::remote::execution::v2::ActionResult& In, FActionResult& Out);
 
+	static bool ToBlob(const google::protobuf::Message& InMessage, TArray<char>& OutData, FDigest& OutDigest);
+
 public:
 	static void ToProto(const FExecuteRequest& In, build::bazel::remote::execution::v2::ExecuteRequest& Out);
 	static void ToProto(const FFindMissingBlobsRequest& In, build::bazel::remote::execution::v2::FindMissingBlobsRequest& Out);
@@ -69,10 +151,7 @@ public:
 
 	static bool ToDigest(const TArray<char>& Data, FDigest& OutDigest);
 
-	static bool ToBlob(const google::protobuf::Message& InMessage, TArray<char>& OutData, FDigest& OutDigest);
-
 	static bool ToBlob(const FDirectory& InDirectory, TArray<char>& OutData, FDigest& OutDigest);
 	static bool ToBlob(const FCommand& InCommand, TArray<char>& OutData, FDigest& OutDigest);
 	static bool ToBlob(const FAction& InAction, TArray<char>& OutData, FDigest& OutDigest);
 };
-
diff --git a/Engine/Source/Editor/Experimental/RemoteExecution/Public/IContentAddressableStorage.h b/Engine/Source/Editor/Experimental/RemoteExecution/Public/IContentAddressableStorage.h
--- a/Engine/Source/Editor/Experimental/RemoteExecution/Public/IContentAddressableStorage.h
+++ b/Engine/Source/Editor/Experimental/RemoteExecution/Public/IContentAddressableStorage.h
@@ -3,7 +3,18 @@
 #pragma once
 
 #include "CoreMinimal.h"
-#include "Messages.h"
+
+struct FAction;
+struct FBatchReadBlobsRequest;
+struct FBatchReadBlobsResponse;
+struct FBatchUpdateBlobsRequest;
+struct FBatchUpdateBlobsResponse;
+struct FCommand;
+struct FDigest;
+struct FDirectory;
+struct FFindMissingBlobsRequest;
+struct FFindMissingBlobsResponse;
+struct FStatus;
 
 
 class IContentAddressableStorage
@@ -12,6 +23,10 @@ public:
 	/** Virtual destructor */
 	virtual ~IContentAddressableStorage() {}
 
+	virtual bool ToBlob(const FDirectory& InDirectory, TArray<char>& OutData, FDigest& OutDigest) = 0;
+	virtual bool ToBlob(const FCommand& InCommand, TArray<char>& OutData, FDigest& OutDigest) = 0;
+	virtual bool ToBlob(const FAction& InAction, TArray<char>& OutData, FDigest& OutDigest) = 0;
+
 	virtual FStatus FindMissingBlobs(const FFindMissingBlobsRequest& Request, FFindMissingBlobsResponse& Response) = 0;
 	virtual FStatus BatchUpdateBlobs(const FBatchUpdateBlobsRequest& Request, FBatchUpdateBlobsResponse& Response) = 0;
 	virtual FStatus BatchReadBlobs(const FBatchReadBlobsRequest& Request, FBatchReadBlobsResponse& Response) = 0;
diff --git a/Engine/Source/Editor/Experimental/RemoteExecution/Public/IExecution.h b/Engine/Source/Editor/Experimental/RemoteExecution/Public/IExecution.h
--- a/Engine/Source/Editor/Experimental/RemoteExecution/Public/IExecution.h
+++ b/Engine/Source/Editor/Experimental/RemoteExecution/Public/IExecution.h
@@ -3,7 +3,9 @@
 #pragma once
 
 #include "CoreMinimal.h"
-#include "Messages.h"
+
+struct FExecuteRequest;
+struct FExecuteResponse;
 
 
 class IExecution
diff --git a/Engine/Source/Editor/Experimental/RemoteExecution/Public/Messages.h b/Engine/Source/Editor/Experimental/RemoteExecution/Public/Messages.h
--- a/Engine/Source/Editor/Experimental/RemoteExecution/Public/Messages.h
+++ b/Engine/Source/Editor/Experimental/RemoteExecution/Public/Messages.h
@@ -3,26 +3,10 @@
 #pragma once
 
 #include "CoreMinimal.h"
+#include "IO/IoHash.h"
+#include "Misc/DateTime.h"
 
 
-// namespace google.protobuf
-// https://developers.google.com/protocol-buffers/docs/reference/google.protobuf
-//
-
-struct FDuration
-{
-	int64 Seconds = 0;
-
-	int32 Nanos = 0;
-};
-
-struct FTimestamp
-{
-	int64 Seconds = 0;
-
-	int32 Nanos = 0;
-};
-
 // namespace grpc
 // https://grpc.github.io/grpc/cpp/namespacegrpc.html
 //
@@ -66,7 +50,7 @@ struct FStatus
 
 struct FDigest
 {
-	FString Hash;
+	FIoHash Hash;
 
 	int64 SizeBytes = 0;
 };
@@ -82,7 +66,7 @@ struct FNodeProperties
 {
 	TArray<FNodeProperty> Properties;
 
-	FTimestamp ModifiedTime;
+	FDateTime ModifiedTime;
 
 	uint32 UnixMode = 0;
 };
@@ -140,7 +124,7 @@ struct FAction
 
 	FDigest InputRootDigest;
 
-	FDuration Timeout;
+	FTimespan Timeout;
 
 	bool DoNotCache = false;
 
@@ -277,15 +261,15 @@ struct FTree
 struct FExecutedActionMetadata
 {
 	FString Worker;
-	FTimestamp QueuedTimestamp;
-	FTimestamp WorkerStartTimestamp;
-	FTimestamp WorkerCompletedTimestamp;
-	FTimestamp InputFetchStartTimestamp;
-	FTimestamp InputFetchCompletedTimestamp;
-	FTimestamp ExecutionStartTimestamp;
-	FTimestamp ExecutionCompletedTimestamp;
-	FTimestamp OutputUploadStartTimestamp;
-	FTimestamp OutputUploadCompletedTimestamp;
+	FDateTime QueuedTimestamp;
+	FDateTime WorkerStartTimestamp;
+	FDateTime WorkerCompletedTimestamp;
+	FDateTime InputFetchStartTimestamp;
+	FDateTime InputFetchCompletedTimestamp;
+	FDateTime ExecutionStartTimestamp;
+	FDateTime ExecutionCompletedTimestamp;
+	FDateTime OutputUploadStartTimestamp;
+	FDateTime OutputUploadCompletedTimestamp;
 	//repeated google.protobuf.Any AuxiliaryMetadata = 11;
 };
 
