BazelExecutor: Correct configuration for Horde CAS & RemoteExecution

diff --git a/Engine/Source/Editor/Experimental/BazelExecutor/Private/BazelExecutor.cpp b/Engine/Source/Editor/Experimental/BazelExecutor/Private/BazelExecutor.cpp
--- a/Engine/Source/Editor/Experimental/BazelExecutor/Private/BazelExecutor.cpp
+++ b/Engine/Source/Editor/Experimental/BazelExecutor/Private/BazelExecutor.cpp
@@ -19,21 +19,34 @@ THIRD_PARTY_INCLUDES_END
 #define LOCTEXT_NAMESPACE "BazelExecutor"
 
 
-void FBazelExecutor::Initialize(const FString& Target, const FSslCredentialsOptions& SslCredentialsOptions)
+void FBazelExecutor::Initialize(const FSettings& Settings)
 {
 	ContentAddressableStorage.Reset();
 	Execution.Reset();
 
-	grpc::SslCredentialsOptions SslOptions;
-	SslOptions.pem_cert_chain.assign(TCHAR_TO_UTF8(*SslCredentialsOptions.PemCertChain));
-	SslOptions.pem_private_key.assign(TCHAR_TO_UTF8(*SslCredentialsOptions.PemPrivateKey));
-	SslOptions.pem_root_certs.assign(TCHAR_TO_UTF8(*SslCredentialsOptions.PemRootCerts));
+	grpc::ChannelArguments ChannelArguments;
+	ChannelArguments.SetMaxSendMessageSize(Settings.MaxSendMessageSize);
+	ChannelArguments.SetMaxReceiveMessageSize(Settings.MaxReceiveMessageSize);
 
-	// Additional configuration could be set with grpc::CreateCustomChannel()
-	std::shared_ptr<grpc::Channel> Channel = grpc::CreateChannel(TCHAR_TO_UTF8(*Target), grpc::SslCredentials(SslOptions));
+	grpc::SslCredentialsOptions ContentAddressableStorageSslOptions;
+	ContentAddressableStorageSslOptions.pem_cert_chain.assign(TCHAR_TO_UTF8(*Settings.ContentAddressableStoragePemCertificateChain));
+	ContentAddressableStorageSslOptions.pem_private_key.assign(TCHAR_TO_UTF8(*Settings.ContentAddressableStoragePemPrivateKey));
+	ContentAddressableStorageSslOptions.pem_root_certs.assign(TCHAR_TO_UTF8(*Settings.ContentAddressableStoragePemRootCertificates));
 
-	ContentAddressableStorage.Reset(new FContentAddressableStorage(Channel));
-	Execution.Reset(new FExecution(Channel));
+	std::shared_ptr<grpc::ChannelCredentials> ContentAddressableStorageChannelCredentials = grpc::SslCredentials(grpc::SslCredentialsOptions(ContentAddressableStorageSslOptions));
+
+	std::shared_ptr<grpc::Channel> ContentAddressableStorageChannel = grpc::CreateCustomChannel(TCHAR_TO_UTF8(*Settings.ContentAddressableStorageTarget), ContentAddressableStorageChannelCredentials, ChannelArguments);
+	ContentAddressableStorage.Reset(new FContentAddressableStorage(ContentAddressableStorageChannel, Settings.ContentAddressableStorageHeaders));
+
+	grpc::SslCredentialsOptions ExecutionSslOptions;
+	ExecutionSslOptions.pem_cert_chain.assign(TCHAR_TO_UTF8(*Settings.ExecutionPemCertificateChain));
+	ExecutionSslOptions.pem_private_key.assign(TCHAR_TO_UTF8(*Settings.ExecutionPemPrivateKey));
+	ExecutionSslOptions.pem_root_certs.assign(TCHAR_TO_UTF8(*Settings.ExecutionPemRootCertificates));
+
+	std::shared_ptr<grpc::ChannelCredentials> ExecutionChannelCredentials = grpc::SslCredentials(grpc::SslCredentialsOptions(ExecutionSslOptions));
+
+	std::shared_ptr<grpc::Channel> ExecutorChannel = grpc::CreateCustomChannel(TCHAR_TO_UTF8(*Settings.ExecutionTarget), ExecutionChannelCredentials, ChannelArguments);
+	Execution.Reset(new FExecution(ExecutorChannel, Settings.ExecutionHeaders));
 }
 
 FName FBazelExecutor::GetFName() const
diff --git a/Engine/Source/Editor/Experimental/BazelExecutor/Private/BazelExecutor.h b/Engine/Source/Editor/Experimental/BazelExecutor/Private/BazelExecutor.h
--- a/Engine/Source/Editor/Experimental/BazelExecutor/Private/BazelExecutor.h
+++ b/Engine/Source/Editor/Experimental/BazelExecutor/Private/BazelExecutor.h
@@ -12,11 +12,20 @@
 class FBazelExecutor : public IRemoteExecutor
 {
 public:
-	struct FSslCredentialsOptions
+	struct FSettings
 	{
-		FString PemCertChain;
-		FString PemPrivateKey;
-		FString PemRootCerts;
+		FString ContentAddressableStorageTarget;
+		FString ExecutionTarget;
+		TMap<FString, FString> ContentAddressableStorageHeaders;
+		TMap<FString, FString> ExecutionHeaders;
+		int32 MaxSendMessageSize;
+		int32 MaxReceiveMessageSize;
+		FString ContentAddressableStoragePemCertificateChain;
+		FString ContentAddressableStoragePemPrivateKey;
+		FString ContentAddressableStoragePemRootCertificates;
+		FString ExecutionPemCertificateChain;
+		FString ExecutionPemPrivateKey;
+		FString ExecutionPemRootCertificates;
 	};
 
 private:
@@ -24,7 +33,7 @@ private:
 	TUniquePtr<IExecution> Execution;
 
 public:
-	void Initialize(const FString& Target, const FSslCredentialsOptions& SslCredentialsOptions);
+	void Initialize(const FSettings& Settings);
 
 	virtual FName GetFName() const override;
 	virtual FText GetNameText() const override;
diff --git a/Engine/Source/Editor/Experimental/BazelExecutor/Private/BazelExecutorModule.cpp b/Engine/Source/Editor/Experimental/BazelExecutor/Private/BazelExecutorModule.cpp
--- a/Engine/Source/Editor/Experimental/BazelExecutor/Private/BazelExecutorModule.cpp
+++ b/Engine/Source/Editor/Experimental/BazelExecutor/Private/BazelExecutorModule.cpp
@@ -1,10 +1,11 @@
 // Copyright Epic Games, Inc. All Rights Reserved.
 
 #include "BazelExecutorModule.h"
-#include "ISettingsModule.h"
-#include "BazelExecutorSettings.h"
 #include "Misc/FileHelper.h"
 #include "Misc/Paths.h"
+#include "Modules/ModuleManager.h"
+#include "ISettingsModule.h"
+#include "BazelExecutorSettings.h"
 #include "Messages.h"
 
 IMPLEMENT_MODULE(FBazelExecutorModule, EditorConfig);
@@ -13,38 +14,99 @@ IMPLEMENT_MODULE(FBazelExecutorModule, EditorConfig);
 
 DEFINE_LOG_CATEGORY(LogBazelExecutor);
 
+static FName RemoteExecutionFeatureName(TEXT("RemoteExecution"));
+
 
 void FBazelExecutorModule::StartupModule()
 {
 	GetMutableDefault<UBazelExecutorSettings>()->LoadConfig();
 
-	const FString Target = GetDefault<UBazelExecutorSettings>()->Target;
-	const FString PemCertificateChain = GetDefault<UBazelExecutorSettings>()->PemCertificateChain;
-	const FString PemPrivateKey = GetDefault<UBazelExecutorSettings>()->PemPrivateKey;
-	const FString PemRootCertificates = GetDefault<UBazelExecutorSettings>()->PemRootCertificates;
+	const UBazelExecutorSettings* BazelExecutorSettings = GetDefault<UBazelExecutorSettings>();
+
+	FBazelExecutor::FSettings Settings;
+	Settings.ContentAddressableStorageTarget = BazelExecutorSettings->ContentAddressableStorageTarget;
+	Settings.ExecutionTarget = BazelExecutorSettings->ExecutionTarget;
+	Settings.ContentAddressableStorageHeaders = BazelExecutorSettings->ContentAddressableStorageHeaders;
+	Settings.ExecutionHeaders = BazelExecutorSettings->ExecutionHeaders;
+	Settings.MaxSendMessageSize = BazelExecutorSettings->MaxSendMessageSize;
+	Settings.MaxReceiveMessageSize = BazelExecutorSettings->MaxReceiveMessageSize;
+
+	if (!BazelExecutorSettings->ContentAddressableStoragePemCertificateChain.IsEmpty())
+	{
+		if (FPaths::FileExists(BazelExecutorSettings->ContentAddressableStoragePemCertificateChain))
+		{
+			FFileHelper::LoadFileToString(Settings.ContentAddressableStoragePemCertificateChain, *BazelExecutorSettings->ContentAddressableStoragePemCertificateChain);
+		}
+		else
+		{
+			Settings.ContentAddressableStoragePemCertificateChain = BazelExecutorSettings->ContentAddressableStoragePemCertificateChain;
+		}
+	}
+	if (!BazelExecutorSettings->ContentAddressableStoragePemPrivateKey.IsEmpty())
+	{
+		if (FPaths::FileExists(BazelExecutorSettings->ContentAddressableStoragePemPrivateKey))
+		{
+			FFileHelper::LoadFileToString(Settings.ContentAddressableStoragePemPrivateKey, *BazelExecutorSettings->ContentAddressableStoragePemPrivateKey);
+		}
+		else
+		{
+			Settings.ContentAddressableStoragePemPrivateKey = BazelExecutorSettings->ContentAddressableStoragePemPrivateKey;
+		}
+	}
+	if (!BazelExecutorSettings->ContentAddressableStoragePemRootCertificates.IsEmpty())
+	{
+		if (FPaths::FileExists(BazelExecutorSettings->ContentAddressableStoragePemRootCertificates))
+		{
+			FFileHelper::LoadFileToString(Settings.ContentAddressableStoragePemRootCertificates, *BazelExecutorSettings->ContentAddressableStoragePemRootCertificates);
+		}
+		else
+		{
+			Settings.ContentAddressableStoragePemRootCertificates = BazelExecutorSettings->ContentAddressableStoragePemRootCertificates;
+		}
+	}
 
-	FBazelExecutor::FSslCredentialsOptions SslCredentialsOptions;
-	if (!PemCertificateChain.IsEmpty() && FPaths::FileExists(PemCertificateChain))
+	if (!BazelExecutorSettings->ExecutionPemCertificateChain.IsEmpty())
 	{
-		FFileHelper::LoadFileToString(SslCredentialsOptions.PemCertChain, *PemCertificateChain);
+		if (FPaths::FileExists(BazelExecutorSettings->ExecutionPemCertificateChain))
+		{
+			FFileHelper::LoadFileToString(Settings.ExecutionPemCertificateChain, *BazelExecutorSettings->ExecutionPemCertificateChain);
+		}
+		else
+		{
+			Settings.ExecutionPemCertificateChain = BazelExecutorSettings->ExecutionPemCertificateChain;
+		}
 	}
-	if (!PemPrivateKey.IsEmpty() && FPaths::FileExists(PemPrivateKey))
+	if (!BazelExecutorSettings->ExecutionPemPrivateKey.IsEmpty())
 	{
-		FFileHelper::LoadFileToString(SslCredentialsOptions.PemPrivateKey, *PemPrivateKey);
+		if (FPaths::FileExists(BazelExecutorSettings->ExecutionPemPrivateKey))
+		{
+			FFileHelper::LoadFileToString(Settings.ExecutionPemPrivateKey, *BazelExecutorSettings->ExecutionPemPrivateKey);
+		}
+		else
+		{
+			Settings.ExecutionPemPrivateKey = BazelExecutorSettings->ExecutionPemPrivateKey;
+		}
 	}
-	if (!PemRootCertificates.IsEmpty() && FPaths::FileExists(PemRootCertificates))
+	if (!BazelExecutorSettings->ExecutionPemRootCertificates.IsEmpty())
 	{
-		FFileHelper::LoadFileToString(SslCredentialsOptions.PemRootCerts, *PemRootCertificates);
+		if (FPaths::FileExists(BazelExecutorSettings->ExecutionPemRootCertificates))
+		{
+			FFileHelper::LoadFileToString(Settings.ExecutionPemRootCertificates, *BazelExecutorSettings->ExecutionPemRootCertificates);
+		}
+		else
+		{
+			Settings.ExecutionPemRootCertificates = BazelExecutorSettings->ExecutionPemRootCertificates;
+		}
 	}
 
-	BazelExecution.Initialize(Target, SslCredentialsOptions);
+	BazelExecution.Initialize(Settings);
 
-	IModularFeatures::Get().RegisterModularFeature(TEXT("RemoteExecution"), &BazelExecution);
+	IModularFeatures::Get().RegisterModularFeature(RemoteExecutionFeatureName, &BazelExecution);
 }
 
 void FBazelExecutorModule::ShutdownModule()
 {
-	IModularFeatures::Get().UnregisterModularFeature(TEXT("RemoteExecution"), &BazelExecution);
+	IModularFeatures::Get().UnregisterModularFeature(RemoteExecutionFeatureName, &BazelExecution);
 }
 
 bool FBazelExecutorModule::SupportsDynamicReloading()
diff --git a/Engine/Source/Editor/Experimental/BazelExecutor/Private/BazelExecutorModule.h b/Engine/Source/Editor/Experimental/BazelExecutor/Private/BazelExecutorModule.h
--- a/Engine/Source/Editor/Experimental/BazelExecutor/Private/BazelExecutorModule.h
+++ b/Engine/Source/Editor/Experimental/BazelExecutor/Private/BazelExecutorModule.h
@@ -4,7 +4,6 @@
 
 #include "CoreMinimal.h"
 #include "Modules/ModuleInterface.h"
-#include "Modules/ModuleManager.h"
 #include "Features/IModularFeatures.h"
 
 #include "BazelExecutor.h"
diff --git a/Engine/Source/Editor/Experimental/BazelExecutor/Private/BazelExecutorSettings.cpp b/Engine/Source/Editor/Experimental/BazelExecutor/Private/BazelExecutorSettings.cpp
--- a/Engine/Source/Editor/Experimental/BazelExecutor/Private/BazelExecutorSettings.cpp
+++ b/Engine/Source/Editor/Experimental/BazelExecutor/Private/BazelExecutorSettings.cpp
@@ -6,4 +6,66 @@
 UBazelExecutorSettings::UBazelExecutorSettings(const FObjectInitializer& ObjectInitializer)
 	: Super(ObjectInitializer)
 {
+	ContentAddressableStorageTarget = TEXT("jupiter.devtools-dev.epicgames.com:8080");
+	ExecutionTarget = TEXT("hordeserver-grpc.devtools-dev.epicgames.com:443");
+
+	ContentAddressableStorageHeaders.Add("accept", "application/json");
+	ContentAddressableStorageHeaders.Add("authorization", "ServiceAccount HordeREAPI");
+
+	ExecutionHeaders.Add("accept", "application/json");
+	ExecutionHeaders.Add("authorization", "Bearer eyJhbGciOiJIUzI1NiIsInR5cCI6IkpXVCJ9.eyJodHRwOi8vc2NoZW1hcy54bWxzb2FwLm9yZy93cy8yMDA1LzA1L2lkZW50aXR5L2NsYWltcy9uYW1lIjoiQmVuIE1hcnNoIiwiaHR0cDovL2VwaWNnYW1lcy5jb20vdWUvaG9yZGUvcm9sZSI6ImFnZW50LXJlZ2lzdHJhdGlvbiIsImlzcyI6Imh0dHBzOi8vaG9yZGVzZXJ2ZXIuZGV2dG9vbHMtZGV2LmVwaWNnYW1lcy5jb20ifQ.qwGHN_BYqpJX2Y54nlJ6J_rinmmf5C1Srxymzk7hp2c");
+
+	MaxSendMessageSize = 1024 * 1024 * 1024;
+	MaxReceiveMessageSize = 1024 * 1024 * 1024;
+
+	ContentAddressableStoragePemRootCertificates = TEXT(
+		"-----BEGIN CERTIFICATE-----\n"
+		"MIIEDzCCAvegAwIBAgIBADANBgkqhkiG9w0BAQUFADBoMQswCQYDVQQGEwJVUzEl\n"
+		"MCMGA1UEChMcU3RhcmZpZWxkIFRlY2hub2xvZ2llcywgSW5jLjEyMDAGA1UECxMp\n"
+		"U3RhcmZpZWxkIENsYXNzIDIgQ2VydGlmaWNhdGlvbiBBdXRob3JpdHkwHhcNMDQw\n"
+		"NjI5MTczOTE2WhcNMzQwNjI5MTczOTE2WjBoMQswCQYDVQQGEwJVUzElMCMGA1UE\n"
+		"ChMcU3RhcmZpZWxkIFRlY2hub2xvZ2llcywgSW5jLjEyMDAGA1UECxMpU3RhcmZp\n"
+		"ZWxkIENsYXNzIDIgQ2VydGlmaWNhdGlvbiBBdXRob3JpdHkwggEgMA0GCSqGSIb3\n"
+		"DQEBAQUAA4IBDQAwggEIAoIBAQC3Msj+6XGmBIWtDBFk385N78gDGIc/oav7PKaf\n"
+		"8MOh2tTYbitTkPskpD6E8J7oX+zlJ0T1KKY/e97gKvDIr1MvnsoFAZMej2YcOadN\n"
+		"+lq2cwQlZut3f+dZxkqZJRRU6ybH838Z1TBwj6+wRir/resp7defqgSHo9T5iaU0\n"
+		"X9tDkYI22WY8sbi5gv2cOj4QyDvvBmVmepsZGD3/cVE8MC5fvj13c7JdBmzDI1aa\n"
+		"K4UmkhynArPkPw2vCHmCuDY96pzTNbO8acr1zJ3o/WSNF4Azbl5KXZnJHoe0nRrA\n"
+		"1W4TNSNe35tfPe/W93bC6j67eA0cQmdrBNj41tpvi/JEoAGrAgEDo4HFMIHCMB0G\n"
+		"A1UdDgQWBBS/X7fRzt0fhvRbVazc1xDCDqmI5zCBkgYDVR0jBIGKMIGHgBS/X7fR\n"
+		"zt0fhvRbVazc1xDCDqmI56FspGowaDELMAkGA1UEBhMCVVMxJTAjBgNVBAoTHFN0\n"
+		"YXJmaWVsZCBUZWNobm9sb2dpZXMsIEluYy4xMjAwBgNVBAsTKVN0YXJmaWVsZCBD\n"
+		"bGFzcyAyIENlcnRpZmljYXRpb24gQXV0aG9yaXR5ggEAMAwGA1UdEwQFMAMBAf8w\n"
+		"DQYJKoZIhvcNAQEFBQADggEBAAWdP4id0ckaVaGsafPzWdqbAYcaT1epoXkJKtv3\n"
+		"L7IezMdeatiDh6GX70k1PncGQVhiv45YuApnP+yz3SFmH8lU+nLMPUxA2IGvd56D\n"
+		"eruix/U0F47ZEUD0/CwqTRV/p2JdLiXTAAsgGh1o+Re49L2L7ShZ3U0WixeDyLJl\n"
+		"xy16paq8U4Zt3VekyvggQQto8PT7dL5WXXp59fkdheMtlb71cZBDzI0fmgAKhynp\n"
+		"VSJYACPq4xJDKVtHCN2MQWplBqjlIapBtJUhlbl90TSrE9atvNziPTnNvT51cKEY\n"
+		"WQPJIrSPnNVeKtelttQKbfi3QBFGmh95DmK/D5fs4C8fF5Q=\n"
+		"-----END CERTIFICATE-----"
+	);
+
+	ExecutionPemRootCertificates = TEXT(
+		"-----BEGIN CERTIFICATE-----\n"
+		"MIIDYTCCAkmgAwIBAgIIK9cc+F8gVfgwDQYJKoZIhvcNAQELBQAwNjE0MDIGA1UE\n"
+		"AxMraG9yZGVzZXJ2ZXItZ3JwYy5kZXZ0b29scy1kZXYuZXBpY2dhbWVzLmNvbTAe\n"
+		"Fw0yMDExMTAwMjE2MjZaFw0yMzAxMjAwMjE2MjZaMDYxNDAyBgNVBAMTK2hvcmRl\n"
+		"c2VydmVyLWdycGMuZGV2dG9vbHMtZGV2LmVwaWNnYW1lcy5jb20wggEiMA0GCSqG\n"
+		"SIb3DQEBAQUAA4IBDwAwggEKAoIBAQDErWwF1mzE7vM0xjBpgkHRuqiItJv+fs/5\n"
+		"A7CniQad9wnTliPWPzMFpNYp1SubrZzMc55Q+f0MkF7LjTdOwYo5gUY+A6x9pt4B\n"
+		"UNu4FP+XINI7j3WnK7mZgueh/JkQ28MGGYL1v7anaMCng/KXs6Mnlni1zirs81zZ\n"
+		"SPniUCZXVIckesKd6zSbjCZ6YqcPDehXMeT1R2HKPwZOkCgxMJisxGL27gtH/2ya\n"
+		"t3+x6ozUilENb5aBOsp3emK5+E7C+wamCMsipJT2mVa5ClddO09Ebc5KiUeBV6IN\n"
+		"3dliqSZ65TwJChicNdDpvnbusfJaCklFd6bQKZzy4vFkfdLQmDXxAgMBAAGjczBx\n"
+		"MAwGA1UdEwEB/wQCMAAwDgYDVR0PAQH/BAQDAgWgMBYGA1UdJQEB/wQMMAoGCCsG\n"
+		"AQUFBwMBMDkGA1UdEQEB/wQvMC2CK2hvcmRlc2VydmVyLWdycGMuZGV2dG9vbHMt\n"
+		"ZGV2LmVwaWNnYW1lcy5jb20wDQYJKoZIhvcNAQELBQADggEBAGv2cc+9rM0i7039\n"
+		"YZ0qEFiqeDQx7gOLk19gBBLhCumP+8HQ5eK3qgiRnYARSHp8hgrQM18ELoQf8muS\n"
+		"RYz5qG5cREumKki1BmfdStsxvWkFNMm1OyEBKHfCJoZ+r8D+/lVg17tkT/a65ewr\n"
+		"TYrL71lqnX022WlRAVycthUnTwADw+kwYS1D4ZEs3WJ7qlk6OWFfWz7x1zfnjM8y\n"
+		"T2Yoe1PkKTmDAaUw6OQa1X9vE8N9wd9zM9NDup5El6x0NoCIoT/ia8J+jhYRdeSc\n"
+		"U/4ARFp/nH+y+fbuSZvmam4ac2OoBltSOMlaoOkNXNNpo1WlB3ZlCCaC5B83Gr2H\n"
+		"IIoNLEM=\n"
+		"-----END CERTIFICATE-----"
+	);
 }
diff --git a/Engine/Source/Editor/Experimental/BazelExecutor/Private/BazelExecutorSettings.h b/Engine/Source/Editor/Experimental/BazelExecutor/Private/BazelExecutorSettings.h
--- a/Engine/Source/Editor/Experimental/BazelExecutor/Private/BazelExecutorSettings.h
+++ b/Engine/Source/Editor/Experimental/BazelExecutor/Private/BazelExecutorSettings.h
@@ -2,7 +2,6 @@
 
 #pragma once
 
-
 #include "CoreMinimal.h"
 #include "UObject/Object.h"
 #include "BazelExecutorSettings.generated.h"
@@ -13,20 +12,51 @@ class UBazelExecutorSettings : public UObject
 	GENERATED_UCLASS_BODY()
 
 public:
-	/** The Bazel server address. */
-	UPROPERTY(Config, EditAnywhere, Category = "Bazel", meta = (DisplayName = "Target", ConfigRestartRequired = true))
-	FString Target;
+	/** The Bazel server content addressable storage address. */
+	UPROPERTY(Config, EditAnywhere, Category = "Bazel", meta = (DisplayName = "Content Addressable Storage Target", ConfigRestartRequired = true))
+	FString ContentAddressableStorageTarget;
+
+	/** The Bazel server execution address. */
+	UPROPERTY(Config, EditAnywhere, Category = "Bazel", meta = (DisplayName = "Execution Target", ConfigRestartRequired = true))
+	FString ExecutionTarget;
+
+	/** Extra headers required for content addressable storage requests. */
+	UPROPERTY(Config, EditAnywhere, Category = "Bazel", meta = (DisplayName = "Content Addressable Storage Headers", ConfigRestartRequired = true))
+	TMap<FString, FString> ContentAddressableStorageHeaders;
+
+	/** Extra headers required for execution requests. */
+	UPROPERTY(Config, EditAnywhere, Category = "Bazel", meta = (DisplayName = "Execution Headers", ConfigRestartRequired = true))
+	TMap<FString, FString> ExecutionHeaders;
+
+	/** Maximum send message size. */
+	UPROPERTY(Config, EditAnywhere, Category = "Bazel", meta = (DisplayName = "Max Send Message Size", ConfigRestartRequired = true))
+	int32 MaxSendMessageSize;
+
+	/** Maximum receive message size. */
+	UPROPERTY(Config, EditAnywhere, Category = "Bazel", meta = (DisplayName = "Max Receive Message Size", ConfigRestartRequired = true))
+	int32 MaxReceiveMessageSize;
 	
 	/** The Bazel PEM certificate chain". */
-	UPROPERTY(Config, EditAnywhere, Category = "Bazel", meta = (DisplayName = "PEM Certificate Chain", ConfigRestartRequired = true))
-	FString PemCertificateChain;
+	UPROPERTY(Config, EditAnywhere, Category = "Bazel", meta = (DisplayName = "Content Addressable Storage PEM Certificate Chain", ConfigRestartRequired = true))
+	FString ContentAddressableStoragePemCertificateChain;
 
 	/** The Bazel PEM private key". */
-	UPROPERTY(Config, EditAnywhere, Category = "Bazel", meta = (DisplayName = "PEM Private Key", ConfigRestartRequired = true))
-	FString PemPrivateKey;
+	UPROPERTY(Config, EditAnywhere, Category = "Bazel", meta = (DisplayName = "Content Addressable Storage PEM Private Key", ConfigRestartRequired = true))
+	FString ContentAddressableStoragePemPrivateKey;
 
 	/** The Bazel pem root certificates. */
-	UPROPERTY(Config, EditAnywhere, Category = "Bazel", meta = (DisplayName = "PEM Root Certificates", ConfigRestartRequired = true))
-	FString PemRootCertificates;
+	UPROPERTY(Config, EditAnywhere, Category = "Bazel", meta = (DisplayName = "Content Addressable Storage PEM Root Certificates", ConfigRestartRequired = true))
+	FString ContentAddressableStoragePemRootCertificates;
 
+	/** The Bazel PEM certificate chain". */
+	UPROPERTY(Config, EditAnywhere, Category = "Bazel", meta = (DisplayName = "Execution PEM Certificate Chain", ConfigRestartRequired = true))
+	FString ExecutionPemCertificateChain;
+
+	/** The Bazel PEM private key". */
+	UPROPERTY(Config, EditAnywhere, Category = "Bazel", meta = (DisplayName = "Execution PEM Private Key", ConfigRestartRequired = true))
+	FString ExecutionPemPrivateKey;
+
+	/** The Bazel pem root certificates. */
+	UPROPERTY(Config, EditAnywhere, Category = "Bazel", meta = (DisplayName = "Execution PEM Root Certificates", ConfigRestartRequired = true))
+	FString ExecutionPemRootCertificates;
 };
diff --git a/Engine/Source/Editor/Experimental/BazelExecutor/Private/ContentAddressableStorage.cpp b/Engine/Source/Editor/Experimental/BazelExecutor/Private/ContentAddressableStorage.cpp
--- a/Engine/Source/Editor/Experimental/BazelExecutor/Private/ContentAddressableStorage.cpp
+++ b/Engine/Source/Editor/Experimental/BazelExecutor/Private/ContentAddressableStorage.cpp
@@ -13,8 +13,8 @@ UE_POP_MACRO("TEXT");
 THIRD_PARTY_INCLUDES_END
 
 
-FContentAddressableStorage::FContentAddressableStorage(const std::shared_ptr<grpc::Channel>& Channel) :
-	Stub(build::bazel::remote::execution::v2::ContentAddressableStorage::NewStub(Channel))
+FContentAddressableStorage::FContentAddressableStorage(const std::shared_ptr<grpc::Channel>& Channel, const TMap<FString, FString>& Headers) :
+	Stub(build::bazel::remote::execution::v2::ContentAddressableStorage::NewStub(Channel)), Headers(Headers)
 {
 }
 
@@ -22,6 +22,11 @@ FContentAddressableStorage::~FContentAddressableStorage()
 {
 }
 
+bool FContentAddressableStorage::ToDigest(const TArray<char>& InData, FDigest& OutDigest)
+{
+	return ProtoConverter::ToDigest(InData, OutDigest);
+}
+
 bool FContentAddressableStorage::ToBlob(const FDirectory& InDirectory, TArray<char>& OutData, FDigest& OutDigest)
 {
 	return ProtoConverter::ToBlob(InDirectory, OutData, OutDigest);
@@ -40,6 +45,10 @@ bool FContentAddressableStorage::ToBlob(const FAction& InAction, TArray<char>& O
 FStatus FContentAddressableStorage::FindMissingBlobs(const FFindMissingBlobsRequest& Request, FFindMissingBlobsResponse& Response)
 {
 	grpc::ClientContext Context;
+	for (const TPair<FString, FString>& Header : Headers)
+	{
+		Context.AddMetadata(TCHAR_TO_UTF8(*Header.Key.ToLower()), TCHAR_TO_UTF8(*Header.Value));
+	}
 
 	build::bazel::remote::execution::v2::FindMissingBlobsRequest ProtoRequest;
 	ProtoConverter::ToProto(Request, ProtoRequest);
@@ -56,6 +65,10 @@ FStatus FContentAddressableStorage::FindMissingBlobs(const FFindMissingBlobsRequ
 FStatus FContentAddressableStorage::BatchUpdateBlobs(const FBatchUpdateBlobsRequest& Request, FBatchUpdateBlobsResponse& Response)
 {
 	grpc::ClientContext Context;
+	for (const TPair<FString, FString>& Header : Headers)
+	{
+		Context.AddMetadata(TCHAR_TO_UTF8(*Header.Key.ToLower()), TCHAR_TO_UTF8(*Header.Value));
+	}
 
 	build::bazel::remote::execution::v2::BatchUpdateBlobsRequest ProtoRequest;
 	ProtoConverter::ToProto(Request, ProtoRequest);
@@ -72,6 +85,10 @@ FStatus FContentAddressableStorage::BatchUpdateBlobs(const FBatchUpdateBlobsRequ
 FStatus FContentAddressableStorage::BatchReadBlobs(const FBatchReadBlobsRequest& Request, FBatchReadBlobsResponse& Response)
 {
 	grpc::ClientContext Context;
+	for (const TPair<FString, FString>& Header : Headers)
+	{
+		Context.AddMetadata(TCHAR_TO_UTF8(*Header.Key.ToLower()), TCHAR_TO_UTF8(*Header.Value));
+	}
 
 	build::bazel::remote::execution::v2::BatchReadBlobsRequest ProtoRequest;
 	ProtoConverter::ToProto(Request, ProtoRequest);
diff --git a/Engine/Source/Editor/Experimental/BazelExecutor/Private/ContentAddressableStorage.h b/Engine/Source/Editor/Experimental/BazelExecutor/Private/ContentAddressableStorage.h
--- a/Engine/Source/Editor/Experimental/BazelExecutor/Private/ContentAddressableStorage.h
+++ b/Engine/Source/Editor/Experimental/BazelExecutor/Private/ContentAddressableStorage.h
@@ -24,11 +24,13 @@ class FContentAddressableStorage : public IContentAddressableStorage
 {
 private:
 	std::unique_ptr<build::bazel::remote::execution::v2::ContentAddressableStorage::Stub> Stub;
+	TMap<FString, FString> Headers;
 
 public:
-	FContentAddressableStorage(const std::shared_ptr<grpc::Channel>& Channel);
+	FContentAddressableStorage(const std::shared_ptr<grpc::Channel>& Channel, const TMap<FString, FString>& Headers);
 	~FContentAddressableStorage();
 
+	bool ToDigest(const TArray<char>& InData, FDigest& OutDigest) override;
 	bool ToBlob(const FDirectory& InDirectory, TArray<char>& OutData, FDigest& OutDigest) override;
 	bool ToBlob(const FCommand& InCommand, TArray<char>& OutData, FDigest& OutDigest) override;
 	bool ToBlob(const FAction& InAction, TArray<char>& OutData, FDigest& OutDigest) override;
diff --git a/Engine/Source/Editor/Experimental/BazelExecutor/Private/Execution.cpp b/Engine/Source/Editor/Experimental/BazelExecutor/Private/Execution.cpp
--- a/Engine/Source/Editor/Experimental/BazelExecutor/Private/Execution.cpp
+++ b/Engine/Source/Editor/Experimental/BazelExecutor/Private/Execution.cpp
@@ -13,8 +13,8 @@ UE_POP_MACRO("TEXT");
 THIRD_PARTY_INCLUDES_END
 
 
-FExecution::FExecution(const std::shared_ptr<grpc::Channel>& Channel) :
-	Stub(build::bazel::remote::execution::v2::Execution::NewStub(Channel))
+FExecution::FExecution(const std::shared_ptr<grpc::Channel>& Channel, const TMap<FString, FString>& Headers) :
+	Stub(build::bazel::remote::execution::v2::Execution::NewStub(Channel)), Headers(Headers)
 {
 }
 
@@ -25,6 +25,10 @@ FExecution::~FExecution()
 bool FExecution::Execute(const FExecuteRequest& Request, FExecuteResponse& Response)
 {
 	grpc::ClientContext Context;
+	for (const TPair<FString, FString>& Header : Headers)
+	{
+		Context.AddMetadata(TCHAR_TO_UTF8(*Header.Key.ToLower()), TCHAR_TO_UTF8(*Header.Value));
+	}
 
 	build::bazel::remote::execution::v2::ExecuteRequest ProtoRequest;
 	ProtoConverter::ToProto(Request, ProtoRequest);
@@ -37,7 +41,11 @@ bool FExecution::Execute(const FExecuteRequest& Request, FExecuteResponse& Respo
 		Call->Read(&Operation);
 
 		build::bazel::remote::execution::v2::ExecuteOperationMetadata Metadata;
-		Operation.metadata().UnpackTo(&Metadata);
+		if (!Operation.metadata().UnpackTo(&Metadata))
+		{
+			UE_LOG(LogBazelExecutor, Error, TEXT("Unable to decode metadata"));
+			break;
+		}
 
 		UE_LOG(LogBazelExecutor, Display, TEXT("Execution state: %s"), UTF8_TO_TCHAR(build::bazel::remote::execution::v2::ExecutionStage_Value_Name(Metadata.stage()).c_str()));
 
diff --git a/Engine/Source/Editor/Experimental/BazelExecutor/Private/Execution.h b/Engine/Source/Editor/Experimental/BazelExecutor/Private/Execution.h
--- a/Engine/Source/Editor/Experimental/BazelExecutor/Private/Execution.h
+++ b/Engine/Source/Editor/Experimental/BazelExecutor/Private/Execution.h
@@ -23,9 +23,10 @@ class FExecution : public IExecution
 {
 private:
 	std::unique_ptr<build::bazel::remote::execution::v2::Execution::Stub> Stub;
+	TMap<FString, FString> Headers;
 
 public:
-	FExecution(const std::shared_ptr<grpc::Channel>& Channel);
+	FExecution(const std::shared_ptr<grpc::Channel>& Channel, const TMap<FString, FString>& Headers);
 	~FExecution();
 
 	bool Execute(const FExecuteRequest& Request, FExecuteResponse& Response) override;
diff --git a/Engine/Source/Editor/Experimental/BazelExecutor/Private/ProtoConverter.cpp b/Engine/Source/Editor/Experimental/BazelExecutor/Private/ProtoConverter.cpp
--- a/Engine/Source/Editor/Experimental/BazelExecutor/Private/ProtoConverter.cpp
+++ b/Engine/Source/Editor/Experimental/BazelExecutor/Private/ProtoConverter.cpp
@@ -417,10 +417,10 @@ void ProtoConverter::FromProto(const build::bazel::remote::execution::v2::Execut
 	FromProto(In.message(), Out.Message);
 }
 
-bool ProtoConverter::ToDigest(const TArray<char>& Data, FDigest& OutDigest)
+bool ProtoConverter::ToDigest(const TArray<char>& ToData, FDigest& OutDigest)
 {
-	OutDigest.Hash = FIoHash::HashBuffer(Data.GetData(), Data.Num());
-	OutDigest.SizeBytes = Data.Num();
+	OutDigest.Hash = FIoHash::HashBuffer(ToData.GetData(), ToData.Num());
+	OutDigest.SizeBytes = ToData.Num();
 	return true;
 }
 
diff --git a/Engine/Source/Editor/Experimental/BazelExecutor/Private/ProtoConverter.h b/Engine/Source/Editor/Experimental/BazelExecutor/Private/ProtoConverter.h
--- a/Engine/Source/Editor/Experimental/BazelExecutor/Private/ProtoConverter.h
+++ b/Engine/Source/Editor/Experimental/BazelExecutor/Private/ProtoConverter.h
@@ -149,7 +149,7 @@ public:
 	static void FromProto(const build::bazel::remote::execution::v2::BatchReadBlobsResponse& In, FBatchReadBlobsResponse& Out);
 	static void FromProto(const build::bazel::remote::execution::v2::ExecuteResponse& In, FExecuteResponse& Out);
 
-	static bool ToDigest(const TArray<char>& Data, FDigest& OutDigest);
+	static bool ToDigest(const TArray<char>& InData, FDigest& OutDigest);
 
 	static bool ToBlob(const FDirectory& InDirectory, TArray<char>& OutData, FDigest& OutDigest);
 	static bool ToBlob(const FCommand& InCommand, TArray<char>& OutData, FDigest& OutDigest);
diff --git a/Engine/Source/Editor/Experimental/RemoteExecution/Private/RemoteExecutionSettings.cpp b/Engine/Source/Editor/Experimental/RemoteExecution/Private/RemoteExecutionSettings.cpp
--- a/Engine/Source/Editor/Experimental/RemoteExecution/Private/RemoteExecutionSettings.cpp
+++ b/Engine/Source/Editor/Experimental/RemoteExecution/Private/RemoteExecutionSettings.cpp
@@ -6,5 +6,5 @@
 URemoteExecutionSettings::URemoteExecutionSettings(const FObjectInitializer& ObjectInitializer)
 	: Super(ObjectInitializer)
 {
-	PreferredRemoteExecutor = TEXT("Horde");
+	PreferredRemoteExecutor = TEXT("Bazel");
 }
diff --git a/Engine/Source/Editor/Experimental/RemoteExecution/Public/IContentAddressableStorage.h b/Engine/Source/Editor/Experimental/RemoteExecution/Public/IContentAddressableStorage.h
--- a/Engine/Source/Editor/Experimental/RemoteExecution/Public/IContentAddressableStorage.h
+++ b/Engine/Source/Editor/Experimental/RemoteExecution/Public/IContentAddressableStorage.h
@@ -23,6 +23,7 @@ public:
 	/** Virtual destructor */
 	virtual ~IContentAddressableStorage() {}
 
+	virtual bool ToDigest(const TArray<char>& InData, FDigest& OutDigest) = 0;
 	virtual bool ToBlob(const FDirectory& InDirectory, TArray<char>& OutData, FDigest& OutDigest) = 0;
 	virtual bool ToBlob(const FCommand& InCommand, TArray<char>& OutData, FDigest& OutDigest) = 0;
 	virtual bool ToBlob(const FAction& InAction, TArray<char>& OutData, FDigest& OutDigest) = 0;
diff --git a/Engine/Source/Editor/Experimental/RemoteExecution/RemoteExecution.Build.cs b/Engine/Source/Editor/Experimental/RemoteExecution/RemoteExecution.Build.cs
--- a/Engine/Source/Editor/Experimental/RemoteExecution/RemoteExecution.Build.cs
+++ b/Engine/Source/Editor/Experimental/RemoteExecution/RemoteExecution.Build.cs
@@ -9,15 +9,12 @@ namespace UnrealBuildTool.Rules
 	{
 		public RemoteExecution(ReadOnlyTargetRules Target) : base(Target)
 		{
-			PublicDependencyModuleNames.AddRange(
-				new string[] {
-					"Core",
-					"Settings",
-				}
-				);
+
 
 			PrivateDependencyModuleNames.AddRange(
 				new string[] {
+					"Core",
+					"Settings",
 					"CoreUObject",
 				}
 			);
