RemoteExecution Async API

diff --git a/Engine/Source/Editor/Experimental/BazelExecutor/Private/BazelCompletionQueueRunnable.cpp b/Engine/Source/Editor/Experimental/BazelExecutor/Private/BazelCompletionQueueRunnable.cpp
new file mode 100644
--- /dev/null
+++ b/Engine/Source/Editor/Experimental/BazelExecutor/Private/BazelCompletionQueueRunnable.cpp
@@ -0,0 +1,190 @@
+// Copyright Epic Games, Inc. All Rights Reserved.
+
+#include "BazelCompletionQueueRunnable.h"
+
+#include <chrono>
+
+
+FBazelCompletionQueueRunnable::FBazelCompletionQueueRunnable()
+{
+}
+
+FBazelCompletionQueueRunnable::~FBazelCompletionQueueRunnable()
+{
+}
+
+void FBazelCompletionQueueRunnable::ProcessNext(void* Tag, bool Ok)
+{
+	FQueuedItem* Item = QueuedItems.Find(Tag);
+	if (!Item)
+	{
+		return;
+	}
+
+	switch (Item->State)
+	{
+	case FQueuedItem::EState::Starting:
+	{
+		if (Item->StartCall)
+		{
+			Item->StartCall(Tag, Ok);
+		}
+
+		if (Ok)
+		{
+			google::longrunning::Operation* Operation = (google::longrunning::Operation*)Item->Message.Get();
+			Item->State = FQueuedItem::EState::Reading;
+			Item->Reader->Read(Operation, Tag);
+		}
+		else
+		{
+			Item->State = FQueuedItem::EState::Finishing;
+			Item->Reader->Finish(Item->Status.Get(), Tag);
+		}
+		break;
+	}
+	case FQueuedItem::EState::Reading:
+	{
+		google::longrunning::Operation* Operation = (google::longrunning::Operation*)Item->Message.Get();
+		if (Item->Read)
+		{
+			Item->Read(Tag, Ok, *Operation);
+		}
+
+		if (Operation->done() || !Ok)
+		{
+			Item->State = FQueuedItem::EState::Finishing;
+			Item->Reader->Finish(Item->Status.Get(), Tag);
+		}
+		else
+		{
+			Item->Reader->Read(Operation, Tag);
+		}
+		break;
+	}
+	case FQueuedItem::EState::Finishing:
+	{
+		if (Item->Finish)
+		{
+			Item->Finish(Tag, Ok, *Item->Status.Get(), *Item->Message.Get());
+		}
+
+		QueuedItems.Remove(Tag);
+		break;
+	}
+	}
+}
+
+bool FBazelCompletionQueueRunnable::Init()
+{
+	Running = true;
+	return Running;
+}
+
+uint32 FBazelCompletionQueueRunnable::Run()
+{
+	void* Tag;
+	bool Ok = false;
+	while (Running && CompletionQueue.Next(&Tag, &Ok))
+	{
+		if (!Running)
+		{
+			break;
+		}
+		ProcessNext(Tag, Ok);
+	}
+	return 0;
+}
+
+void FBazelCompletionQueueRunnable::Stop()
+{
+	// Request shutdown of the CompletionQueue
+	Running = false;
+	CompletionQueue.Shutdown();
+}
+
+void FBazelCompletionQueueRunnable::Exit()
+{
+	Stop();
+
+	// Drain the completion queue
+	void* Tag;
+	bool Ok;
+	while (CompletionQueue.Next(&Tag, &Ok)) {}
+
+	for (const TPair<void*,FQueuedItem>& Pair : QueuedItems)
+	{
+		const FQueuedItem& Item = Pair.Value;
+		if (Item.Finish)
+		{
+			Item.Finish(Pair.Key, false, *Item.Status.Get(), *Item.Message.Get());
+		}
+	}
+
+	QueuedItems.Empty();
+}
+
+FSingleThreadRunnable* FBazelCompletionQueueRunnable::GetSingleThreadInterface()
+{
+	return this;
+}
+
+void FBazelCompletionQueueRunnable::Tick()
+{
+	if (!Running || QueuedItems.IsEmpty())
+	{
+		return;
+	}
+
+	void* Tag;
+	bool Ok = false;
+	const std::chrono::time_point Deadline = std::chrono::system_clock::now() + std::chrono::microseconds(100);
+	while (Running && std::chrono::system_clock::now() < Deadline)
+	{
+		grpc::CompletionQueue::NextStatus Status = CompletionQueue.AsyncNext(&Tag, &Ok, Deadline);
+		if (Status == grpc::CompletionQueue::NextStatus::GOT_EVENT)
+		{
+			ProcessNext(Tag, Ok);
+			continue;
+		}
+
+		break;
+	}
+}
+
+bool FBazelCompletionQueueRunnable::AddAsyncOperation(
+	TUniquePtr<grpc::ClientContext> ClientContext,
+	TUniquePtr<grpc::ClientAsyncReader<google::longrunning::Operation>> Reader,
+	TStartCallFunction OnStartCall,
+	TReadFunction OnRead,
+	TFinishFunction OnFinish)
+{
+	if (!Running || !ClientContext.IsValid() || !Reader.IsValid())
+	{
+		return false;
+	}
+
+	grpc::ClientAsyncReader<google::longrunning::Operation>* AsyncReader = Reader.Get();
+	FQueuedItem Item;
+	Item.State = FQueuedItem::EState::Starting;
+	Item.ClientContext = MoveTemp(ClientContext);
+	Item.Reader = MoveTemp(Reader);
+	Item.Message.Reset(new google::longrunning::Operation());
+	Item.Status.Reset(new grpc::Status());
+	Item.StartCall = MoveTemp(OnStartCall);
+	Item.Read = MoveTemp(OnRead);
+	Item.Finish = MoveTemp(OnFinish);
+	QueuedItems.Add(AsyncReader, MoveTemp(Item));
+	AsyncReader->StartCall(AsyncReader);
+	return true;
+}
+
+grpc::CompletionQueue* FBazelCompletionQueueRunnable::GetCompletionQueue()
+{
+	if (!Running)
+	{
+		return nullptr;
+	}
+
+	return &CompletionQueue;
+}
\ No newline at end of file
diff --git a/Engine/Source/Editor/Experimental/BazelExecutor/Private/BazelCompletionQueueRunnable.h b/Engine/Source/Editor/Experimental/BazelExecutor/Private/BazelCompletionQueueRunnable.h
new file mode 100644
--- /dev/null
+++ b/Engine/Source/Editor/Experimental/BazelExecutor/Private/BazelCompletionQueueRunnable.h
@@ -0,0 +1,121 @@
+// Copyright Epic Games, Inc. All Rights Reserved.
+
+#pragma once
+
+#include "CoreMinimal.h"
+#include "HAL/Runnable.h"
+#include "HAL/RunnableThread.h"
+#include "Misc/SingleThreadRunnable.h"
+
+THIRD_PARTY_INCLUDES_START
+UE_PUSH_MACRO("TEXT")
+#undef TEXT
+#include <grpcpp/grpcpp.h>
+#include "build\bazel\remote\execution\v2\remote_execution.grpc.pb.h"
+UE_POP_MACRO("TEXT");
+THIRD_PARTY_INCLUDES_END
+
+
+using TStartCallFunction = TUniqueFunction<void(void* Tag, bool Ok)>;
+using TReadFunction = TUniqueFunction<void(void* Tag, bool Ok, const google::longrunning::Operation& Operation)>;
+using TFinishFunction = TUniqueFunction<void(void* Tag, bool Ok, const grpc::Status& Status, const google::protobuf::Message& Message)>;
+
+class FBazelCompletionQueueRunnable : public FRunnable, FSingleThreadRunnable
+{
+private:
+	struct FQueuedItem
+	{
+		enum class EState
+		{
+			Starting,
+			Reading,
+			Finishing,
+		};
+		EState State;
+		TUniquePtr<grpc::ClientContext> ClientContext;
+		TUniquePtr<grpc::ClientAsyncReader<google::longrunning::Operation>> Reader;
+		TUniquePtr<google::protobuf::Message> Message;
+		TUniquePtr<grpc::Status> Status;
+		TStartCallFunction StartCall;
+		TReadFunction Read;
+		TFinishFunction Finish;
+	};
+
+	bool Running;
+	TMap<void*, FQueuedItem> QueuedItems;
+	grpc::CompletionQueue CompletionQueue;
+
+	void ProcessNext(void* Tag, bool Ok);
+
+public:
+
+	FBazelCompletionQueueRunnable();
+	virtual ~FBazelCompletionQueueRunnable();
+
+	// FRunnable methods.
+	virtual bool Init() override;
+	virtual uint32 Run() override;
+	virtual void Stop() override;
+	virtual void Exit() override;
+
+	FSingleThreadRunnable* GetSingleThreadInterface() override;
+
+	// FSingleThreadRunnable methods.
+	virtual void Tick() override;
+
+	inline bool IsRunning() const { return Running; }
+
+public:
+	bool AddAsyncOperation(
+		TUniquePtr<grpc::ClientContext> ClientContext,
+		TUniquePtr<grpc::ClientAsyncReader<google::longrunning::Operation>> Reader,
+		TStartCallFunction OnStartCall = nullptr,
+		TReadFunction OnRead = nullptr,
+		TFinishFunction OnFinish = nullptr);
+
+	template<class MessageType>
+	bool AddAsyncResponse(
+		TUniquePtr<grpc::ClientContext> ClientContext,
+		TUniquePtr<grpc::ClientAsyncResponseReader<MessageType>> Reader,
+		TFinishFunction OnFinish = nullptr);
+
+	grpc::CompletionQueue* GetCompletionQueue();
+};
+
+template<class MessageType>
+bool FBazelCompletionQueueRunnable::AddAsyncResponse(
+	TUniquePtr<grpc::ClientContext> ClientContext,
+	TUniquePtr<grpc::ClientAsyncResponseReader<MessageType>> Reader,
+	TFinishFunction OnFinish)
+{
+	static_assert(std::is_base_of<google::protobuf::Message, MessageType>::value, "MessageType is not derived from google::protobuf::Message");
+	if (!Running || !ClientContext.IsValid() || !Reader.IsValid())
+	{
+		return false;
+	}
+
+	grpc::ClientAsyncResponseReader<MessageType>* AsyncReader = Reader.Get();
+
+	// Keep Reader in scope for the lifetime of this async call
+	TFinishFunction OnFinishResponse = [Reader = MoveTemp(Reader), OnFinish = MoveTemp(OnFinish)](void* Tag, bool Ok, const grpc::Status& Status, const google::protobuf::Message& Message)
+	{
+		if (OnFinish)
+		{
+			OnFinish(Tag, Ok, Status, Message);
+		}
+	};
+
+	MessageType* Response = new MessageType();
+	grpc::Status* Status = new grpc::Status();
+
+	FQueuedItem Item;
+	Item.State = FQueuedItem::EState::Finishing;
+	Item.Status.Reset(Status);
+	Item.Message.Reset(Response);
+	Item.ClientContext = MoveTemp(ClientContext);
+	Item.Finish = MoveTemp(OnFinishResponse);
+	QueuedItems.Add(AsyncReader, MoveTemp(Item));
+	AsyncReader->StartCall();
+	AsyncReader->Finish(Response, Status, AsyncReader);
+	return true;
+}
diff --git a/Engine/Source/Editor/Experimental/BazelExecutor/Private/BazelExecutor.cpp b/Engine/Source/Editor/Experimental/BazelExecutor/Private/BazelExecutor.cpp
--- a/Engine/Source/Editor/Experimental/BazelExecutor/Private/BazelExecutor.cpp
+++ b/Engine/Source/Editor/Experimental/BazelExecutor/Private/BazelExecutor.cpp
@@ -4,6 +4,7 @@
 
 #include "Execution.h"
 #include "ContentAddressableStorage.h"
+#include "BazelCompletionQueueRunnable.h"
 
 THIRD_PARTY_INCLUDES_START
 UE_PUSH_MACRO("TEXT")
@@ -21,8 +22,10 @@ THIRD_PARTY_INCLUDES_END
 
 void FBazelExecutor::Initialize(const FSettings& Settings)
 {
-	ContentAddressableStorage.Reset();
-	Execution.Reset();
+	Shutdown();
+
+	Runnable = MakeShared<FBazelCompletionQueueRunnable>();
+	Thread.Reset(FRunnableThread::Create(Runnable.Get(), TEXT("BazelExecutorRunnable"), 0, TPri_BelowNormal));
 
 	grpc::ChannelArguments ChannelArguments;
 	ChannelArguments.SetMaxSendMessageSize(Settings.MaxSendMessageSize);
@@ -36,7 +39,7 @@ void FBazelExecutor::Initialize(const FSettings& Settings)
 	std::shared_ptr<grpc::ChannelCredentials> ContentAddressableStorageChannelCredentials = grpc::SslCredentials(grpc::SslCredentialsOptions(ContentAddressableStorageSslOptions));
 
 	std::shared_ptr<grpc::Channel> ContentAddressableStorageChannel = grpc::CreateCustomChannel(TCHAR_TO_UTF8(*Settings.ContentAddressableStorageTarget), ContentAddressableStorageChannelCredentials, ChannelArguments);
-	ContentAddressableStorage.Reset(new FContentAddressableStorage(ContentAddressableStorageChannel, Settings.ContentAddressableStorageHeaders));
+	ContentAddressableStorage.Reset(new FContentAddressableStorage(ContentAddressableStorageChannel, Runnable, Settings.ContentAddressableStorageHeaders));
 
 	grpc::SslCredentialsOptions ExecutionSslOptions;
 	ExecutionSslOptions.pem_cert_chain.assign(TCHAR_TO_UTF8(*Settings.ExecutionPemCertificateChain));
@@ -46,7 +49,25 @@ void FBazelExecutor::Initialize(const FSettings& Settings)
 	std::shared_ptr<grpc::ChannelCredentials> ExecutionChannelCredentials = grpc::SslCredentials(grpc::SslCredentialsOptions(ExecutionSslOptions));
 
 	std::shared_ptr<grpc::Channel> ExecutorChannel = grpc::CreateCustomChannel(TCHAR_TO_UTF8(*Settings.ExecutionTarget), ExecutionChannelCredentials, ChannelArguments);
-	Execution.Reset(new FExecution(ExecutorChannel, Settings.ExecutionHeaders));
+	Execution.Reset(new FExecution(ExecutorChannel, Runnable, Settings.ExecutionHeaders));
+}
+
+void FBazelExecutor::Shutdown()
+{
+	// FFakeThread doesn't call exit, call it manually.
+	if (Runnable.IsValid() && !FPlatformProcess::SupportsMultithreading())
+	{
+		Runnable->Stop();
+		Runnable->Exit();
+	}
+	if (Thread.IsValid())
+	{
+		Thread->Kill();
+	}
+	Thread.Reset();
+	Runnable.Reset();
+	ContentAddressableStorage.Reset();
+	Execution.Reset();
 }
 
 FName FBazelExecutor::GetFName() const
diff --git a/Engine/Source/Editor/Experimental/BazelExecutor/Private/BazelExecutor.h b/Engine/Source/Editor/Experimental/BazelExecutor/Private/BazelExecutor.h
--- a/Engine/Source/Editor/Experimental/BazelExecutor/Private/BazelExecutor.h
+++ b/Engine/Source/Editor/Experimental/BazelExecutor/Private/BazelExecutor.h
@@ -7,6 +7,9 @@
 #include "IRemoteExecutor.h"
 #include "IContentAddressableStorage.h"
 #include "IExecution.h"
+#include "HAL/RunnableThread.h"
+
+class FBazelCompletionQueueRunnable;
 
 
 class FBazelExecutor : public IRemoteExecutor
@@ -31,9 +34,12 @@ public:
 private:
 	TUniquePtr<IContentAddressableStorage> ContentAddressableStorage;
 	TUniquePtr<IExecution> Execution;
+	TUniquePtr<FRunnableThread> Thread;
+	TSharedPtr<FBazelCompletionQueueRunnable> Runnable;
 
 public:
 	void Initialize(const FSettings& Settings);
+	void Shutdown();
 
 	virtual FName GetFName() const override;
 	virtual FText GetNameText() const override;
diff --git a/Engine/Source/Editor/Experimental/BazelExecutor/Private/BazelExecutorModule.cpp b/Engine/Source/Editor/Experimental/BazelExecutor/Private/BazelExecutorModule.cpp
--- a/Engine/Source/Editor/Experimental/BazelExecutor/Private/BazelExecutorModule.cpp
+++ b/Engine/Source/Editor/Experimental/BazelExecutor/Private/BazelExecutorModule.cpp
@@ -106,6 +106,7 @@ void FBazelExecutorModule::StartupModule()
 
 void FBazelExecutorModule::ShutdownModule()
 {
+	BazelExecution.Shutdown();
 	IModularFeatures::Get().UnregisterModularFeature(RemoteExecutionFeatureName, &BazelExecution);
 }
 
diff --git a/Engine/Source/Editor/Experimental/BazelExecutor/Private/ContentAddressableStorage.cpp b/Engine/Source/Editor/Experimental/BazelExecutor/Private/ContentAddressableStorage.cpp
--- a/Engine/Source/Editor/Experimental/BazelExecutor/Private/ContentAddressableStorage.cpp
+++ b/Engine/Source/Editor/Experimental/BazelExecutor/Private/ContentAddressableStorage.cpp
@@ -3,8 +3,11 @@
 #include "ContentAddressableStorage.h"
 #include "BazelExecutorModule.h"
 #include "ProtoConverter.h"
+#include "BazelCompletionQueueRunnable.h"
 #include "Messages.h"
 
+#include <memory>
+
 THIRD_PARTY_INCLUDES_START
 UE_PUSH_MACRO("TEXT")
 #undef TEXT
@@ -13,8 +16,10 @@ UE_POP_MACRO("TEXT");
 THIRD_PARTY_INCLUDES_END
 
 
-FContentAddressableStorage::FContentAddressableStorage(const std::shared_ptr<grpc::Channel>& Channel, const TMap<FString, FString>& Headers) :
-	Stub(build::bazel::remote::execution::v2::ContentAddressableStorage::NewStub(Channel)), Headers(Headers)
+FContentAddressableStorage::FContentAddressableStorage(const std::shared_ptr<grpc::Channel>& Channel, TSharedPtr<FBazelCompletionQueueRunnable> CompletionQueueRunnable, const TMap<FString, FString>& Headers) :
+	Stub(build::bazel::remote::execution::v2::ContentAddressableStorage::NewStub(Channel).release()),
+	CompletionQueueRunnable(CompletionQueueRunnable),
+	Headers(Headers)
 {
 }
 
@@ -42,19 +47,16 @@ bool FContentAddressableStorage::ToBlob(const FAction& InAction, TArray<char>& O
 	return ProtoConverter::ToBlob(InAction, OutData, OutDigest);
 }
 
-FStatus FContentAddressableStorage::FindMissingBlobs(const FFindMissingBlobsRequest& Request, FFindMissingBlobsResponse& Response)
+FStatus FContentAddressableStorage::FindMissingBlobs(const FFindMissingBlobsRequest& Request, FFindMissingBlobsResponse& Response, int64 TimeoutMs)
 {
-	grpc::ClientContext Context;
-	for (const TPair<FString, FString>& Header : Headers)
-	{
-		Context.AddMetadata(TCHAR_TO_UTF8(*Header.Key.ToLower()), TCHAR_TO_UTF8(*Header.Value));
-	}
+	grpc::ClientContext ClientContext;
+	PrepareContext(ClientContext, TimeoutMs);
 
 	build::bazel::remote::execution::v2::FindMissingBlobsRequest ProtoRequest;
 	ProtoConverter::ToProto(Request, ProtoRequest);
 
 	build::bazel::remote::execution::v2::FindMissingBlobsResponse ProtoResponse;
-	grpc::Status ProtoStatus = Stub->FindMissingBlobs(&Context, ProtoRequest, &ProtoResponse);
+	grpc::Status ProtoStatus = Stub->FindMissingBlobs(&ClientContext, ProtoRequest, &ProtoResponse);
 	FStatus Status;
 	ProtoConverter::FromProto(ProtoStatus, Status);
 	ProtoConverter::FromProto(ProtoResponse, Response);
@@ -62,19 +64,16 @@ FStatus FContentAddressableStorage::FindMissingBlobs(const FFindMissingBlobsRequ
 	return MoveTemp(Status);
 }
 
-FStatus FContentAddressableStorage::BatchUpdateBlobs(const FBatchUpdateBlobsRequest& Request, FBatchUpdateBlobsResponse& Response)
+FStatus FContentAddressableStorage::BatchUpdateBlobs(const FBatchUpdateBlobsRequest& Request, FBatchUpdateBlobsResponse& Response, int64 TimeoutMs)
 {
-	grpc::ClientContext Context;
-	for (const TPair<FString, FString>& Header : Headers)
-	{
-		Context.AddMetadata(TCHAR_TO_UTF8(*Header.Key.ToLower()), TCHAR_TO_UTF8(*Header.Value));
-	}
+	grpc::ClientContext ClientContext;
+	PrepareContext(ClientContext, TimeoutMs);
 
 	build::bazel::remote::execution::v2::BatchUpdateBlobsRequest ProtoRequest;
 	ProtoConverter::ToProto(Request, ProtoRequest);
 
 	build::bazel::remote::execution::v2::BatchUpdateBlobsResponse ProtoResponse;
-	grpc::Status ProtoStatus = Stub->BatchUpdateBlobs(&Context, ProtoRequest, &ProtoResponse);
+	grpc::Status ProtoStatus = Stub->BatchUpdateBlobs(&ClientContext, ProtoRequest, &ProtoResponse);
 	FStatus Status;
 	ProtoConverter::FromProto(ProtoStatus, Status);
 	ProtoConverter::FromProto(ProtoResponse, Response);
@@ -82,22 +81,164 @@ FStatus FContentAddressableStorage::BatchUpdateBlobs(const FBatchUpdateBlobsRequ
 	return MoveTemp(Status);
 }
 
-FStatus FContentAddressableStorage::BatchReadBlobs(const FBatchReadBlobsRequest& Request, FBatchReadBlobsResponse& Response)
+FStatus FContentAddressableStorage::BatchReadBlobs(const FBatchReadBlobsRequest& Request, FBatchReadBlobsResponse& Response, int64 TimeoutMs)
 {
-	grpc::ClientContext Context;
-	for (const TPair<FString, FString>& Header : Headers)
-	{
-		Context.AddMetadata(TCHAR_TO_UTF8(*Header.Key.ToLower()), TCHAR_TO_UTF8(*Header.Value));
-	}
+	grpc::ClientContext ClientContext;
+	PrepareContext(ClientContext, TimeoutMs);
 
 	build::bazel::remote::execution::v2::BatchReadBlobsRequest ProtoRequest;
 	ProtoConverter::ToProto(Request, ProtoRequest);
 
 	build::bazel::remote::execution::v2::BatchReadBlobsResponse ProtoResponse;
-	grpc::Status ProtoStatus = Stub->BatchReadBlobs(&Context, ProtoRequest, &ProtoResponse);
+	grpc::Status ProtoStatus = Stub->BatchReadBlobs(&ClientContext, ProtoRequest, &ProtoResponse);
 	FStatus Status;
 	ProtoConverter::FromProto(ProtoStatus, Status);
 	ProtoConverter::FromProto(ProtoResponse, Response);
 
 	return MoveTemp(Status);
 }
+
+TFuture<TPair<FStatus, FFindMissingBlobsResponse>> FContentAddressableStorage::FindMissingBlobsAsync(const FFindMissingBlobsRequest& Request, TUniqueFunction<void()> CompletionCallback, int64 TimeoutMs)
+{
+	TUniquePtr<grpc::ClientContext> ClientContext(new grpc::ClientContext());
+	PrepareContext(*ClientContext.Get(), TimeoutMs);
+
+	build::bazel::remote::execution::v2::FindMissingBlobsRequest ProtoRequest;
+	ProtoConverter::ToProto(Request, ProtoRequest);
+
+	TUniquePtr<grpc::ClientAsyncResponseReader<build::bazel::remote::execution::v2::FindMissingBlobsResponse>> ProtoReader(
+		Stub->PrepareAsyncFindMissingBlobs(ClientContext.Get(), ProtoRequest, CompletionQueueRunnable->GetCompletionQueue()).release());
+
+	TSharedPtr<TPromise<TPair<FStatus, FFindMissingBlobsResponse>>> ReturnPromise = MakeShared<TPromise<TPair<FStatus, FFindMissingBlobsResponse>>>(MoveTemp(CompletionCallback));
+
+	TFinishFunction Finish = [ReturnPromise](void* Tag, bool Ok, const grpc::Status& ProtoStatus, const google::protobuf::Message& ProtoMessage)
+	{
+		const build::bazel::remote::execution::v2::FindMissingBlobsResponse& ProtoResponse = (const build::bazel::remote::execution::v2::FindMissingBlobsResponse&)ProtoMessage;
+		FStatus Status;
+		FFindMissingBlobsResponse Response;
+		if (Ok)
+		{
+			ProtoConverter::FromProto(ProtoStatus, Status);
+			ProtoConverter::FromProto(ProtoResponse, Response);
+		}
+		else
+		{
+			UE_LOG(LogBazelExecutor, Error, TEXT("FindMissingBlobsAsync: Finish !Ok"));
+			Status.Code = EStatusCode::ABORTED;
+		}
+		ReturnPromise->EmplaceValue(TPair<FStatus, FFindMissingBlobsResponse>(MoveTemp(Status), MoveTemp(Response)));
+	};
+
+	if (!CompletionQueueRunnable->AddAsyncResponse<build::bazel::remote::execution::v2::FindMissingBlobsResponse>(
+		MoveTemp(ClientContext), MoveTemp(ProtoReader), MoveTemp(Finish)))
+	{
+		FStatus Status;
+		FFindMissingBlobsResponse Response;
+		Status.Code = EStatusCode::UNAVAILABLE;
+		Status.Message = TEXT("FBazelCompletionQueueRunnable not running");
+		ReturnPromise->EmplaceValue(TPair<FStatus, FFindMissingBlobsResponse>(MoveTemp(Status), MoveTemp(Response)));
+	}
+
+	return ReturnPromise->GetFuture();
+}
+
+
+TFuture<TPair<FStatus, FBatchUpdateBlobsResponse>> FContentAddressableStorage::BatchUpdateBlobsAsync(const FBatchUpdateBlobsRequest& Request, TUniqueFunction<void()> CompletionCallback, int64 TimeoutMs)
+{
+	TUniquePtr<grpc::ClientContext> ClientContext(new grpc::ClientContext());
+	PrepareContext(*ClientContext.Get(), TimeoutMs);
+
+	build::bazel::remote::execution::v2::BatchUpdateBlobsRequest ProtoRequest;
+	ProtoConverter::ToProto(Request, ProtoRequest);
+
+	TUniquePtr<grpc::ClientAsyncResponseReader<build::bazel::remote::execution::v2::BatchUpdateBlobsResponse>> ProtoReader(
+		Stub->PrepareAsyncBatchUpdateBlobs(ClientContext.Get(), ProtoRequest, CompletionQueueRunnable->GetCompletionQueue()).release());
+
+	TSharedPtr<TPromise<TPair<FStatus, FBatchUpdateBlobsResponse>>> ReturnPromise = MakeShared<TPromise<TPair<FStatus, FBatchUpdateBlobsResponse>>>(MoveTemp(CompletionCallback));
+
+	TFinishFunction Finish = [ReturnPromise](void* Tag, bool Ok, const grpc::Status& ProtoStatus, const google::protobuf::Message& ProtoMessage)
+	{
+		const build::bazel::remote::execution::v2::BatchUpdateBlobsResponse& ProtoResponse = (const build::bazel::remote::execution::v2::BatchUpdateBlobsResponse&)ProtoMessage;
+		FStatus Status;
+		FBatchUpdateBlobsResponse Response;
+		if (Ok)
+		{
+			ProtoConverter::FromProto(ProtoStatus, Status);
+			ProtoConverter::FromProto(ProtoResponse, Response);
+		}
+		else
+		{
+			UE_LOG(LogBazelExecutor, Error, TEXT("BatchUpdateBlobsAsync: Finish !Ok"));
+			Status.Code = EStatusCode::ABORTED;
+		}
+		ReturnPromise->EmplaceValue(TPair<FStatus, FBatchUpdateBlobsResponse>(MoveTemp(Status), MoveTemp(Response)));
+	};
+
+	if (!CompletionQueueRunnable->AddAsyncResponse<build::bazel::remote::execution::v2::BatchUpdateBlobsResponse>(
+		MoveTemp(ClientContext), MoveTemp(ProtoReader), MoveTemp(Finish)))
+	{
+		FStatus Status;
+		FBatchUpdateBlobsResponse Response;
+		Status.Code = EStatusCode::UNAVAILABLE;
+		Status.Message = TEXT("FBazelCompletionQueueRunnable not running");
+		ReturnPromise->EmplaceValue(TPair<FStatus, FBatchUpdateBlobsResponse>(MoveTemp(Status), MoveTemp(Response)));
+	}
+
+	return ReturnPromise->GetFuture();
+}
+
+TFuture<TPair<FStatus, FBatchReadBlobsResponse>> FContentAddressableStorage::BatchReadBlobsAsync(const FBatchReadBlobsRequest& Request, TUniqueFunction<void()> CompletionCallback, int64 TimeoutMs)
+{
+	TUniquePtr<grpc::ClientContext> ClientContext(new grpc::ClientContext());
+	PrepareContext(*ClientContext.Get(), TimeoutMs);
+
+	build::bazel::remote::execution::v2::BatchReadBlobsRequest ProtoRequest;
+	ProtoConverter::ToProto(Request, ProtoRequest);
+
+	TUniquePtr<grpc::ClientAsyncResponseReader<build::bazel::remote::execution::v2::BatchReadBlobsResponse>> ProtoReader(
+		Stub->PrepareAsyncBatchReadBlobs(ClientContext.Get(), ProtoRequest, CompletionQueueRunnable->GetCompletionQueue()).release());
+
+	TSharedPtr<TPromise<TPair<FStatus, FBatchReadBlobsResponse>>> ReturnPromise = MakeShared<TPromise<TPair<FStatus, FBatchReadBlobsResponse>>>(MoveTemp(CompletionCallback));
+
+	TFinishFunction Finish = [ReturnPromise](void* Tag, bool Ok, const grpc::Status& ProtoStatus, const google::protobuf::Message& ProtoMessage)
+	{
+		const build::bazel::remote::execution::v2::BatchReadBlobsResponse& ProtoResponse = (const build::bazel::remote::execution::v2::BatchReadBlobsResponse&)ProtoMessage;
+		FStatus Status;
+		FBatchReadBlobsResponse Response;
+		if (Ok)
+		{
+			ProtoConverter::FromProto(ProtoStatus, Status);
+			ProtoConverter::FromProto(ProtoResponse, Response);
+		}
+		else
+		{
+			UE_LOG(LogBazelExecutor, Error, TEXT("BatchReadBlobsAsync: Finish !Ok"));
+			Status.Code = EStatusCode::ABORTED;
+		}
+		ReturnPromise->EmplaceValue(TPair<FStatus, FBatchReadBlobsResponse>(MoveTemp(Status), MoveTemp(Response)));
+	};
+
+	if (!CompletionQueueRunnable->AddAsyncResponse<build::bazel::remote::execution::v2::BatchReadBlobsResponse>(
+		MoveTemp(ClientContext), MoveTemp(ProtoReader), MoveTemp(Finish)))
+	{
+		FStatus Status;
+		FBatchReadBlobsResponse Response;
+		Status.Code = EStatusCode::UNAVAILABLE;
+		Status.Message = TEXT("FBazelCompletionQueueRunnable not running");
+		ReturnPromise->EmplaceValue(TPair<FStatus, FBatchReadBlobsResponse>(MoveTemp(Status), MoveTemp(Response)));
+	}
+
+	return ReturnPromise->GetFuture();
+}
+
+void FContentAddressableStorage::PrepareContext(grpc::ClientContext& ClientContext, int64 TimeoutMs) const
+{
+	for (const TPair<FString, FString>& Header : Headers)
+	{
+		ClientContext.AddMetadata(TCHAR_TO_UTF8(*Header.Key.ToLower()), TCHAR_TO_UTF8(*Header.Value));
+	}
+	if (TimeoutMs > 0)
+	{
+		ClientContext.set_deadline(std::chrono::system_clock::now() + std::chrono::milliseconds(TimeoutMs));
+	}
+}
diff --git a/Engine/Source/Editor/Experimental/BazelExecutor/Private/ContentAddressableStorage.h b/Engine/Source/Editor/Experimental/BazelExecutor/Private/ContentAddressableStorage.h
--- a/Engine/Source/Editor/Experimental/BazelExecutor/Private/ContentAddressableStorage.h
+++ b/Engine/Source/Editor/Experimental/BazelExecutor/Private/ContentAddressableStorage.h
@@ -5,8 +5,6 @@
 #include "CoreMinimal.h"
 #include "IContentAddressableStorage.h"
 
-#include <memory>
-
 THIRD_PARTY_INCLUDES_START
 UE_PUSH_MACRO("TEXT")
 #undef TEXT
@@ -14,20 +12,18 @@ UE_PUSH_MACRO("TEXT")
 UE_POP_MACRO("TEXT");
 THIRD_PARTY_INCLUDES_END
 
-namespace grpc
-{
-	class Channel;
-}
+class FBazelCompletionQueueRunnable;
 
 
 class FContentAddressableStorage : public IContentAddressableStorage
 {
 private:
-	std::unique_ptr<build::bazel::remote::execution::v2::ContentAddressableStorage::Stub> Stub;
+	TUniquePtr<build::bazel::remote::execution::v2::ContentAddressableStorage::Stub> Stub;
+	TSharedPtr<FBazelCompletionQueueRunnable> CompletionQueueRunnable;
 	TMap<FString, FString> Headers;
 
 public:
-	FContentAddressableStorage(const std::shared_ptr<grpc::Channel>& Channel, const TMap<FString, FString>& Headers);
+	FContentAddressableStorage(const std::shared_ptr<grpc::Channel>& Channel, TSharedPtr<FBazelCompletionQueueRunnable> CompletionQueueRunnable, const TMap<FString, FString>& Headers);
 	~FContentAddressableStorage();
 
 	bool ToDigest(const TArray<char>& InData, FDigest& OutDigest) override;
@@ -35,7 +31,14 @@ public:
 	bool ToBlob(const FCommand& InCommand, TArray<char>& OutData, FDigest& OutDigest) override;
 	bool ToBlob(const FAction& InAction, TArray<char>& OutData, FDigest& OutDigest) override;
 
-	FStatus FindMissingBlobs(const FFindMissingBlobsRequest& Request, FFindMissingBlobsResponse& Response) override;
-	FStatus BatchUpdateBlobs(const FBatchUpdateBlobsRequest& Request, FBatchUpdateBlobsResponse& Response) override;
-	FStatus BatchReadBlobs(const FBatchReadBlobsRequest& Request, FBatchReadBlobsResponse& Response) override;
+	FStatus FindMissingBlobs(const FFindMissingBlobsRequest& Request, FFindMissingBlobsResponse& Response, int64 TimeoutMs = 0) override;
+	FStatus BatchUpdateBlobs(const FBatchUpdateBlobsRequest& Request, FBatchUpdateBlobsResponse& Response, int64 TimeoutMs = 0) override;
+	FStatus BatchReadBlobs(const FBatchReadBlobsRequest& Request, FBatchReadBlobsResponse& Response, int64 TimeoutMs = 0) override;
+
+	TFuture<TPair<FStatus, FFindMissingBlobsResponse>> FindMissingBlobsAsync(const FFindMissingBlobsRequest& Request, TUniqueFunction<void()> CompletionCallback = nullptr, int64 TimeoutMs = 0) override;
+	TFuture<TPair<FStatus, FBatchUpdateBlobsResponse>> BatchUpdateBlobsAsync(const FBatchUpdateBlobsRequest& Request, TUniqueFunction<void()> CompletionCallback = nullptr, int64 TimeoutMs = 0) override;
+	TFuture<TPair<FStatus, FBatchReadBlobsResponse>> BatchReadBlobsAsync(const FBatchReadBlobsRequest& Request, TUniqueFunction<void()> CompletionCallback = nullptr, int64 TimeoutMs = 0) override;
+
+private:
+	void PrepareContext(grpc::ClientContext& ClientContext, int64 TimeoutMs = 0) const;
 };
diff --git a/Engine/Source/Editor/Experimental/BazelExecutor/Private/Execution.cpp b/Engine/Source/Editor/Experimental/BazelExecutor/Private/Execution.cpp
--- a/Engine/Source/Editor/Experimental/BazelExecutor/Private/Execution.cpp
+++ b/Engine/Source/Editor/Experimental/BazelExecutor/Private/Execution.cpp
@@ -3,8 +3,11 @@
 #include "Execution.h"
 #include "BazelExecutorModule.h"
 #include "ProtoConverter.h"
+#include "BazelCompletionQueueRunnable.h"
 #include "Messages.h"
 
+#include <memory>
+
 THIRD_PARTY_INCLUDES_START
 UE_PUSH_MACRO("TEXT")
 #undef TEXT
@@ -13,8 +16,10 @@ UE_POP_MACRO("TEXT");
 THIRD_PARTY_INCLUDES_END
 
 
-FExecution::FExecution(const std::shared_ptr<grpc::Channel>& Channel, const TMap<FString, FString>& Headers) :
-	Stub(build::bazel::remote::execution::v2::Execution::NewStub(Channel)), Headers(Headers)
+FExecution::FExecution(const std::shared_ptr<grpc::Channel>& Channel, TSharedPtr<FBazelCompletionQueueRunnable> CompletionQueueRunnable, const TMap<FString, FString>& Headers) :
+	Stub(build::bazel::remote::execution::v2::Execution::NewStub(Channel).release()),
+	CompletionQueueRunnable(CompletionQueueRunnable),
+	Headers(Headers)
 {
 }
 
@@ -22,18 +27,15 @@ FExecution::~FExecution()
 {
 }
 
-bool FExecution::Execute(const FExecuteRequest& Request, FExecuteResponse& Response)
+bool FExecution::Execute(const FExecuteRequest& Request, FExecuteResponse& Response, int64 TimeoutMs)
 {
-	grpc::ClientContext Context;
-	for (const TPair<FString, FString>& Header : Headers)
-	{
-		Context.AddMetadata(TCHAR_TO_UTF8(*Header.Key.ToLower()), TCHAR_TO_UTF8(*Header.Value));
-	}
+	grpc::ClientContext ClientContext;
+	PrepareContext(ClientContext, TimeoutMs);
 
 	build::bazel::remote::execution::v2::ExecuteRequest ProtoRequest;
 	ProtoConverter::ToProto(Request, ProtoRequest);
 
-	std::unique_ptr<grpc::ClientReader<google::longrunning::Operation>> Call = Stub->Execute(&Context, ProtoRequest);
+	std::unique_ptr<grpc::ClientReader<google::longrunning::Operation>> Call = Stub->Execute(&ClientContext, ProtoRequest);
 	Call->WaitForInitialMetadata();
 	while (true)
 	{
@@ -43,28 +45,134 @@ bool FExecution::Execute(const FExecuteRequest& Request, FExecuteResponse& Respo
 		build::bazel::remote::execution::v2::ExecuteOperationMetadata Metadata;
 		if (!Operation.metadata().UnpackTo(&Metadata))
 		{
-			UE_LOG(LogBazelExecutor, Error, TEXT("Unable to decode metadata"));
+			UE_LOG(LogBazelExecutor, Error, TEXT("Execute: Unable to decode metadata"));
 			break;
 		}
 
-		UE_LOG(LogBazelExecutor, Display, TEXT("Execution state: %s"), UTF8_TO_TCHAR(build::bazel::remote::execution::v2::ExecutionStage_Value_Name(Metadata.stage()).c_str()));
+		UE_LOG(LogBazelExecutor, Display, TEXT("Execute: Execution state: %s"), UTF8_TO_TCHAR(build::bazel::remote::execution::v2::ExecutionStage_Value_Name(Metadata.stage()).c_str()));
 
 		if (Operation.done())
 		{
 			build::bazel::remote::execution::v2::ExecuteResponse ProtoResponse;
 			if (!Operation.response().UnpackTo(&ProtoResponse))
 			{
-				UE_LOG(LogBazelExecutor, Error, TEXT("Unable to decode response"));
+				UE_LOG(LogBazelExecutor, Error, TEXT("Execute: Unable to decode response"));
 				break;
 			}
+			ProtoConverter::FromProto(ProtoResponse, Response);
 			if ((grpc::StatusCode)ProtoResponse.status().code() != grpc::StatusCode::OK) {
 				UE_LOG(LogBazelExecutor, Error, TEXT("%s"), UTF8_TO_TCHAR(ProtoResponse.status().message().c_str()));
 				break;
 			}
-			ProtoConverter::FromProto(ProtoResponse, Response);
 			return true;
 		}
 	}
 
 	return false;
 }
+
+TFuture<FExecuteResponse> FExecution::ExecuteAsync(const FExecuteRequest& Request, TUniqueFunction<void()> CompletionCallback, int64 TimeoutMs)
+{
+	TUniquePtr<grpc::ClientContext> ClientContext(new grpc::ClientContext());
+	PrepareContext(*ClientContext.Get(), TimeoutMs);
+
+	build::bazel::remote::execution::v2::ExecuteRequest ProtoRequest;
+	ProtoConverter::ToProto(Request, ProtoRequest);
+
+	TUniquePtr<grpc::ClientAsyncReader<google::longrunning::Operation>> OperationReader(
+		Stub->PrepareAsyncExecute(ClientContext.Get(), ProtoRequest, CompletionQueueRunnable->GetCompletionQueue()).release());
+
+	TSharedPtr<TPromise<FExecuteResponse>> ReturnPromise = MakeShared<TPromise<FExecuteResponse>>(MoveTemp(CompletionCallback));
+
+	TStartCallFunction StartCall = [](void* Tag, bool Ok)
+	{
+		if (!Ok)
+		{
+			UE_LOG(LogBazelExecutor, Error, TEXT("ExecuteAsync: Call Started !Ok"));
+			return;
+		}
+	};
+
+	TReadFunction Read = [](void* Tag, bool Ok, const google::longrunning::Operation& Operation)
+	{
+		if (!Ok)
+		{
+			UE_LOG(LogBazelExecutor, Error, TEXT("ExecuteAsync: Read !Ok"));
+			return;
+		}
+
+		build::bazel::remote::execution::v2::ExecuteOperationMetadata Metadata;
+		if (!Operation.metadata().UnpackTo(&Metadata))
+		{
+			UE_LOG(LogBazelExecutor, Error, TEXT("ExecuteAsync: Unable to decode metadata"));
+			return;
+		}
+
+		UE_LOG(LogBazelExecutor, Display, TEXT("ExecuteAsync: Execution state: %s"), UTF8_TO_TCHAR(build::bazel::remote::execution::v2::ExecutionStage_Value_Name(Metadata.stage()).c_str()));
+	};
+
+	TFinishFunction Finish = [ReturnPromise](void* Tag, bool Ok, const grpc::Status& ProtoStatus, const google::protobuf::Message& Message)
+	{
+		const google::longrunning::Operation& Operation = (const google::longrunning::Operation&)Message;
+		FStatus Status;
+		if (Ok)
+		{
+			ProtoConverter::FromProto(ProtoStatus, Status);
+		}
+		else
+		{
+			UE_LOG(LogBazelExecutor, Error, TEXT("ExecuteAsync: Finish !Ok"));
+			Status.Code = EStatusCode::ABORTED;
+		}
+
+
+		if (!Status.Ok())
+		{
+			FExecuteResponse Response;
+			Response.Status = Status;
+			ReturnPromise->EmplaceValue(MoveTemp(Response));
+			return;
+		}
+
+		build::bazel::remote::execution::v2::ExecuteResponse ProtoResponse;
+		if (!Operation.response().UnpackTo(&ProtoResponse))
+		{
+			UE_LOG(LogBazelExecutor, Error, TEXT("ExecuteAsync: Unable to decode response"));
+			FExecuteResponse Response;
+			Response.Status.Code = EStatusCode::INTERNAL;
+			Response.Status.Message = TEXT("Unable to decode response");
+			ReturnPromise->EmplaceValue(MoveTemp(Response));
+			return;
+		}
+
+		FExecuteResponse Response;
+		ProtoConverter::FromProto(ProtoResponse, Response);
+		ReturnPromise->EmplaceValue(MoveTemp(Response));
+	};
+
+	if (!CompletionQueueRunnable->AddAsyncOperation(MoveTemp(ClientContext),
+		MoveTemp(OperationReader),
+		MoveTemp(StartCall),
+		MoveTemp(Read),
+		MoveTemp(Finish)))
+	{
+		FExecuteResponse Response;
+		Response.Status.Code = EStatusCode::UNAVAILABLE;
+		Response.Status.Message = TEXT("FBazelCompletionQueueRunnable not running");
+		ReturnPromise->EmplaceValue(MoveTemp(Response));
+	}
+
+	return ReturnPromise->GetFuture();
+}
+
+void FExecution::PrepareContext(grpc::ClientContext& ClientContext, int64 TimeoutMs) const
+{
+	for (const TPair<FString, FString>& Header : Headers)
+	{
+		ClientContext.AddMetadata(TCHAR_TO_UTF8(*Header.Key.ToLower()), TCHAR_TO_UTF8(*Header.Value));
+	}
+	if (TimeoutMs > 0)
+	{
+		ClientContext.set_deadline(std::chrono::system_clock::now() + std::chrono::milliseconds(TimeoutMs));
+	}
+}
diff --git a/Engine/Source/Editor/Experimental/BazelExecutor/Private/Execution.h b/Engine/Source/Editor/Experimental/BazelExecutor/Private/Execution.h
--- a/Engine/Source/Editor/Experimental/BazelExecutor/Private/Execution.h
+++ b/Engine/Source/Editor/Experimental/BazelExecutor/Private/Execution.h
@@ -8,26 +8,28 @@
 THIRD_PARTY_INCLUDES_START
 UE_PUSH_MACRO("TEXT")
 #undef TEXT
-#include <memory>
 #include "build\bazel\remote\execution\v2\remote_execution.grpc.pb.h"
 UE_POP_MACRO("TEXT");
 THIRD_PARTY_INCLUDES_END
 
-namespace grpc
-{
-	class Channel;
-}
+class FBazelCompletionQueueRunnable;
 
 
 class FExecution : public IExecution
 {
 private:
-	std::unique_ptr<build::bazel::remote::execution::v2::Execution::Stub> Stub;
+	TUniquePtr<build::bazel::remote::execution::v2::Execution::Stub> Stub;
+	TSharedPtr<FBazelCompletionQueueRunnable> CompletionQueueRunnable;
 	TMap<FString, FString> Headers;
 
 public:
-	FExecution(const std::shared_ptr<grpc::Channel>& Channel, const TMap<FString, FString>& Headers);
+	FExecution(const std::shared_ptr<grpc::Channel>& Channel, TSharedPtr<FBazelCompletionQueueRunnable> CompletionQueueRunnable, const TMap<FString, FString>& Headers);
 	~FExecution();
 
-	bool Execute(const FExecuteRequest& Request, FExecuteResponse& Response) override;
+	bool Execute(const FExecuteRequest& Request, FExecuteResponse& Response, int64 TimeoutMs = 0) override;
+
+	TFuture<FExecuteResponse> ExecuteAsync(const FExecuteRequest& Request, TUniqueFunction<void()> CompletionCallback = nullptr, int64 TimeoutMs = 0) override;
+
+private:
+	void PrepareContext(grpc::ClientContext& ClientContext, int64 TimeoutMs = 0) const;
 };
diff --git a/Engine/Source/Editor/Experimental/RemoteExecution/Public/IContentAddressableStorage.h b/Engine/Source/Editor/Experimental/RemoteExecution/Public/IContentAddressableStorage.h
--- a/Engine/Source/Editor/Experimental/RemoteExecution/Public/IContentAddressableStorage.h
+++ b/Engine/Source/Editor/Experimental/RemoteExecution/Public/IContentAddressableStorage.h
@@ -3,6 +3,7 @@
 #pragma once
 
 #include "CoreMinimal.h"
+#include "Async/Future.h"
 
 struct FAction;
 struct FBatchReadBlobsRequest;
@@ -28,7 +29,11 @@ public:
 	virtual bool ToBlob(const FCommand& InCommand, TArray<char>& OutData, FDigest& OutDigest) = 0;
 	virtual bool ToBlob(const FAction& InAction, TArray<char>& OutData, FDigest& OutDigest) = 0;
 
-	virtual FStatus FindMissingBlobs(const FFindMissingBlobsRequest& Request, FFindMissingBlobsResponse& Response) = 0;
-	virtual FStatus BatchUpdateBlobs(const FBatchUpdateBlobsRequest& Request, FBatchUpdateBlobsResponse& Response) = 0;
-	virtual FStatus BatchReadBlobs(const FBatchReadBlobsRequest& Request, FBatchReadBlobsResponse& Response) = 0;
+	virtual FStatus FindMissingBlobs(const FFindMissingBlobsRequest& Request, FFindMissingBlobsResponse& Response, int64 TimeoutMs = 0) = 0;
+	virtual FStatus BatchUpdateBlobs(const FBatchUpdateBlobsRequest& Request, FBatchUpdateBlobsResponse& Response, int64 TimeoutMs = 0) = 0;
+	virtual FStatus BatchReadBlobs(const FBatchReadBlobsRequest& Request, FBatchReadBlobsResponse& Response, int64 TimeoutMs = 0) = 0;
+
+	virtual TFuture<TPair<FStatus, FFindMissingBlobsResponse>> FindMissingBlobsAsync(const FFindMissingBlobsRequest& Request, TUniqueFunction<void()> CompletionCallback = nullptr, int64 TimeoutMs = 0) = 0;
+	virtual TFuture<TPair<FStatus, FBatchUpdateBlobsResponse>> BatchUpdateBlobsAsync(const FBatchUpdateBlobsRequest& Request, TUniqueFunction<void()> CompletionCallback = nullptr, int64 TimeoutMs = 0) = 0;
+	virtual TFuture<TPair<FStatus, FBatchReadBlobsResponse>> BatchReadBlobsAsync(const FBatchReadBlobsRequest& Request, TUniqueFunction<void()> CompletionCallback = nullptr, int64 TimeoutMs = 0) = 0;
 };
diff --git a/Engine/Source/Editor/Experimental/RemoteExecution/Public/IExecution.h b/Engine/Source/Editor/Experimental/RemoteExecution/Public/IExecution.h
--- a/Engine/Source/Editor/Experimental/RemoteExecution/Public/IExecution.h
+++ b/Engine/Source/Editor/Experimental/RemoteExecution/Public/IExecution.h
@@ -3,6 +3,7 @@
 #pragma once
 
 #include "CoreMinimal.h"
+#include "Async/Future.h"
 
 struct FExecuteRequest;
 struct FExecuteResponse;
@@ -14,5 +15,7 @@ public:
 	/** Virtual destructor */
 	virtual ~IExecution() {}
 
-	virtual bool Execute(const FExecuteRequest& Request, FExecuteResponse& Response) = 0;
+	virtual bool Execute(const FExecuteRequest& Request, FExecuteResponse& Response, int64 TimeoutMs = 0) = 0;
+
+	virtual TFuture<FExecuteResponse> ExecuteAsync(const FExecuteRequest& Request, TUniqueFunction<void()> CompletionCallback = nullptr, int64 TimeoutMs = 0) = 0;
 };
