Added Editor Gizmo Subsystem and Editor Gizmo Manager to facilitate registration and creation of selection-based gizmos.

diff --git a/Engine/Source/Editor/Experimental/EditorInteractiveToolsFramework/EditorInteractiveToolsFramework.Build.cs b/Engine/Source/Editor/Experimental/EditorInteractiveToolsFramework/EditorInteractiveToolsFramework.Build.cs
--- a/Engine/Source/Editor/Experimental/EditorInteractiveToolsFramework/EditorInteractiveToolsFramework.Build.cs
+++ b/Engine/Source/Editor/Experimental/EditorInteractiveToolsFramework/EditorInteractiveToolsFramework.Build.cs
@@ -50,7 +50,10 @@ public class EditorInteractiveToolsFramework : ModuleRules
                 "EditorStyle",
                 "InteractiveToolsFramework",
 				"MeshDescription",
-				"StaticMeshDescription"
+				"StaticMeshDescription",
+                "EditorSubsystem",
+                "TypedElementFramework",
+                "TypedElementRuntime"
 
 				// ... add private dependencies that you statically link with here ...	
 			}
diff --git a/Engine/Source/Editor/Experimental/EditorInteractiveToolsFramework/Private/EdModeInteractiveToolsContext.cpp b/Engine/Source/Editor/Experimental/EditorInteractiveToolsFramework/Private/EdModeInteractiveToolsContext.cpp
--- a/Engine/Source/Editor/Experimental/EditorInteractiveToolsFramework/Private/EdModeInteractiveToolsContext.cpp
+++ b/Engine/Source/Editor/Experimental/EditorInteractiveToolsFramework/Private/EdModeInteractiveToolsContext.cpp
@@ -28,6 +28,7 @@
 #include "ToolContextInterfaces.h"
 #include "InteractiveToolObjects.h"
 #include "InteractiveToolsSelectionStoreSubsystem.h"
+#include "EditorInteractiveGizmoManager.h"
 #include "BaseBehaviors/ClickDragBehavior.h"
 #include "EditorModeManager.h"
 #include "EdMode.h"
@@ -613,6 +614,14 @@ void UEdModeInteractiveToolsContext::InitializeContextWithEditorModeManager(FEdi
 	this->TransactionAPI = new FEdModeToolsContextTransactionImpl(this, InEditorModeManager);
 	this->QueriesAPI = new FEdModeToolsContextQueriesImpl(this, InEditorModeManager);
 
+	SetCreateGizmoManagerFunc([this](const FContextInitInfo& ContextInfo)
+	{
+		UEditorInteractiveGizmoManager* NewGizmoManager = NewObject<UEditorInteractiveGizmoManager>(ContextInfo.ToolsContext);
+		NewGizmoManager->InitializeWithEditorModeManager(ContextInfo.QueriesAPI, ContextInfo.TransactionsAPI, ContextInfo.InputRouter, EditorModeManager);
+		NewGizmoManager->RegisterDefaultGizmos();
+		return NewGizmoManager;
+	});
+
 	Initialize(QueriesAPI, TransactionAPI);
 
 	// enable auto invalidation in Editor, because invalidating for all hover and capture events is unpleasant
diff --git a/Engine/Source/Editor/Experimental/EditorInteractiveToolsFramework/Private/EditorInteractiveGizmoManager.cpp b/Engine/Source/Editor/Experimental/EditorInteractiveToolsFramework/Private/EditorInteractiveGizmoManager.cpp
new file mode 100644
--- /dev/null
+++ b/Engine/Source/Editor/Experimental/EditorInteractiveToolsFramework/Private/EditorInteractiveGizmoManager.cpp
@@ -0,0 +1,310 @@
+// Copyright Epic Games, Inc. All Rights Reserved.
+
+
+#include "EditorInteractiveGizmoManager.h"
+#include "EditorInteractiveGizmoSelectionBuilder.h"
+#include "EditorInteractiveGizmoSubsystem.h"
+#include "EdModeInteractiveToolsContext.h"
+#include "EditorModeManager.h"
+#include "HAL/IConsoleManager.h"
+#include "InputRouter.h"
+#include "InteractiveGizmo.h"
+#include "InteractiveToolChange.h"
+#include "InteractiveToolsContext.h"
+#include "ToolContextInterfaces.h"
+
+
+#define LOCTEXT_NAMESPACE "UEditorInteractiveGizmoManager"
+
+#if 0
+static TAutoConsoleVariable<int32> CVarUseLegacyWidget(
+	TEXT("Gizmos.UseLegacyWidget"),
+	1,
+	TEXT("Specify whether to use selection-based gizmos or legacy widget\n")
+	TEXT("0 = enable UE5 transform and other selection-based gizmos.\n")
+	TEXT("1 = enable legacy UE4 transform widget."),
+	ECVF_RenderThreadSafe);
+#endif
+
+
+UEditorInteractiveGizmoManager::UEditorInteractiveGizmoManager() :
+	UInteractiveGizmoManager()
+{
+
+}
+
+
+void UEditorInteractiveGizmoManager::InitializeWithEditorModeManager(IToolsContextQueriesAPI* QueriesAPIIn, IToolsContextTransactionsAPI* TransactionsAPIIn, UInputRouter* InputRouterIn, FEditorModeTools* InEditorModeManager)
+{
+	Super::Initialize(QueriesAPIIn, TransactionsAPIIn, InputRouterIn);
+	EditorModeManager = InEditorModeManager;
+
+}
+
+
+void UEditorInteractiveGizmoManager::Shutdown()
+{
+	DestroyAllSelectionGizmos();
+	Super::Shutdown();
+}
+
+void UEditorInteractiveGizmoManager::RegisterGizmoSelectionType(const TObjectPtr<UEditorInteractiveGizmoSelectionBuilder> InGizmoSelectionBuilder)
+{
+	if (ensure(InGizmoSelectionBuilder))
+	{
+		if (GizmoSelectionBuilders.Contains(InGizmoSelectionBuilder))
+		{
+			DisplayMessage(
+				FText::Format(LOCTEXT("DeregisterFailedMessage", "UInteractiveGizmoSubsystem::DeregisterGizmoSelectionType: type has already been registered {0}"), FText::FromName(InGizmoSelectionBuilder->GetFName())),
+				EToolMessageLevel::Internal);
+			return;
+		}
+
+		GizmoSelectionBuilders.Add(InGizmoSelectionBuilder);
+		GizmoSelectionBuilders.StableSort(
+			[](UEditorInteractiveGizmoSelectionBuilder& A, UEditorInteractiveGizmoSelectionBuilder& B) {
+				return (A).GetPriority() > (B).GetPriority();
+			});
+	}
+}
+
+bool UEditorInteractiveGizmoManager::DeregisterGizmoSelectionType(const TObjectPtr<UEditorInteractiveGizmoSelectionBuilder> InGizmoSelectionBuilder)
+{
+	if (ensure(InGizmoSelectionBuilder))
+	{
+		if (GizmoSelectionBuilders.Contains(InGizmoSelectionBuilder) == false)
+		{
+			DisplayMessage(
+				FText::Format(LOCTEXT("DeregisterFailedMessage", "UInteractiveGizmoSubsystem::DeregisterGizmoSelectionType: could not find requested type {0}"), FText::FromName(InGizmoSelectionBuilder->GetFName())),
+				EToolMessageLevel::Internal);
+			return false;
+		}
+		GizmoSelectionBuilders.Remove(InGizmoSelectionBuilder);
+
+		return true;
+	}
+
+	return false;
+
+}
+
+TArray<TObjectPtr<UEditorInteractiveGizmoSelectionBuilder>> UEditorInteractiveGizmoManager::GetQualifiedGizmoSelectionBuilders(const FToolBuilderState& InToolBuilderState)
+{
+	TArray<TObjectPtr<UEditorInteractiveGizmoSelectionBuilder>> FoundBuilders;
+
+	FEditorGizmoTypePriority FoundPriority = 0;
+
+	for (TObjectPtr<UEditorInteractiveGizmoSelectionBuilder> Builder : GizmoSelectionBuilders)
+	{
+		if (Builder->GetPriority() < FoundPriority)
+		{
+			break;
+		}
+
+		if (Builder->SatisfiesCondition(InToolBuilderState))
+		{
+			FoundBuilders.Add(Builder);
+			FoundPriority = Builder->GetPriority();
+		}
+	}
+
+	if (!bSearchLocalBuildersOnly)
+	{
+		UEditorInteractiveGizmoSubsystem* GizmoSubsystem = GEditor->GetEditorSubsystem<UEditorInteractiveGizmoSubsystem>();
+		if (ensure(GizmoSubsystem))
+		{
+			TArray<TObjectPtr<UEditorInteractiveGizmoSelectionBuilder>> FoundSubsystemBuilders = GizmoSubsystem->GetQualifiedGizmoSelectionBuilders(InToolBuilderState);
+
+			FEditorGizmoTypePriority FoundPriority0 = FoundBuilders.Num() > 0 ? FoundBuilders[0]->GetPriority() : 0;
+			FEditorGizmoTypePriority FoundPriority1 = FoundSubsystemBuilders.Num() > 0 ? FoundSubsystemBuilders[0]->GetPriority() : 0;
+
+			if (FoundPriority0 == FoundPriority1)
+			{
+				FoundBuilders.Append(FoundSubsystemBuilders);
+			}
+			else if (FoundPriority0 < FoundPriority1)
+			{
+				FoundBuilders = FoundSubsystemBuilders;
+			}
+		}
+	}
+
+	return FoundBuilders;
+}
+
+TArray<UInteractiveGizmo*> UEditorInteractiveGizmoManager::CreateSelectionGizmos(void* Owner)
+{
+	// always destroy the previous active auto gizmo
+	DestroyAllSelectionGizmos();
+
+	if (bShowSelectionGizmos)
+	{
+		FToolBuilderState CurrentSceneState;
+		QueriesAPI->GetCurrentSelectionState(CurrentSceneState);
+
+		if (UTypedElementSelectionSet* SelectionSet = CurrentSceneState.TypedElementSelectionSet.Get())
+		{
+			if (SelectionSet->HasSelectedElements())
+			{
+				TArray<UInteractiveGizmo*> NewGizmos;
+
+				TArray<TObjectPtr<UEditorInteractiveGizmoSelectionBuilder>> FoundBuilders = GetQualifiedGizmoSelectionBuilders(CurrentSceneState);
+
+				for (TObjectPtr<UEditorInteractiveGizmoSelectionBuilder> FoundBuilder : FoundBuilders)
+				{
+					UInteractiveGizmo* NewGizmo = FoundBuilder->BuildGizmo(CurrentSceneState);
+					if (NewGizmo == nullptr)
+					{
+						DisplayMessage(LOCTEXT("CreateGizmoReturnNullMessage", "UEditorInteractiveGizmoManager::CreateGizmo: BuildGizmo() returned null"), EToolMessageLevel::Internal);
+						return NewGizmos;
+					}
+
+					// register new active input behaviors
+					InputRouter->RegisterSource(NewGizmo);
+
+					NewGizmos.Add(NewGizmo);
+				}
+
+				PostInvalidation();
+
+				for (UInteractiveGizmo* NewGizmo : NewGizmos)
+				{
+					FActiveSelectionGizmo ActiveGizmo = { NewGizmo, Owner };
+					ActiveSelectionGizmos.Add(ActiveGizmo);
+				}
+
+				return NewGizmos;
+			}
+		}
+	}
+
+	return TArray<UInteractiveGizmo*>();
+}
+
+
+bool UEditorInteractiveGizmoManager::DestroySelectionGizmo(UInteractiveGizmo* Gizmo)
+{
+	auto Pred = [Gizmo](const FActiveSelectionGizmo& ActiveSelectionGizmo) {return ActiveSelectionGizmo.Gizmo == Gizmo; };
+	if (!ensure(ActiveSelectionGizmos.FindByPredicate(Pred)))
+	{
+		return false;
+	}
+
+	InputRouter->ForceTerminateSource(Gizmo);
+
+	Gizmo->Shutdown();
+
+	InputRouter->DeregisterSource(Gizmo);
+
+	ActiveSelectionGizmos.RemoveAll(Pred);
+
+	PostInvalidation();
+
+	return true;
+}
+
+void UEditorInteractiveGizmoManager::DestroyAllSelectionGizmos()
+{
+	for (int i = 0; i < ActiveSelectionGizmos.Num(); i++)
+	{
+		UInteractiveGizmo* Gizmo = ActiveSelectionGizmos[i].Gizmo;
+		if (ensure(Gizmo))
+		{
+			DestroySelectionGizmo(Gizmo);
+		}
+	}
+
+	ActiveSelectionGizmos.Reset();
+
+	PostInvalidation();
+}
+
+
+void UEditorInteractiveGizmoManager::OnEditorSelectionChanged()
+{
+	CreateSelectionGizmos();
+}
+
+void UEditorInteractiveGizmoManager::OnEditorSelectNone()
+{
+	DestroyAllSelectionGizmos();
+}
+
+// @todo move this to a gizmo context object
+bool UEditorInteractiveGizmoManager::GetShowSelectionGizmos()
+{
+	return bShowSelectionGizmos;
+}
+
+bool UEditorInteractiveGizmoManager::GetShowSelectionGizmosForView(IToolsContextRenderAPI* RenderAPI)
+{
+	const bool bEngineShowFlagsModeWidget = (RenderAPI && RenderAPI->GetSceneView() && 
+											 RenderAPI->GetSceneView()->Family &&
+											 RenderAPI->GetSceneView()->Family->EngineShowFlags.ModeWidgets);
+	return bShowSelectionGizmos && bEngineShowFlagsModeWidget;
+}
+
+void UEditorInteractiveGizmoManager::UpdateActiveSelectionGizmos()
+{
+#if 0
+	const bool bEditorModeToolsSupportsWidgetDrawing = EditorModeManager ? EditorModeManager->GetShowWidget() : true;
+	const bool bEnableSelectionGizmos = (CVarUseLegacyWidget.GetValueOnGameThread() == 0);
+	const bool bNewShowSelectionGizmos = bEditorModeToolsSupportsWidgetDrawing && bEnableSelectionGizmos;
+#else
+	const bool bNewShowSelectionGizmos = EditorModeManager ? EditorModeManager->GetShowWidget() : true;
+#endif
+
+	if (bShowSelectionGizmos != bNewShowSelectionGizmos)
+	{
+		bShowSelectionGizmos = bNewShowSelectionGizmos;
+		if (bShowSelectionGizmos)
+		{
+			CreateSelectionGizmos();
+		}
+		else
+		{
+			DestroyAllSelectionGizmos();
+		}
+	}
+}
+
+void UEditorInteractiveGizmoManager::Tick(float DeltaTime)
+{
+	Super::Tick(DeltaTime);
+
+	UpdateActiveSelectionGizmos();
+
+	for (FActiveSelectionGizmo& ActiveSelectionGizmo : ActiveSelectionGizmos)
+	{
+		ActiveSelectionGizmo.Gizmo->Tick(DeltaTime);
+	}
+}
+
+
+void UEditorInteractiveGizmoManager::Render(IToolsContextRenderAPI* RenderAPI)
+{
+	Super::Render(RenderAPI);
+
+	if (GetShowSelectionGizmosForView(RenderAPI))
+	{
+		for (FActiveSelectionGizmo& ActiveSelectionGizmo : ActiveSelectionGizmos)
+		{
+			ActiveSelectionGizmo.Gizmo->Render(RenderAPI);
+		}
+	}
+}
+
+void UEditorInteractiveGizmoManager::DrawHUD(FCanvas* Canvas, IToolsContextRenderAPI* RenderAPI)
+{
+	Super::DrawHUD(Canvas, RenderAPI);
+
+	if (GetShowSelectionGizmosForView(RenderAPI))
+	{
+		for (FActiveSelectionGizmo& ActiveSelectionGizmo : ActiveSelectionGizmos)
+		{
+			ActiveSelectionGizmo.Gizmo->DrawHUD(Canvas, RenderAPI);
+		}
+	}
+}
+
+#undef LOCTEXT_NAMESPACE
diff --git a/Engine/Source/Editor/Experimental/EditorInteractiveToolsFramework/Private/EditorInteractiveGizmoSubsystem.cpp b/Engine/Source/Editor/Experimental/EditorInteractiveToolsFramework/Private/EditorInteractiveGizmoSubsystem.cpp
new file mode 100644
--- /dev/null
+++ b/Engine/Source/Editor/Experimental/EditorInteractiveToolsFramework/Private/EditorInteractiveGizmoSubsystem.cpp
@@ -0,0 +1,120 @@
+// Copyright Epic Games, Inc. All Rights Reserved.
+
+#pragma once
+
+#include "EditorInteractiveGizmoSubsystem.h"
+#include "EditorGizmos/EditorTransformGizmo.h"
+
+
+#define LOCTEXT_NAMESPACE "UEditorInteractiveGizmoSubsystem"
+
+DEFINE_LOG_CATEGORY_STATIC(LogEditorInteractiveGizmoSubsystem, Log, All);
+
+
+UEditorInteractiveGizmoSubsystem::UEditorInteractiveGizmoSubsystem()
+	: UEditorSubsystem()
+{
+
+}
+
+void UEditorInteractiveGizmoSubsystem::Initialize(FSubsystemCollectionBase& Collection)
+{
+	if (GEngine && GEngine->IsInitialized())
+	{
+		RegisterBuiltinGizmoSelectionTypes();
+	}
+	else
+	{
+		FCoreDelegates::OnPostEngineInit.AddUObject(this, &UEditorInteractiveGizmoSubsystem::RegisterBuiltinGizmoSelectionTypes);
+	}
+}
+
+void UEditorInteractiveGizmoSubsystem::Deinitialize()
+{
+	DeregisterBuiltinGizmoSelectionTypes();
+}
+
+void UEditorInteractiveGizmoSubsystem::RegisterBuiltinGizmoSelectionTypes()
+{
+#if 0
+	// Register built-in gizmo types here
+	TObjectPtr<UEditorTransformGizmoBuilder> EditorTransformBuilder = NewObject<UEditorTransformGizmoBuilder>();
+	RegisterGizmoSelectionType(EditorTransformBuilder);
+#endif
+
+	RegisterEditorGizmoSelectionTypesDelegate.Broadcast();
+}
+
+void UEditorInteractiveGizmoSubsystem::DeregisterBuiltinGizmoSelectionTypes()
+{
+	DeregisterEditorGizmoSelectionTypesDelegate.Broadcast();
+	ClearGizmoSelectionTypeRegistry();
+}
+
+void UEditorInteractiveGizmoSubsystem::RegisterGizmoSelectionType(const TObjectPtr<UEditorInteractiveGizmoSelectionBuilder> InGizmoSelectionBuilder)
+{
+	if (ensure(InGizmoSelectionBuilder != nullptr) == false)
+	{
+		return;
+	}
+
+	if (GizmoSelectionBuilders.Contains(InGizmoSelectionBuilder))
+	{
+		UE_LOG(LogEditorInteractiveGizmoSubsystem, Warning,
+			TEXT("UInteractiveGizmoSubsystem::RegisterGizmoSelectionType: type has already been registered %s"), *InGizmoSelectionBuilder->GetName());
+		return;
+	}
+
+	GizmoSelectionBuilders.Add(InGizmoSelectionBuilder);
+	GizmoSelectionBuilders.StableSort(
+		[](UEditorInteractiveGizmoSelectionBuilder& A, UEditorInteractiveGizmoSelectionBuilder& B) {
+			return (A).GetPriority() > (B).GetPriority();
+		});
+}
+
+TArray<TObjectPtr<UEditorInteractiveGizmoSelectionBuilder>> UEditorInteractiveGizmoSubsystem::GetQualifiedGizmoSelectionBuilders(const FToolBuilderState& InToolBuilderState)
+{
+	TArray<TObjectPtr<UEditorInteractiveGizmoSelectionBuilder>> FoundBuilders;
+	FEditorGizmoTypePriority FoundPriority = 0;
+
+	for (TObjectPtr<UEditorInteractiveGizmoSelectionBuilder> Builder : GizmoSelectionBuilders)
+	{
+		if (Builder->GetPriority() < FoundPriority)
+		{
+			break;
+		}
+
+		if (Builder->SatisfiesCondition(InToolBuilderState))
+		{
+			FoundBuilders.Add(Builder);
+			FoundPriority = Builder->GetPriority();
+		}
+	}
+
+	return FoundBuilders;
+}
+
+bool UEditorInteractiveGizmoSubsystem::DeregisterGizmoSelectionType(const TObjectPtr<UEditorInteractiveGizmoSelectionBuilder> InGizmoSelectionBuilder)
+{
+	if (ensure(InGizmoSelectionBuilder != nullptr) == false)
+	{
+		return false;
+	}
+
+	if (GizmoSelectionBuilders.Contains(InGizmoSelectionBuilder) == false)
+	{
+		UE_LOG(LogEditorInteractiveGizmoSubsystem, Warning,
+			TEXT("UInteractiveGizmoSubsystem::DeregisterGizmoSelectionType: type has already been registered %s"), *InGizmoSelectionBuilder->GetName());
+		return false;
+	}
+	GizmoSelectionBuilders.Remove(InGizmoSelectionBuilder);
+	return true;
+}
+
+void UEditorInteractiveGizmoSubsystem::ClearGizmoSelectionTypeRegistry()
+{
+	GizmoSelectionBuilders.Reset();
+}
+
+#undef LOCTEXT_NAMESPACE
+
diff --git a/Engine/Source/Editor/Experimental/EditorInteractiveToolsFramework/Public/EditorInteractiveGizmoConditionalBuilder.h b/Engine/Source/Editor/Experimental/EditorInteractiveToolsFramework/Public/EditorInteractiveGizmoConditionalBuilder.h
new file mode 100644
--- /dev/null
+++ b/Engine/Source/Editor/Experimental/EditorInteractiveToolsFramework/Public/EditorInteractiveGizmoConditionalBuilder.h
@@ -0,0 +1,85 @@
+// Copyright Epic Games, Inc. All Rights Reserved.
+
+#pragma once
+
+#include "CoreMinimal.h"
+#include "InteractiveGizmoBuilder.h"
+#include "ToolContextInterfaces.h"
+#include "EditorInteractiveGizmoConditionalBuilder.generated.h"
+
+/**
+ * FEditorGizmoTypePriority is used to establish relative priority between conditional 
+ * gizmo builders. It is up to the gizmo manager to determine how the priority is used. 
+ * In the EditorInteractiveGizmoManager, if more than one gizmo builder returns true 
+ * from SatsifiesCondition(), the gizmo builder with highest priority will be used. If 
+ * there are multiple builders the highest priority, multiple gizmos will be built.
+ */
+struct EDITORINTERACTIVETOOLSFRAMEWORK_API FEditorGizmoTypePriority
+{
+	static constexpr int DEFAULT_GIZMO_TYPE_PRIORITY = 50;
+
+	/** Constant priority value */
+	int Priority;
+
+	FEditorGizmoTypePriority(int InPriority = DEFAULT_GIZMO_TYPE_PRIORITY)
+	{
+		Priority = InPriority;
+	}
+
+	/** @return a priority lower than this priority */
+	FEditorGizmoTypePriority MakeLower(int DeltaAmount = 1) const
+	{
+		return FEditorGizmoTypePriority(Priority + DeltaAmount);
+	}
+
+	/** @return a priority higher than this priority */
+	FEditorGizmoTypePriority MakeHigher(int DeltaAmount = 1) const
+	{
+		return FEditorGizmoTypePriority(Priority - DeltaAmount);
+	}
+
+	friend bool operator<(const FEditorGizmoTypePriority& l, const FEditorGizmoTypePriority& r)
+	{
+		return l.Priority < r.Priority;
+	}
+	friend bool operator==(const FEditorGizmoTypePriority& l, const FEditorGizmoTypePriority& r)
+	{
+		return l.Priority == r.Priority;
+	}
+	friend bool operator>(const FEditorGizmoTypePriority& l, const FEditorGizmoTypePriority& r)
+	{
+		return l.Priority > r.Priority;
+	}
+};
+
+/** UInteractiveGizmoType provides actions and other information about gizmo types */
+UCLASS(Transient, Abstract)
+class EDITORINTERACTIVETOOLSFRAMEWORK_API UEditorInteractiveGizmoConditionalBuilder : public UInteractiveGizmoBuilder
+{
+	GENERATED_BODY()
+
+public:
+
+	/** Returns the priority for this gizmo type.  */
+	virtual FEditorGizmoTypePriority GetPriority() const 
+	{ 
+		return FEditorGizmoTypePriority::DEFAULT_GIZMO_TYPE_PRIORITY;
+	}
+
+	/** Update the priority for this gizmo type.  */
+	virtual void UpdatePriority(const FEditorGizmoTypePriority& InPriority)
+	{
+		Priority = InPriority;
+	}
+
+	/** Returns true if this gizmo is valid for creation based on the current state. */
+	virtual bool SatisfiesCondition(const FToolBuilderState& SceneState) const
+	{
+		unimplemented();
+		return false;
+	}
+
+protected:
+
+	FEditorGizmoTypePriority Priority = FEditorGizmoTypePriority::DEFAULT_GIZMO_TYPE_PRIORITY;
+};
diff --git a/Engine/Source/Editor/Experimental/EditorInteractiveToolsFramework/Public/EditorInteractiveGizmoManager.h b/Engine/Source/Editor/Experimental/EditorInteractiveToolsFramework/Public/EditorInteractiveGizmoManager.h
new file mode 100644
--- /dev/null
+++ b/Engine/Source/Editor/Experimental/EditorInteractiveToolsFramework/Public/EditorInteractiveGizmoManager.h
@@ -0,0 +1,161 @@
+// Copyright Epic Games, Inc. All Rights Reserved.
+
+#pragma once
+
+#include "CoreMinimal.h"
+#include "InteractiveGizmoManager.h"
+#include "EditorInteractiveGizmoManager.generated.h"
+
+class UEditorInteractiveGizmoSelectionBuilder;
+class UEdModeInteractiveToolsContext;
+class UinteractiveGizmo;
+class IToolsContextRenderAPI;
+
+USTRUCT()
+struct FActiveSelectionGizmo
+{
+	GENERATED_BODY()
+
+	UPROPERTY()
+	TObjectPtr<UInteractiveGizmo> Gizmo = nullptr;
+	void* Owner = nullptr;
+};
+
+/**
+ * UEditorInteractiveGizmoManager allows users of the Tools framework to register and create selection-based Gizmo instances.
+ * For each selection-based Gizmo, a builder derived from UInteractiveGizmoSelectionBuilder is registered with the GizmoManager.
+ * When the section changes, the highest priority builders for which SatisfiesCondition() return true, will be used to
+ * build gizmos.
+ */
+UCLASS(Transient)
+class EDITORINTERACTIVETOOLSFRAMEWORK_API UEditorInteractiveGizmoManager : public UInteractiveGizmoManager
+{
+	GENERATED_BODY()
+
+protected:
+	friend class UEdModeInteractiveToolsContext;		// to call Initialize/Shutdown
+
+	UEditorInteractiveGizmoManager();
+
+	/** Initialize the GizmoManager with the necessary Context-level state. UEdModeInteractiveToolsContext calls this, you should not. */
+	virtual void InitializeWithEditorModeManager(IToolsContextQueriesAPI* QueriesAPI, IToolsContextTransactionsAPI* TransactionsAPI, UInputRouter* InputRouter, FEditorModeTools* InEditorModeManager);
+
+	// UInteractiveGizmoManager interface
+	virtual void Shutdown() override;
+
+public:
+
+	// UInteractiveGizmoManager interface
+	virtual void Tick(float DeltaTime) override;
+	virtual void Render(IToolsContextRenderAPI* RenderAPI) override;
+	virtual void DrawHUD(FCanvas* Canvas, IToolsContextRenderAPI* RenderAPI) override;
+
+	/**
+	 * Register a new selection gizmo type
+	 * @param InGizmoSelectionBuilder new auto gizmo builder
+	 */
+	void RegisterGizmoSelectionType(const TObjectPtr<UEditorInteractiveGizmoSelectionBuilder> InGizmoSelectionBuilder);
+
+	/**
+	* Remove a gizmo auto type from the set of known gizmo auto types
+	* @param InGizmoSelectionBuilder same object pointer that was passed to RegisterGizmoSelectionType()
+	* @return true if gizmo type was found and deregistered
+	*/
+	bool DeregisterGizmoSelectionType(const TObjectPtr<UEditorInteractiveGizmoSelectionBuilder> InGizmoSelectionBuilder);
+
+	/**
+	 * Get all qualified gizmo auto builders based on the current state. Qualification is determined by the gizmo builder
+	 * returning true from SatisfiesCondition() and relative priority. All qualified builders at the highest found priority
+	 * will be returned.
+	 * @return array of qualified Gizmo auto builders based on current state
+	 */
+	virtual TArray<TObjectPtr<UEditorInteractiveGizmoSelectionBuilder>> GetQualifiedGizmoSelectionBuilders(const FToolBuilderState& InToolBuilderState);
+
+	/** 
+	 * Set how auto gizmo resolution should occur when CreateSelectionGizmo is invoked. If bSearchLocalOnly is true, only the current
+	 * @param bLocalOnly - if true, only the current gizmo manager registry will be searched for candidate gizmos. If false,
+	 *   both the gizmo manager registry and any higher gizmo manager or gizmo subsystem (in the case of selection builders) will be searched
+	 */
+	virtual void SetGizmoSelectionBuilderResolution(bool bLocalOnly)
+	{
+		bSearchLocalBuildersOnly = bLocalOnly;
+	}
+
+	/**
+	 * Returns the current auto gizmo resolution setting 
+	 */
+	virtual bool GetGizmoSelectionBuilderResolution() const
+	{
+		return bSearchLocalBuildersOnly;
+	}
+
+	/**
+	 * Try to automatically activate a new Gizmo instance based on the current state
+	 * @param Owner void pointer to whatever "owns" this Gizmo. Allows Gizmo to later be deleted using DestroyAllGizmosByOwner()
+	 * @return array of new Gizmo instances that have been created and initialized
+	 */
+	virtual TArray<UInteractiveGizmo*> CreateSelectionGizmos(void* Owner = nullptr);
+
+	/**
+	 * Handle Editor selection changes
+	 * @param Tools - Mode Manager which invoked this selection changed call
+	 * @param NewSelection - Object which is undergoing selection change
+	 */
+	virtual void OnEditorSelectionChanged();
+
+	/**
+	 * Handle case when selection has been cleared.
+	 */
+	virtual void OnEditorSelectNone();
+
+	/**
+	 * Shutdown and remove a selection-based Gizmo
+	 * @param Gizmo the Gizmo to shutdown and remove
+	 * @return true if the Gizmo was found and removed
+	 */
+	virtual bool DestroySelectionGizmo(UInteractiveGizmo* Gizmo);
+
+		/**
+	 * Shutdown and remove all active auto gizmos
+	 */
+	virtual void DestroyAllSelectionGizmos();
+
+protected:
+
+	/**
+	 * Returns true if selection gizmos should be visible. 
+	 * @todo move this to a gizmo context object
+	 */
+	virtual bool GetShowSelectionGizmos();
+
+	/**
+	 * Returns true if gizmos should be visible based on the current view's engine show flag.
+	 * @todo move this to a gizmo context object
+	 */
+	virtual bool GetShowSelectionGizmosForView(IToolsContextRenderAPI* RenderAPI);
+
+	/**
+	 * Updates active selection gizmos when show selection state changes
+	 */
+	void UpdateActiveSelectionGizmos();
+
+protected:
+
+	/** set of Currently-active Gizmos */
+	UPROPERTY()
+	TArray<FActiveSelectionGizmo> ActiveSelectionGizmos;
+
+	/** Current set of GizmoSelectionBuilders */
+	UPROPERTY()
+	TArray<TObjectPtr<UEditorInteractiveGizmoSelectionBuilder>> GizmoSelectionBuilders;
+
+	/** If false, only search gizmo builders in current gizmo manager. If true, also search gizmo subsystem */
+	bool bSearchLocalBuildersOnly = false;
+
+private:
+	/** @todo: remove when GetShowSelectionGizmos() is moved to gizmo context object */
+	FEditorModeTools* EditorModeManager = nullptr;
+
+	/** Whether selection gizmos are enabled. UpdateActiveSelectionGizmos() determines this value each tick and updates if it has changed. */
+	bool bShowSelectionGizmos = false;
+};
diff --git a/Engine/Source/Editor/Experimental/EditorInteractiveToolsFramework/Public/EditorInteractiveGizmoSelectionBuilder.h b/Engine/Source/Editor/Experimental/EditorInteractiveToolsFramework/Public/EditorInteractiveGizmoSelectionBuilder.h
new file mode 100644
--- /dev/null
+++ b/Engine/Source/Editor/Experimental/EditorInteractiveToolsFramework/Public/EditorInteractiveGizmoSelectionBuilder.h
@@ -0,0 +1,28 @@
+// Copyright Epic Games, Inc. All Rights Reserved.
+
+#pragma once
+
+#include "CoreMinimal.h"
+#include "EditorInteractiveGizmoConditionalBuilder.h"
+#include "ToolContextInterfaces.h"
+#include "EditorInteractiveGizmoSelectionBuilder.generated.h"
+
+
+/** 
+ * UEditorInteractiveGizmoSelectionBuilder provides a method for checking that the current selection and widget mode satisfy 
+ * the conditions of this builder. Builders derived from this class should be registered in the gizmo subsystem, for gizmos
+ * available globally in the Editor, or in the gizmo manager for gizmos only relevant to a particle ed mode or asset editor.
+ */
+UCLASS(Transient, Abstract)
+class EDITORINTERACTIVETOOLSFRAMEWORK_API UEditorInteractiveGizmoSelectionBuilder : public UEditorInteractiveGizmoConditionalBuilder
+{
+	GENERATED_BODY()
+
+public:
+
+	/** Returns true if this gizmo is valid for creation based on the current state. */
+	virtual bool SatisfiesCondition(const FToolBuilderState& SceneState) const override
+	{
+		return false;
+	}
+};
diff --git a/Engine/Source/Editor/Experimental/EditorInteractiveToolsFramework/Public/EditorInteractiveGizmoSubsystem.h b/Engine/Source/Editor/Experimental/EditorInteractiveToolsFramework/Public/EditorInteractiveGizmoSubsystem.h
new file mode 100644
--- /dev/null
+++ b/Engine/Source/Editor/Experimental/EditorInteractiveToolsFramework/Public/EditorInteractiveGizmoSubsystem.h
@@ -0,0 +1,99 @@
+// Copyright Epic Games, Inc. All Rights Reserved.
+
+#pragma once
+
+#include "CoreMinimal.h"
+#include "Delegates/Delegate.h"
+#include "EditorSubsystem.h"
+#include "EditorInteractiveGizmoSelectionBuilder.h"
+
+#include "EditorInteractiveGizmoSubsystem.generated.h"
+
+class FSubsystemCollectionBase;
+
+/**
+ * The InteractiveGizmoSubsystem provides methods for registering and unregistering 
+ * selection-based gizmos builders. Editor gizmo managers which are not marked local-only, 
+ * will query this subsystem for qualified builders based on the current selection.
+ *
+ * This subsystem should also be used to register gizmo selection builders from plugins
+ * by binding to the delegates returned from OnEditorGizmoSubsystemRegisterGizmoSelectionTypes() and 
+ * OnEditorGizmoSubsystemDeregisterGizmoSelectionTypes().
+ */
+UCLASS()
+class EDITORINTERACTIVETOOLSFRAMEWORK_API UEditorInteractiveGizmoSubsystem : public UEditorSubsystem
+{
+	GENERATED_BODY()
+
+public:
+
+	UEditorInteractiveGizmoSubsystem();
+
+	//~ Begin USubsystem interface
+	virtual void Initialize(FSubsystemCollectionBase& Collection) override;
+
+	virtual void Deinitialize() override;
+	//~ End USubsystem interface
+
+	/**
+	 * Event which is broadcast just after default types are registered in the gizmo subsystem
+	 */
+	DECLARE_EVENT(UEditorInteractiveSelectionGizmoSubsystem, FOnEditorGizmoSubsystemRegisterGizmoSelectionTypes);
+	FOnEditorGizmoSubsystemRegisterGizmoSelectionTypes& OnEditorGizmoSubsystemRegisterGizmoSelectionTypes() { return RegisterEditorGizmoSelectionTypesDelegate; }
+
+	/**
+	 * Event which is broadcast just before default types are deregistered in the gizmo subsystem
+	 */
+	DECLARE_EVENT(UEditorInteractiveSelectionGizmoSubsystem, FOnEditorGizmoSubsystemDeregisterGizmoSelectionTypes);
+	FOnEditorGizmoSubsystemDeregisterGizmoSelectionTypes& OnEditorGizmoSubsystemDeregisterGizmoSelectionTypes() { return DeregisterEditorGizmoSelectionTypesDelegate; }
+
+	/**
+	 * Registers all built-in gizmo selection types and broadcast registration event.
+	 */
+	void RegisterBuiltinGizmoSelectionTypes();
+
+	/**
+	 * Removes all built-in gizmo selection types and broadcast deregistration event.
+	 */
+	void DeregisterBuiltinGizmoSelectionTypes();
+
+	/**
+	 * Register a new selection gizmo type
+	 * @param InGizmoSelectionBuilder new selection gizmo builder
+	 */
+	void RegisterGizmoSelectionType(const TObjectPtr<UEditorInteractiveGizmoSelectionBuilder> InGizmoSelectionBuilder);
+
+	/**
+	* Remove a gizmo selection type from the set of known gizmo selection types
+	* @param InGizmoSelectionBuilder same object pointer that was passed to RegisterGizmoSelectionType()
+	* @return true if gizmo type was found and deregistered
+	*/
+	bool DeregisterGizmoSelectionType(const TObjectPtr<UEditorInteractiveGizmoSelectionBuilder> InGizmoSelectionBuilder);
+
+	/**
+	 * Clear all registered gizmo types
+	 */
+	void ClearGizmoSelectionTypeRegistry();
+
+	/**
+	 * Get all qualified gizmo selection builders based on the current state. Qualification is determined by the gizmo builder
+	 * returning true from SatisfiesCondition() and relative priority. All qualified builders at the highest found priority
+	 * will be returned.
+	 * @return array of qualified Gizmo selection builders based on current state
+	 */
+	virtual TArray<TObjectPtr<UEditorInteractiveGizmoSelectionBuilder>> GetQualifiedGizmoSelectionBuilders(const FToolBuilderState& InToolBuilderState);
+
+
+private:
+
+	/** Current set of GizmoSelectionBuilders */
+	UPROPERTY()
+	TArray<TObjectPtr<UEditorInteractiveGizmoSelectionBuilder> > GizmoSelectionBuilders;
+
+	/** Call to register gizmo types */
+	FOnEditorGizmoSubsystemRegisterGizmoSelectionTypes RegisterEditorGizmoSelectionTypesDelegate;
+
+	/** Call to deregister gizmo types */
+	FOnEditorGizmoSubsystemDeregisterGizmoSelectionTypes DeregisterEditorGizmoSelectionTypesDelegate;
+
+};
\ No newline at end of file
