Fixes for no unity errors
#rnx

diff --git a/Engine/Plugins/Runtime/Nvidia/Reflex/Source/Private/ReflexMaxTickRateHandler.cpp b/Engine/Plugins/Runtime/Nvidia/Reflex/Source/Private/ReflexMaxTickRateHandler.cpp
--- a/Engine/Plugins/Runtime/Nvidia/Reflex/Source/Private/ReflexMaxTickRateHandler.cpp
+++ b/Engine/Plugins/Runtime/Nvidia/Reflex/Source/Private/ReflexMaxTickRateHandler.cpp
@@ -6,6 +6,8 @@
 
 #include "RHI.h"
 #include "Engine/Console.h"
+#include "Engine/LocalPlayer.h"
+#include "Engine/World.h"
 #include "GameFramework/PlayerController.h"
 #include "Kismet/GameplayStatics.h"
 
diff --git a/Engine/Plugins/Runtime/Nvidia/Reflex/Source/Public/ReflexMaxTickRateHandler.h b/Engine/Plugins/Runtime/Nvidia/Reflex/Source/Public/ReflexMaxTickRateHandler.h
--- a/Engine/Plugins/Runtime/Nvidia/Reflex/Source/Public/ReflexMaxTickRateHandler.h
+++ b/Engine/Plugins/Runtime/Nvidia/Reflex/Source/Public/ReflexMaxTickRateHandler.h
@@ -3,6 +3,7 @@
 #pragma once
 
 #include "CoreMinimal.h"
+#include "Misc/CoreMisc.h"
 #include "Performance/MaxTickRateHandlerModule.h"
 
 class FReflexMaxTickRateHandler : public IMaxTickRateHandlerModule, public FSelfRegisteringExec
diff --git a/Engine/Source/Editor/Experimental/EditorInteractiveToolsFramework/Public/EditorInteractiveGizmoManager.h b/Engine/Source/Editor/Experimental/EditorInteractiveToolsFramework/Public/EditorInteractiveGizmoManager.h
--- a/Engine/Source/Editor/Experimental/EditorInteractiveToolsFramework/Public/EditorInteractiveGizmoManager.h
+++ b/Engine/Source/Editor/Experimental/EditorInteractiveToolsFramework/Public/EditorInteractiveGizmoManager.h
@@ -6,6 +6,7 @@
 #include "InteractiveGizmoManager.h"
 #include "EditorInteractiveGizmoManager.generated.h"
 
+class FEditorModeTools;
 class UEditorInteractiveGizmoSelectionBuilder;
 class UEdModeInteractiveToolsContext;
 class UinteractiveGizmo;
diff --git a/Engine/Source/Editor/UnrealEd/Private/Elements/SMInstance/SMInstanceElementDetailsProxyObject.cpp b/Engine/Source/Editor/UnrealEd/Private/Elements/SMInstance/SMInstanceElementDetailsProxyObject.cpp
--- a/Engine/Source/Editor/UnrealEd/Private/Elements/SMInstance/SMInstanceElementDetailsProxyObject.cpp
+++ b/Engine/Source/Editor/UnrealEd/Private/Elements/SMInstance/SMInstanceElementDetailsProxyObject.cpp
@@ -7,6 +7,7 @@
 
 #include "UnrealEdGlobals.h"
 #include "Editor/UnrealEdEngine.h"
+#include "Engine/StaticMesh.h"
 
 #define LOCTEXT_NAMESPACE "SMInstanceElementDetails"
 
diff --git a/Engine/Source/Editor/UnrealEd/Private/PackageSourceControlHelper.cpp b/Engine/Source/Editor/UnrealEd/Private/PackageSourceControlHelper.cpp
--- a/Engine/Source/Editor/UnrealEd/Private/PackageSourceControlHelper.cpp
+++ b/Engine/Source/Editor/UnrealEd/Private/PackageSourceControlHelper.cpp
@@ -6,6 +6,7 @@
 
 #include "PackageSourceControlHelper.h"
 #include "Logging/LogMacros.h"
+#include "UObject/Linker.h"
 #include "UObject/Package.h"
 #include "HAL/FileManager.h"
 #include "HAL/PlatformFileManager.h"
