Add selection-based builder for Editor transform gizmo which will work with the new Editor gizmo subsystem and manager.

diff --git a/Engine/Source/Editor/Experimental/EditorInteractiveToolsFramework/Private/EditorGizmos/EditorTransformGizmo.cpp b/Engine/Source/Editor/Experimental/EditorInteractiveToolsFramework/Private/EditorGizmos/EditorTransformGizmo.cpp
--- a/Engine/Source/Editor/Experimental/EditorInteractiveToolsFramework/Private/EditorGizmos/EditorTransformGizmo.cpp
+++ b/Engine/Source/Editor/Experimental/EditorInteractiveToolsFramework/Private/EditorGizmos/EditorTransformGizmo.cpp
@@ -11,7 +11,10 @@
 #include "EditorGizmos/EditorAxisSources.h"
 #include "EditorGizmos/EditorParameterToTransformAdapters.h"
 #include "Engine/CollisionProfile.h"
+#include "Engine/Selection.h"
 #include "Engine/World.h"
+#include "EditorModeTools.h"
+#include "UnrealEdGlobals.h"
 #include "UnrealEngine.h"
 
 #define LOCTEXT_NAMESPACE "UEditorTransformGizmo"
@@ -44,9 +47,61 @@ constexpr FColor UEditorTransformGizmo::CurrentColor;
 
 UInteractiveGizmo* UEditorTransformGizmoBuilder::BuildGizmo(const FToolBuilderState& SceneState) const
 {
-	UEditorTransformGizmo* NewGizmo = NewObject<UEditorTransformGizmo>(SceneState.GizmoManager);
-	NewGizmo->SetWorld(SceneState.World);
-	return NewGizmo;
+	// @todo - remove global call
+	FEditorModeTools& ModeTools = GLevelEditorModeTools();
+
+	ETransformGizmoSubElements Elements  = ETransformGizmoSubElements::None;
+	bool bUseContextCoordinateSystem = true;
+	UE::Widget::EWidgetMode WidgetMode = ModeTools.GetWidgetMode();
+	switch ( WidgetMode )
+	{
+	case UE::Widget::EWidgetMode::WM_Translate:
+		Elements = ETransformGizmoSubElements::TranslateAllAxes | ETransformGizmoSubElements::TranslateAllPlanes;
+		break;
+	case UE::Widget::EWidgetMode::WM_Rotate:
+		Elements = ETransformGizmoSubElements::RotateAllAxes;
+		break;
+	case UE::Widget::EWidgetMode::WM_Scale:
+		Elements = ETransformGizmoSubElements::ScaleAllAxes | ETransformGizmoSubElements::ScaleAllPlanes;
+		bUseContextCoordinateSystem = false;
+		break;
+	case UE::Widget::EWidgetMode::WM_2D:
+		Elements = ETransformGizmoSubElements::RotateAxisY | ETransformGizmoSubElements::TranslatePlaneXZ;
+		break;
+	default:
+		Elements = ETransformGizmoSubElements::FullTranslateRotateScale;
+		break;
+	}
+
+	UEditorTransformGizmo* TransformGizmo = NewObject<UEditorTransformGizmo>(SceneState.GizmoManager);
+	TransformGizmo->Setup();
+
+	TransformGizmo->SetWorld(SceneState.World);
+	TransformGizmo->bUseContextCoordinateSystem = bUseContextCoordinateSystem;
+
+	// @todo - update to work with typed elements
+	TArray<AActor*> SelectedActors;
+	ModeTools.GetSelectedActors()->GetSelectedObjects<AActor>(SelectedActors);
+
+	UTransformProxy* TransformProxy = NewObject<UTransformProxy>();
+	for (auto Actor : SelectedActors)
+	{
+		USceneComponent* SceneComponent = Actor->GetRootComponent();
+		TransformProxy->AddComponent(SceneComponent);
+	}
+	TransformGizmo->SetActiveTarget( TransformProxy );
+	TransformGizmo->SetVisibility(SelectedActors.Num() > 0);
+
+	return TransformGizmo;
+}
+
+bool UEditorTransformGizmoBuilder::SatisfiesCondition(const FToolBuilderState& SceneState) const 
+{
+	if (UTypedElementSelectionSet* SelectionSet = SceneState.TypedElementSelectionSet.Get())
+	{
+		return (SelectionSet->HasSelectedElements());
+	}
+	return true;
 }
 
 void UEditorTransformGizmo::SetWorld(UWorld* WorldIn)
@@ -426,7 +481,7 @@ UInteractiveGizmo* UEditorTransformGizmo::AddAxisTranslationGizmo(
 
 	// create axis-position gizmo, axis-position parameter will drive translation
 	UAxisPositionGizmo* TranslateGizmo = Cast<UAxisPositionGizmo>(GetGizmoManager()->CreateGizmo(
-		UInteractiveGizmoManager::DefaultAxisPositionBuilderIdentifier));
+		UInteractiveGizmoManager::DefaultAxisPositionBuilderIdentifier, FString(), this));
 	check(TranslateGizmo);
 
 	InArrowObject->CylinderObject->Length = AXIS_LENGTH;
@@ -472,7 +527,7 @@ UInteractiveGizmo* UEditorTransformGizmo::AddPlaneTranslationGizmo(
 	/* @todo
 	// create axis-position gizmo, axis-position parameter will drive translation
 	UEditorPlanePositionGizmo* TranslateGizmo = Cast<UEditorPlanePositionGizmo>(GetGizmoManager()->CreateGizmo(
-		UEditorInteractiveGizmoManager::DefaultEditorPlanePositionBuilderIdentifier));
+		UEditorInteractiveGizmoManager::DefaultEditorPlanePositionBuilderIdentifier, FString(), this));
 	check(TranslateGizmo);
 
 	// axis source provides the translation axis
@@ -515,7 +570,7 @@ UInteractiveGizmo* UEditorTransformGizmo::AddAxisRotationGizmo(
 	/* @todo
 	// create axis-angle gizmo, angle will drive axis-rotation
 	UEditorAxisAngleGizmo* RotateGizmo = Cast<UEditorAxisAngleGizmo>(GetGizmoManager()->CreateGizmo(
-		UEditorInteractiveGizmoManager::DefaultEditorAxisAngleBuilderIdentifier));
+		UEditorInteractiveGizmoManager::DefaultEditorAxisAngleBuilderIdentifier, FString(), this));
 	check(RotateGizmo);
 
 	// axis source provides the translation axis
@@ -593,7 +648,7 @@ UInteractiveGizmo* UEditorTransformGizmo::AddAxisScaleGizmo(
 
 	// create axis-position gizmo, axis-position parameter will drive scale
 	UAxisPositionGizmo* ScaleGizmo = Cast<UAxisPositionGizmo>(GetGizmoManager()->CreateGizmo(
-		UInteractiveGizmoManager::DefaultAxisPositionBuilderIdentifier));
+		UInteractiveGizmoManager::DefaultAxisPositionBuilderIdentifier, FString(), this));
 	ScaleGizmo->bEnableSignedAxis = true;
 	check(ScaleGizmo);
 
@@ -626,7 +681,7 @@ UInteractiveGizmo* UEditorTransformGizmo::AddPlaneScaleGizmo(
 	/* @todo
 	// create axis-position gizmo, axis-position parameter will drive scale
 	UEditorPlanePositionGizmo* ScaleGizmo = Cast<UEditorPlanePositionGizmo>(GetGizmoManager()->CreateGizmo(
-		UEditorInteractiveGizmoManager::DefaultEditorPlanePositionBuilderIdentifier));
+		UEditorInteractiveGizmoManager::DefaultEditorPlanePositionBuilderIdentifier, FString(), this));
 	ScaleGizmo->bEnableSignedAxis = true;
 	check(ScaleGizmo);
 
@@ -664,7 +719,7 @@ UInteractiveGizmo* UEditorTransformGizmo::AddUniformScaleGizmo(
 	/* @todo
 	// create plane-position gizmo, plane-position parameter will drive scale
 	UEditorPlanePositionGizmo* ScaleGizmo = Cast<UEditorPlanePositionGizmo>(GetGizmoManager()->CreateGizmo(
-		UEditorInteractiveGizmoManager::DefaultEditorPlanePositionBuilderIdentifier));
+		UEditorInteractiveGizmoManager::DefaultEditorPlanePositionBuilderIdentifier, FString(), this));
 	check(ScaleGizmo);
 
 	// axis source provides the translation plane
diff --git a/Engine/Source/Editor/Experimental/EditorInteractiveToolsFramework/Public/EditorGizmos/EditorTransformGizmo.h b/Engine/Source/Editor/Experimental/EditorInteractiveToolsFramework/Public/EditorGizmos/EditorTransformGizmo.h
--- a/Engine/Source/Editor/Experimental/EditorInteractiveToolsFramework/Public/EditorGizmos/EditorTransformGizmo.h
+++ b/Engine/Source/Editor/Experimental/EditorInteractiveToolsFramework/Public/EditorGizmos/EditorTransformGizmo.h
@@ -8,8 +8,8 @@
 #include "BaseGizmos/GizmoGroupObject.h"
 #include "BaseGizmos/TransformProxy.h"
 #include "EditorGizmos/EditorAxisSources.h"
+#include "EditorInteractiveGizmoSelectionBuilder.h"
 #include "InteractiveGizmo.h"
-#include "InteractiveGizmoBuilder.h"
 #include "InteractiveToolObjects.h"
 #include "InteractiveToolChange.h"
 #include "Materials/MaterialInstanceDynamic.h"
@@ -26,13 +26,15 @@ class FEditorTransformGizmoTransformChange;
 class FRay;
 
 UCLASS()
-class EDITORINTERACTIVETOOLSFRAMEWORK_API UEditorTransformGizmoBuilder : public UInteractiveGizmoBuilder
+class EDITORINTERACTIVETOOLSFRAMEWORK_API UEditorTransformGizmoBuilder : public UEditorInteractiveGizmoSelectionBuilder
 {
 	GENERATED_BODY()
 
 public:
 
 	virtual UInteractiveGizmo* BuildGizmo(const FToolBuilderState& SceneState) const override;
+
+	virtual bool SatisfiesCondition(const FToolBuilderState& SceneState) const override;
 };
 
 
