BazelExecutor: Fix templated function to not cause a HeaderTool warning

diff --git a/Engine/Source/Editor/Experimental/BazelExecutor/Private/BazelCompletionQueueRunnable.h b/Engine/Source/Editor/Experimental/BazelExecutor/Private/BazelCompletionQueueRunnable.h
--- a/Engine/Source/Editor/Experimental/BazelExecutor/Private/BazelCompletionQueueRunnable.h
+++ b/Engine/Source/Editor/Experimental/BazelExecutor/Private/BazelCompletionQueueRunnable.h
@@ -77,45 +77,39 @@ public:
 	bool AddAsyncResponse(
 		TUniquePtr<grpc::ClientContext> ClientContext,
 		TUniquePtr<grpc::ClientAsyncResponseReader<MessageType>> Reader,
-		TFinishFunction OnFinish = nullptr);
-
-	grpc::CompletionQueue* GetCompletionQueue();
-};
-
-template<class MessageType>
-bool FBazelCompletionQueueRunnable::AddAsyncResponse(
-	TUniquePtr<grpc::ClientContext> ClientContext,
-	TUniquePtr<grpc::ClientAsyncResponseReader<MessageType>> Reader,
-	TFinishFunction OnFinish)
-{
-	static_assert(std::is_base_of<google::protobuf::Message, MessageType>::value, "MessageType is not derived from google::protobuf::Message");
-	if (!Running || !ClientContext.IsValid() || !Reader.IsValid())
+		TFinishFunction OnFinish = nullptr)
 	{
-		return false;
-	}
+		static_assert(std::is_base_of<google::protobuf::Message, MessageType>::value, "MessageType is not derived from google::protobuf::Message");
+		if (!Running || !ClientContext.IsValid() || !Reader.IsValid())
+		{
+			return false;
+		}
 
-	grpc::ClientAsyncResponseReader<MessageType>* AsyncReader = Reader.Get();
+		grpc::ClientAsyncResponseReader<MessageType>* AsyncReader = Reader.Get();
 
-	// Keep Reader in scope for the lifetime of this async call
-	TFinishFunction OnFinishResponse = [Reader = MoveTemp(Reader), OnFinish = MoveTemp(OnFinish)](void* Tag, bool Ok, const grpc::Status& Status, const google::protobuf::Message& Message)
-	{
-		if (OnFinish)
+		// Keep Reader in scope for the lifetime of this async call
+		TFinishFunction OnFinishResponse = [Reader = MoveTemp(Reader), OnFinish = MoveTemp(OnFinish)](void* Tag, bool Ok, const grpc::Status& Status, const google::protobuf::Message& Message)
 		{
-			OnFinish(Tag, Ok, Status, Message);
-		}
-	};
+			if (OnFinish)
+			{
+				OnFinish(Tag, Ok, Status, Message);
+			}
+		};
 
-	MessageType* Response = new MessageType();
-	grpc::Status* Status = new grpc::Status();
-
-	FQueuedItem Item;
-	Item.State = FQueuedItem::EState::Finishing;
-	Item.Status.Reset(Status);
-	Item.Message.Reset(Response);
-	Item.ClientContext = MoveTemp(ClientContext);
-	Item.Finish = MoveTemp(OnFinishResponse);
-	QueuedItems.Add(AsyncReader, MoveTemp(Item));
-	AsyncReader->StartCall();
-	AsyncReader->Finish(Response, Status, AsyncReader);
-	return true;
-}
+		MessageType* Response = new MessageType();
+		grpc::Status* Status = new grpc::Status();
+
+		FQueuedItem Item;
+		Item.State = FQueuedItem::EState::Finishing;
+		Item.Status.Reset(Status);
+		Item.Message.Reset(Response);
+		Item.ClientContext = MoveTemp(ClientContext);
+		Item.Finish = MoveTemp(OnFinishResponse);
+		QueuedItems.Add(AsyncReader, MoveTemp(Item));
+		AsyncReader->StartCall();
+		AsyncReader->Finish(Response, Status, AsyncReader);
+		return true;
+	}
+
+	grpc::CompletionQueue* GetCompletionQueue();
+};
