BazelExecutor: Add operation name to output

#rnx

diff --git a/Engine/Source/Editor/Experimental/BazelExecutor/Private/Execution.cpp b/Engine/Source/Editor/Experimental/BazelExecutor/Private/Execution.cpp
--- a/Engine/Source/Editor/Experimental/BazelExecutor/Private/Execution.cpp
+++ b/Engine/Source/Editor/Experimental/BazelExecutor/Private/Execution.cpp
@@ -40,28 +40,36 @@ bool FExecution::Execute(const FExecuteRequest& Request, FExecuteResponse& Respo
 	while (true)
 	{
 		google::longrunning::Operation Operation;
-		Call->Read(&Operation);
+		if (!Call->Read(&Operation))
+		{
+			UE_LOG(LogBazelExecutor, Error, TEXT("Execute: Unable to read operation"));
+			break;
+		}
 
 		build::bazel::remote::execution::v2::ExecuteOperationMetadata Metadata;
 		if (!Operation.metadata().UnpackTo(&Metadata))
 		{
-			UE_LOG(LogBazelExecutor, Error, TEXT("Execute: Unable to decode metadata"));
+			UE_LOG(LogBazelExecutor, Error, TEXT("Execute: %s Unable to decode metadata"), UTF8_TO_TCHAR(Operation.name().c_str()));
 			break;
 		}
 
-		UE_LOG(LogBazelExecutor, Display, TEXT("Execute: Execution state: %s"), UTF8_TO_TCHAR(build::bazel::remote::execution::v2::ExecutionStage_Value_Name(Metadata.stage()).c_str()));
+		UE_LOG(LogBazelExecutor, Display, TEXT("Execute: %s Execution state: %s"),
+			UTF8_TO_TCHAR(Operation.name().c_str()),
+			UTF8_TO_TCHAR(build::bazel::remote::execution::v2::ExecutionStage_Value_Name(Metadata.stage()).c_str()));
 
 		if (Operation.done())
 		{
 			build::bazel::remote::execution::v2::ExecuteResponse ProtoResponse;
 			if (!Operation.response().UnpackTo(&ProtoResponse))
 			{
-				UE_LOG(LogBazelExecutor, Error, TEXT("Execute: Unable to decode response"));
+				UE_LOG(LogBazelExecutor, Error, TEXT("Execute: %s Unable to decode response"), UTF8_TO_TCHAR(Operation.name().c_str()));
 				break;
 			}
 			ProtoConverter::FromProto(ProtoResponse, Response);
 			if ((grpc::StatusCode)ProtoResponse.status().code() != grpc::StatusCode::OK) {
-				UE_LOG(LogBazelExecutor, Error, TEXT("%s"), UTF8_TO_TCHAR(ProtoResponse.status().message().c_str()));
+				UE_LOG(LogBazelExecutor, Error, TEXT("Execute: %s Error: %s"),
+					UTF8_TO_TCHAR(Operation.name().c_str()),
+					UTF8_TO_TCHAR(ProtoResponse.status().message().c_str()));
 				break;
 			}
 			return true;
@@ -104,11 +112,13 @@ TFuture<FExecuteResponse> FExecution::ExecuteAsync(const FExecuteRequest& Reques
 		build::bazel::remote::execution::v2::ExecuteOperationMetadata Metadata;
 		if (!Operation.metadata().UnpackTo(&Metadata))
 		{
-			UE_LOG(LogBazelExecutor, Error, TEXT("ExecuteAsync: Unable to decode metadata"));
+			UE_LOG(LogBazelExecutor, Error, TEXT("ExecuteAsync: %s Unable to decode metadata"), UTF8_TO_TCHAR(Operation.name().c_str()));
 			return;
 		}
 
-		UE_LOG(LogBazelExecutor, Display, TEXT("ExecuteAsync: Execution state: %s"), UTF8_TO_TCHAR(build::bazel::remote::execution::v2::ExecutionStage_Value_Name(Metadata.stage()).c_str()));
+		UE_LOG(LogBazelExecutor, Display, TEXT("ExecuteAsync: %s Execution state: %s"),
+			UTF8_TO_TCHAR(Operation.name().c_str()),
+			UTF8_TO_TCHAR(build::bazel::remote::execution::v2::ExecutionStage_Value_Name(Metadata.stage()).c_str()));
 	};
 
 	TFinishFunction Finish = [ReturnPromise](void* Tag, bool Ok, const grpc::Status& ProtoStatus, const google::protobuf::Message& Message)
@@ -121,7 +131,7 @@ TFuture<FExecuteResponse> FExecution::ExecuteAsync(const FExecuteRequest& Reques
 		}
 		else
 		{
-			UE_LOG(LogBazelExecutor, Error, TEXT("ExecuteAsync: Finish !Ok"));
+			UE_LOG(LogBazelExecutor, Error, TEXT("ExecuteAsync: %s Finish !Ok"), UTF8_TO_TCHAR(Operation.name().c_str()));
 			Status.Code = EStatusCode::ABORTED;
 		}
 
@@ -137,7 +147,7 @@ TFuture<FExecuteResponse> FExecution::ExecuteAsync(const FExecuteRequest& Reques
 		build::bazel::remote::execution::v2::ExecuteResponse ProtoResponse;
 		if (!Operation.response().UnpackTo(&ProtoResponse))
 		{
-			UE_LOG(LogBazelExecutor, Error, TEXT("ExecuteAsync: Unable to decode response"));
+			UE_LOG(LogBazelExecutor, Error, TEXT("ExecuteAsync: %s Unable to decode response"), UTF8_TO_TCHAR(Operation.name().c_str()));
 			FExecuteResponse Response;
 			Response.Status.Code = EStatusCode::INTERNAL;
 			Response.Status.Message = TEXT("Unable to decode response");
@@ -145,6 +155,12 @@ TFuture<FExecuteResponse> FExecution::ExecuteAsync(const FExecuteRequest& Reques
 			return;
 		}
 
+		if ((grpc::StatusCode)ProtoResponse.status().code() != grpc::StatusCode::OK) {
+			UE_LOG(LogBazelExecutor, Error, TEXT("ExecuteAsync: %s Error: %s"),
+				UTF8_TO_TCHAR(Operation.name().c_str()),
+				UTF8_TO_TCHAR(ProtoResponse.status().message().c_str()));
+		}
+
 		FExecuteResponse Response;
 		ProtoConverter::FromProto(ProtoResponse, Response);
 		ReturnPromise->EmplaceValue(MoveTemp(Response));
