ToolsFramework: Add UEdModeInteractiveToolsContext::SetEnableRenderingDuringHitProxyPass() to allow ITF clients to enable/disable hitproxy rendering. Default to enabled as it is expected that most existing FEdModes/etc are using hitproxy. Disable in ModelingMode because we don't want hitproxy rendering there.
#rnx

diff --git a/Engine/Plugins/Experimental/ModelingToolsEditorMode/Source/ModelingToolsEditorMode/Private/ModelingToolsEditorMode.cpp b/Engine/Plugins/Experimental/ModelingToolsEditorMode/Source/ModelingToolsEditorMode/Private/ModelingToolsEditorMode.cpp
--- a/Engine/Plugins/Experimental/ModelingToolsEditorMode/Source/ModelingToolsEditorMode/Private/ModelingToolsEditorMode.cpp
+++ b/Engine/Plugins/Experimental/ModelingToolsEditorMode/Source/ModelingToolsEditorMode/Private/ModelingToolsEditorMode.cpp
@@ -313,6 +313,9 @@ void UModelingToolsEditorMode::Enter()
 	// register gizmo helper
 	UE::TransformGizmoUtil::RegisterTransformGizmoContextObject(ToolsContext.Get());
 
+	// disable HitProxy rendering, it is not used in Modeling Mode and adds overhead to Render() calls
+	ToolsContext.Get()->SetEnableRenderingDuringHitProxyPass(false);
+
 	// register object creation api
 	UEditorModelingObjectsCreationAPI* ModelCreationAPI = UEditorModelingObjectsCreationAPI::Register(ToolsContext.Get());
 	if (ModelCreationAPI)
@@ -719,6 +722,9 @@ void UModelingToolsEditorMode::Exit()
 	// clear realtime viewport override
 	ConfigureRealTimeViewportsOverride(false);
 
+	// re-enable HitProxy rendering
+	ToolsContext.Get()->SetEnableRenderingDuringHitProxyPass(true);
+
 	// Call base Exit method to ensure proper cleanup
 	UEdMode::Exit();
 }
diff --git a/Engine/Source/Editor/Experimental/EditorInteractiveToolsFramework/Private/EdModeInteractiveToolsContext.cpp b/Engine/Source/Editor/Experimental/EditorInteractiveToolsFramework/Private/EdModeInteractiveToolsContext.cpp
--- a/Engine/Source/Editor/Experimental/EditorInteractiveToolsFramework/Private/EdModeInteractiveToolsContext.cpp
+++ b/Engine/Source/Editor/Experimental/EditorInteractiveToolsFramework/Private/EdModeInteractiveToolsContext.cpp
@@ -820,8 +820,8 @@ public:
 
 void UEdModeInteractiveToolsContext::Render(const FSceneView* View, FViewport* Viewport, FPrimitiveDrawInterface* PDI)
 {
-	// tools framework cannot use HitProxy so skip these calls
-	if (PDI->IsHitTesting())
+	// skip HitProxy rendering passes if desired
+	if (PDI->IsHitTesting() && bEnableRenderingDuringHitProxyPass == false)
 	{
 		return;
 	}
@@ -1273,3 +1273,9 @@ void UEdModeInteractiveToolsContext::OnToolPostBuild(UInteractiveToolManager* In
 {
 	// todo: Add any shared tool targets for the mode toolkit
 }
+
+
+void UEdModeInteractiveToolsContext::SetEnableRenderingDuringHitProxyPass(bool bEnabled)
+{
+	bEnableRenderingDuringHitProxyPass = bEnabled;
+}
\ No newline at end of file
diff --git a/Engine/Source/Editor/Experimental/EditorInteractiveToolsFramework/Public/EdModeInteractiveToolsContext.h b/Engine/Source/Editor/Experimental/EditorInteractiveToolsFramework/Public/EdModeInteractiveToolsContext.h
--- a/Engine/Source/Editor/Experimental/EditorInteractiveToolsFramework/Public/EdModeInteractiveToolsContext.h
+++ b/Engine/Source/Editor/Experimental/EditorInteractiveToolsFramework/Public/EdModeInteractiveToolsContext.h
@@ -89,6 +89,19 @@ public:
 
 	FRay GetLastWorldRay() const;
 
+	//
+	// Configuration functions
+	//
+
+	/*
+	 * Configure whether ::Render() should early-out for HitProxy rendering passes.
+	 * If the Mode does not use HitProxy, and the Tools/Gizmos have expensive Render() calls, this can help with interactive performance.
+	 */
+	void SetEnableRenderingDuringHitProxyPass(bool bEnabled);
+
+	/** @return true if HitProxy rendering will be allowed in ::Render() */
+	bool GetEnableRenderingDuringHitProxyPass() const { return bEnableRenderingDuringHitProxyPass; }
+
 protected:
 	// we hide these 
 	virtual void Initialize(IToolsContextQueriesAPI* QueriesAPI, IToolsContextTransactionsAPI* TransactionsAPI) override;
@@ -150,4 +163,7 @@ protected:
 private:
 	FEditorModeTools* EditorModeManager = nullptr;
 	bool bIsTrackingMouse;
+
+	// currently defaulting to enabled as FEdModes generally assume this, and in most cases hitproxy pass is not expensive.
+	bool bEnableRenderingDuringHitProxyPass = true;
 };
