BazelExecutor: Enable Linux support

diff --git a/Engine/Source/Editor/Experimental/BazelExecutor/BazelExecutor.Build.cs b/Engine/Source/Editor/Experimental/BazelExecutor/BazelExecutor.Build.cs
--- a/Engine/Source/Editor/Experimental/BazelExecutor/BazelExecutor.Build.cs
+++ b/Engine/Source/Editor/Experimental/BazelExecutor/BazelExecutor.Build.cs
@@ -9,7 +9,7 @@ namespace UnrealBuildTool.Rules
 	{
 		public BazelExecutor(ReadOnlyTargetRules Target) : base(Target)
 		{
-			if (Target.Platform != UnrealTargetPlatform.Win64)
+			if (!IsVcPackageSupported)
 			{
 				// Grpc is currently only supported for Win64, Linux, & LinuxAArch64, set to external to prevent compiling for any other platform
 				Type = ModuleType.External;
