[Backout] - CL16786490
Original CL Desc
-----------------------------------------------------------------
BazelExecutor: Enable Linux support

diff --git a/Engine/Source/Editor/Experimental/BazelExecutor/BazelExecutor.Build.cs b/Engine/Source/Editor/Experimental/BazelExecutor/BazelExecutor.Build.cs
--- a/Engine/Source/Editor/Experimental/BazelExecutor/BazelExecutor.Build.cs
+++ b/Engine/Source/Editor/Experimental/BazelExecutor/BazelExecutor.Build.cs
@@ -9,7 +9,7 @@ namespace UnrealBuildTool.Rules
 	{
 		public BazelExecutor(ReadOnlyTargetRules Target) : base(Target)
 		{
-			if (!IsVcPackageSupported)
+			if (Target.Platform != UnrealTargetPlatform.Win64)
 			{
 				// Grpc is currently only supported for Win64, Linux, & LinuxAArch64, set to external to prevent compiling for any other platform
 				Type = ModuleType.External;
