BazelExecutor: Forward declare anything related to protobuf or grpc
Enable Linux build

#rnx

diff --git a/Engine/Source/Editor/Experimental/BazelExecutor/BazelExecutor.Build.cs b/Engine/Source/Editor/Experimental/BazelExecutor/BazelExecutor.Build.cs
--- a/Engine/Source/Editor/Experimental/BazelExecutor/BazelExecutor.Build.cs
+++ b/Engine/Source/Editor/Experimental/BazelExecutor/BazelExecutor.Build.cs
@@ -9,7 +9,7 @@ namespace UnrealBuildTool.Rules
 	{
 		public BazelExecutor(ReadOnlyTargetRules Target) : base(Target)
 		{
-			if (Target.Platform != UnrealTargetPlatform.Win64)
+			if (!IsVcPackageSupported)
 			{
 				// Grpc is currently only supported for Win64, Linux, & LinuxAArch64, set to external to prevent compiling for any other platform
 				Type = ModuleType.External;
diff --git a/Engine/Source/Editor/Experimental/BazelExecutor/Private/BazelCompletionQueueRunnable.cpp b/Engine/Source/Editor/Experimental/BazelExecutor/Private/BazelCompletionQueueRunnable.cpp
--- a/Engine/Source/Editor/Experimental/BazelExecutor/Private/BazelCompletionQueueRunnable.cpp
+++ b/Engine/Source/Editor/Experimental/BazelExecutor/Private/BazelCompletionQueueRunnable.cpp
@@ -4,8 +4,17 @@
 
 #include <chrono>
 
+THIRD_PARTY_INCLUDES_START
+UE_PUSH_MACRO("TEXT")
+#undef TEXT
+#include <grpcpp/grpcpp.h>
+#include "build\bazel\remote\execution\v2\remote_execution.grpc.pb.h"
+UE_POP_MACRO("TEXT");
+THIRD_PARTY_INCLUDES_END
+
 
 FBazelCompletionQueueRunnable::FBazelCompletionQueueRunnable()
+	: CompletionQueue(new grpc::CompletionQueue())
 {
 }
 
@@ -85,7 +94,7 @@ uint32 FBazelCompletionQueueRunnable::Run()
 {
 	void* Tag;
 	bool Ok = false;
-	while (Running && CompletionQueue.Next(&Tag, &Ok))
+	while (Running && CompletionQueue->Next(&Tag, &Ok))
 	{
 		if (!Running)
 		{
@@ -100,7 +109,7 @@ void FBazelCompletionQueueRunnable::Stop()
 {
 	// Request shutdown of the CompletionQueue
 	Running = false;
-	CompletionQueue.Shutdown();
+	CompletionQueue->Shutdown();
 }
 
 void FBazelCompletionQueueRunnable::Exit()
@@ -110,7 +119,7 @@ void FBazelCompletionQueueRunnable::Exit()
 	// Drain the completion queue
 	void* Tag;
 	bool Ok;
-	while (CompletionQueue.Next(&Tag, &Ok)) {}
+	while (CompletionQueue->Next(&Tag, &Ok)) {}
 
 	for (const TPair<void*,FQueuedItem>& Pair : QueuedItems)
 	{
@@ -141,7 +150,7 @@ void FBazelCompletionQueueRunnable::Tick()
 	const std::chrono::time_point Deadline = std::chrono::system_clock::now() + std::chrono::microseconds(100);
 	while (Running && std::chrono::system_clock::now() < Deadline)
 	{
-		grpc::CompletionQueue::NextStatus Status = CompletionQueue.AsyncNext(&Tag, &Ok, Deadline);
+		grpc::CompletionQueue::NextStatus Status = CompletionQueue->AsyncNext(&Tag, &Ok, Deadline);
 		if (Status == grpc::CompletionQueue::NextStatus::GOT_EVENT)
 		{
 			ProcessNext(Tag, Ok);
@@ -186,5 +195,5 @@ grpc::CompletionQueue* FBazelCompletionQueueRunnable::GetCompletionQueue()
 		return nullptr;
 	}
 
-	return &CompletionQueue;
+	return CompletionQueue.Get();
 }
\ No newline at end of file
diff --git a/Engine/Source/Editor/Experimental/BazelExecutor/Private/BazelCompletionQueueRunnable.h b/Engine/Source/Editor/Experimental/BazelExecutor/Private/BazelCompletionQueueRunnable.h
--- a/Engine/Source/Editor/Experimental/BazelExecutor/Private/BazelCompletionQueueRunnable.h
+++ b/Engine/Source/Editor/Experimental/BazelExecutor/Private/BazelCompletionQueueRunnable.h
@@ -7,13 +7,29 @@
 #include "HAL/RunnableThread.h"
 #include "Misc/SingleThreadRunnable.h"
 
-THIRD_PARTY_INCLUDES_START
-UE_PUSH_MACRO("TEXT")
-#undef TEXT
-#include <grpcpp/grpcpp.h>
-#include "build\bazel\remote\execution\v2\remote_execution.grpc.pb.h"
-UE_POP_MACRO("TEXT");
-THIRD_PARTY_INCLUDES_END
+namespace grpc
+{
+	class ClientContext;
+	class Status;
+	template <class R>
+	class ClientAsyncReader;
+	template <class R>
+	class ClientAsyncResponseReader;
+	class CompletionQueue;
+}
+
+namespace google
+{
+	namespace longrunning
+	{
+		class Operation;
+	}
+
+	namespace protobuf
+	{
+		class Message;
+	}
+}
 
 
 using TStartCallFunction = TUniqueFunction<void(void* Tag, bool Ok)>;
@@ -43,7 +59,7 @@ private:
 
 	bool Running;
 	TMap<void*, FQueuedItem> QueuedItems;
-	grpc::CompletionQueue CompletionQueue;
+	TUniquePtr<grpc::CompletionQueue> CompletionQueue;
 
 	void ProcessNext(void* Tag, bool Ok);
 
@@ -77,39 +93,7 @@ public:
 	bool AddAsyncResponse(
 		TUniquePtr<grpc::ClientContext> ClientContext,
 		TUniquePtr<grpc::ClientAsyncResponseReader<MessageType>> Reader,
-		TFinishFunction OnFinish = nullptr)
-	{
-		static_assert(std::is_base_of<google::protobuf::Message, MessageType>::value, "MessageType is not derived from google::protobuf::Message");
-		if (!Running || !ClientContext.IsValid() || !Reader.IsValid())
-		{
-			return false;
-		}
-
-		grpc::ClientAsyncResponseReader<MessageType>* AsyncReader = Reader.Get();
-
-		// Keep Reader in scope for the lifetime of this async call
-		TFinishFunction OnFinishResponse = [Reader = MoveTemp(Reader), OnFinish = MoveTemp(OnFinish)](void* Tag, bool Ok, const grpc::Status& Status, const google::protobuf::Message& Message)
-		{
-			if (OnFinish)
-			{
-				OnFinish(Tag, Ok, Status, Message);
-			}
-		};
-
-		MessageType* Response = new MessageType();
-		grpc::Status* Status = new grpc::Status();
-
-		FQueuedItem Item;
-		Item.State = FQueuedItem::EState::Finishing;
-		Item.Status.Reset(Status);
-		Item.Message.Reset(Response);
-		Item.ClientContext = MoveTemp(ClientContext);
-		Item.Finish = MoveTemp(OnFinishResponse);
-		QueuedItems.Add(AsyncReader, MoveTemp(Item));
-		AsyncReader->StartCall();
-		AsyncReader->Finish(Response, Status, AsyncReader);
-		return true;
-	}
+		TFinishFunction OnFinish = nullptr);
 
 	grpc::CompletionQueue* GetCompletionQueue();
 };
diff --git a/Engine/Source/Editor/Experimental/BazelExecutor/Private/BazelCompletionQueueRunnable.inl b/Engine/Source/Editor/Experimental/BazelExecutor/Private/BazelCompletionQueueRunnable.inl
new file mode 100644
--- /dev/null
+++ b/Engine/Source/Editor/Experimental/BazelExecutor/Private/BazelCompletionQueueRunnable.inl
@@ -0,0 +1,42 @@
+// Copyright Epic Games, Inc. All Rights Reserved.
+
+#pragma once
+
+
+template<class MessageType>
+bool FBazelCompletionQueueRunnable::AddAsyncResponse(
+	TUniquePtr<grpc::ClientContext> ClientContext,
+	TUniquePtr<grpc::ClientAsyncResponseReader<MessageType>> Reader,
+	TFinishFunction OnFinish)
+{
+	static_assert(std::is_base_of<google::protobuf::Message, MessageType>::value, "MessageType is not derived from google::protobuf::Message");
+	if (!Running || !ClientContext.IsValid() || !Reader.IsValid())
+	{
+		return false;
+	}
+
+	grpc::ClientAsyncResponseReader<MessageType>* AsyncReader = Reader.Get();
+
+	// Keep Reader in scope for the lifetime of this async call
+	TFinishFunction OnFinishResponse = [Reader = MoveTemp(Reader), OnFinish = MoveTemp(OnFinish)](void* Tag, bool Ok, const grpc::Status& Status, const google::protobuf::Message& Message)
+	{
+		if (OnFinish)
+		{
+			OnFinish(Tag, Ok, Status, Message);
+		}
+	};
+
+	MessageType* Response = new MessageType();
+	grpc::Status* Status = new grpc::Status();
+
+	FQueuedItem Item;
+	Item.State = FQueuedItem::EState::Finishing;
+	Item.Status.Reset(Status);
+	Item.Message.Reset(Response);
+	Item.ClientContext = MoveTemp(ClientContext);
+	Item.Finish = MoveTemp(OnFinishResponse);
+	QueuedItems.Add(AsyncReader, MoveTemp(Item));
+	AsyncReader->StartCall();
+	AsyncReader->Finish(Response, Status, AsyncReader);
+	return true;
+}
\ No newline at end of file
diff --git a/Engine/Source/Editor/Experimental/BazelExecutor/Private/ContentAddressableStorage.cpp b/Engine/Source/Editor/Experimental/BazelExecutor/Private/ContentAddressableStorage.cpp
--- a/Engine/Source/Editor/Experimental/BazelExecutor/Private/ContentAddressableStorage.cpp
+++ b/Engine/Source/Editor/Experimental/BazelExecutor/Private/ContentAddressableStorage.cpp
@@ -12,12 +12,27 @@ THIRD_PARTY_INCLUDES_START
 UE_PUSH_MACRO("TEXT")
 #undef TEXT
 #include <grpcpp/grpcpp.h>
+#include "build\bazel\remote\execution\v2\remote_execution.grpc.pb.h"
 UE_POP_MACRO("TEXT");
 THIRD_PARTY_INCLUDES_END
 
+#include "BazelCompletionQueueRunnable.inl"
+
+
+// Stub class so grpc isn't included in the header and everything can be forward declared
+class FContentAddressableStorageStub
+{
+public:
+	TUniquePtr<build::bazel::remote::execution::v2::ContentAddressableStorage::Stub> Stub;
+	FContentAddressableStorageStub(build::bazel::remote::execution::v2::ContentAddressableStorage::Stub* Stub)
+		: Stub(Stub)
+	{
+	}
+};
+
 
 FContentAddressableStorage::FContentAddressableStorage(const std::shared_ptr<grpc::Channel>& Channel, TSharedPtr<FBazelCompletionQueueRunnable> CompletionQueueRunnable, const TMap<FString, FString>& Headers) :
-	Stub(build::bazel::remote::execution::v2::ContentAddressableStorage::NewStub(Channel).release()),
+	ContentAddressableStorageStub(new FContentAddressableStorageStub(build::bazel::remote::execution::v2::ContentAddressableStorage::NewStub(Channel).release())),
 	CompletionQueueRunnable(CompletionQueueRunnable),
 	Headers(Headers)
 {
@@ -56,7 +71,7 @@ FStatus FContentAddressableStorage::FindMissingBlobs(const FFindMissingBlobsRequ
 	ProtoConverter::ToProto(Request, ProtoRequest);
 
 	build::bazel::remote::execution::v2::FindMissingBlobsResponse ProtoResponse;
-	grpc::Status ProtoStatus = Stub->FindMissingBlobs(&ClientContext, ProtoRequest, &ProtoResponse);
+	grpc::Status ProtoStatus = ContentAddressableStorageStub->Stub->FindMissingBlobs(&ClientContext, ProtoRequest, &ProtoResponse);
 	FStatus Status;
 	ProtoConverter::FromProto(ProtoStatus, Status);
 	ProtoConverter::FromProto(ProtoResponse, Response);
@@ -73,7 +88,7 @@ FStatus FContentAddressableStorage::BatchUpdateBlobs(const FBatchUpdateBlobsRequ
 	ProtoConverter::ToProto(Request, ProtoRequest);
 
 	build::bazel::remote::execution::v2::BatchUpdateBlobsResponse ProtoResponse;
-	grpc::Status ProtoStatus = Stub->BatchUpdateBlobs(&ClientContext, ProtoRequest, &ProtoResponse);
+	grpc::Status ProtoStatus = ContentAddressableStorageStub->Stub->BatchUpdateBlobs(&ClientContext, ProtoRequest, &ProtoResponse);
 	FStatus Status;
 	ProtoConverter::FromProto(ProtoStatus, Status);
 	ProtoConverter::FromProto(ProtoResponse, Response);
@@ -90,7 +105,7 @@ FStatus FContentAddressableStorage::BatchReadBlobs(const FBatchReadBlobsRequest&
 	ProtoConverter::ToProto(Request, ProtoRequest);
 
 	build::bazel::remote::execution::v2::BatchReadBlobsResponse ProtoResponse;
-	grpc::Status ProtoStatus = Stub->BatchReadBlobs(&ClientContext, ProtoRequest, &ProtoResponse);
+	grpc::Status ProtoStatus = ContentAddressableStorageStub->Stub->BatchReadBlobs(&ClientContext, ProtoRequest, &ProtoResponse);
 	FStatus Status;
 	ProtoConverter::FromProto(ProtoStatus, Status);
 	ProtoConverter::FromProto(ProtoResponse, Response);
@@ -107,7 +122,7 @@ TFuture<TPair<FStatus, FFindMissingBlobsResponse>> FContentAddressableStorage::F
 	ProtoConverter::ToProto(Request, ProtoRequest);
 
 	TUniquePtr<grpc::ClientAsyncResponseReader<build::bazel::remote::execution::v2::FindMissingBlobsResponse>> ProtoReader(
-		Stub->PrepareAsyncFindMissingBlobs(ClientContext.Get(), ProtoRequest, CompletionQueueRunnable->GetCompletionQueue()).release());
+		ContentAddressableStorageStub->Stub->PrepareAsyncFindMissingBlobs(ClientContext.Get(), ProtoRequest, CompletionQueueRunnable->GetCompletionQueue()).release());
 
 	TSharedPtr<TPromise<TPair<FStatus, FFindMissingBlobsResponse>>> ReturnPromise = MakeShared<TPromise<TPair<FStatus, FFindMissingBlobsResponse>>>(MoveTemp(CompletionCallback));
 
@@ -152,7 +167,7 @@ TFuture<TPair<FStatus, FBatchUpdateBlobsResponse>> FContentAddressableStorage::B
 	ProtoConverter::ToProto(Request, ProtoRequest);
 
 	TUniquePtr<grpc::ClientAsyncResponseReader<build::bazel::remote::execution::v2::BatchUpdateBlobsResponse>> ProtoReader(
-		Stub->PrepareAsyncBatchUpdateBlobs(ClientContext.Get(), ProtoRequest, CompletionQueueRunnable->GetCompletionQueue()).release());
+		ContentAddressableStorageStub->Stub->PrepareAsyncBatchUpdateBlobs(ClientContext.Get(), ProtoRequest, CompletionQueueRunnable->GetCompletionQueue()).release());
 
 	TSharedPtr<TPromise<TPair<FStatus, FBatchUpdateBlobsResponse>>> ReturnPromise = MakeShared<TPromise<TPair<FStatus, FBatchUpdateBlobsResponse>>>(MoveTemp(CompletionCallback));
 
@@ -196,7 +211,7 @@ TFuture<TPair<FStatus, FBatchReadBlobsResponse>> FContentAddressableStorage::Bat
 	ProtoConverter::ToProto(Request, ProtoRequest);
 
 	TUniquePtr<grpc::ClientAsyncResponseReader<build::bazel::remote::execution::v2::BatchReadBlobsResponse>> ProtoReader(
-		Stub->PrepareAsyncBatchReadBlobs(ClientContext.Get(), ProtoRequest, CompletionQueueRunnable->GetCompletionQueue()).release());
+		ContentAddressableStorageStub->Stub->PrepareAsyncBatchReadBlobs(ClientContext.Get(), ProtoRequest, CompletionQueueRunnable->GetCompletionQueue()).release());
 
 	TSharedPtr<TPromise<TPair<FStatus, FBatchReadBlobsResponse>>> ReturnPromise = MakeShared<TPromise<TPair<FStatus, FBatchReadBlobsResponse>>>(MoveTemp(CompletionCallback));
 
diff --git a/Engine/Source/Editor/Experimental/BazelExecutor/Private/ContentAddressableStorage.h b/Engine/Source/Editor/Experimental/BazelExecutor/Private/ContentAddressableStorage.h
--- a/Engine/Source/Editor/Experimental/BazelExecutor/Private/ContentAddressableStorage.h
+++ b/Engine/Source/Editor/Experimental/BazelExecutor/Private/ContentAddressableStorage.h
@@ -5,20 +5,22 @@
 #include "CoreMinimal.h"
 #include "IContentAddressableStorage.h"
 
-THIRD_PARTY_INCLUDES_START
-UE_PUSH_MACRO("TEXT")
-#undef TEXT
-#include "build\bazel\remote\execution\v2\remote_execution.grpc.pb.h"
-UE_POP_MACRO("TEXT");
-THIRD_PARTY_INCLUDES_END
+#include <memory>
+
+namespace grpc
+{
+	class Channel;
+	class ClientContext;
+}
 
 class FBazelCompletionQueueRunnable;
+class FContentAddressableStorageStub;
 
 
 class FContentAddressableStorage : public IContentAddressableStorage
 {
 private:
-	TUniquePtr<build::bazel::remote::execution::v2::ContentAddressableStorage::Stub> Stub;
+	TUniquePtr<FContentAddressableStorageStub> ContentAddressableStorageStub;
 	TSharedPtr<FBazelCompletionQueueRunnable> CompletionQueueRunnable;
 	TMap<FString, FString> Headers;
 
diff --git a/Engine/Source/Editor/Experimental/BazelExecutor/Private/Execution.cpp b/Engine/Source/Editor/Experimental/BazelExecutor/Private/Execution.cpp
--- a/Engine/Source/Editor/Experimental/BazelExecutor/Private/Execution.cpp
+++ b/Engine/Source/Editor/Experimental/BazelExecutor/Private/Execution.cpp
@@ -12,12 +12,27 @@ THIRD_PARTY_INCLUDES_START
 UE_PUSH_MACRO("TEXT")
 #undef TEXT
 #include <grpcpp/grpcpp.h>
+#include "build\bazel\remote\execution\v2\remote_execution.grpc.pb.h"
 UE_POP_MACRO("TEXT");
 THIRD_PARTY_INCLUDES_END
 
+#include "BazelCompletionQueueRunnable.inl"
+
+
+// Stub class so grpc isn't included in the header and everything can be forward declared
+class FExecutionStub
+{
+public:
+	TUniquePtr<build::bazel::remote::execution::v2::Execution::Stub> Stub;
+	FExecutionStub(build::bazel::remote::execution::v2::Execution::Stub* Stub)
+		: Stub(Stub)
+	{
+	}
+};
+
 
 FExecution::FExecution(const std::shared_ptr<grpc::Channel>& Channel, TSharedPtr<FBazelCompletionQueueRunnable> CompletionQueueRunnable, const TMap<FString, FString>& Headers) :
-	Stub(build::bazel::remote::execution::v2::Execution::NewStub(Channel).release()),
+	ExecutionStub(new FExecutionStub(build::bazel::remote::execution::v2::Execution::NewStub(Channel).release())),
 	CompletionQueueRunnable(CompletionQueueRunnable),
 	Headers(Headers)
 {
@@ -35,7 +50,7 @@ bool FExecution::Execute(const FExecuteRequest& Request, FExecuteResponse& Respo
 	build::bazel::remote::execution::v2::ExecuteRequest ProtoRequest;
 	ProtoConverter::ToProto(Request, ProtoRequest);
 
-	std::unique_ptr<grpc::ClientReader<google::longrunning::Operation>> Call = Stub->Execute(&ClientContext, ProtoRequest);
+	std::unique_ptr<grpc::ClientReader<google::longrunning::Operation>> Call = ExecutionStub->Stub->Execute(&ClientContext, ProtoRequest);
 	Call->WaitForInitialMetadata();
 	while (true)
 	{
@@ -99,7 +114,7 @@ TFuture<FExecuteResponse> FExecution::ExecuteAsync(const FExecuteRequest& Reques
 	ProtoConverter::ToProto(Request, ProtoRequest);
 
 	TUniquePtr<grpc::ClientAsyncReader<google::longrunning::Operation>> OperationReader(
-		Stub->PrepareAsyncExecute(ClientContext.Get(), ProtoRequest, CompletionQueueRunnable->GetCompletionQueue()).release());
+		ExecutionStub->Stub->PrepareAsyncExecute(ClientContext.Get(), ProtoRequest, CompletionQueueRunnable->GetCompletionQueue()).release());
 
 	TSharedPtr<TPromise<FExecuteResponse>> ReturnPromise = MakeShared<TPromise<FExecuteResponse>>(MoveTemp(CompletionCallback));
 
diff --git a/Engine/Source/Editor/Experimental/BazelExecutor/Private/Execution.h b/Engine/Source/Editor/Experimental/BazelExecutor/Private/Execution.h
--- a/Engine/Source/Editor/Experimental/BazelExecutor/Private/Execution.h
+++ b/Engine/Source/Editor/Experimental/BazelExecutor/Private/Execution.h
@@ -5,20 +5,22 @@
 #include "CoreMinimal.h"
 #include "IExecution.h"
 
-THIRD_PARTY_INCLUDES_START
-UE_PUSH_MACRO("TEXT")
-#undef TEXT
-#include "build\bazel\remote\execution\v2\remote_execution.grpc.pb.h"
-UE_POP_MACRO("TEXT");
-THIRD_PARTY_INCLUDES_END
+#include <memory>
+
+namespace grpc
+{
+	class Channel;
+	class ClientContext;
+}
 
 class FBazelCompletionQueueRunnable;
+class FExecutionStub;
 
 
 class FExecution : public IExecution
 {
 private:
-	TUniquePtr<build::bazel::remote::execution::v2::Execution::Stub> Stub;
+	TUniquePtr<FExecutionStub> ExecutionStub;
 	TSharedPtr<FBazelCompletionQueueRunnable> CompletionQueueRunnable;
 	TMap<FString, FString> Headers;
 
