Add typed element selection set support to Editor gizmo manager. Add toggle to enable new TRS gizmo.

diff --git a/Engine/Source/Editor/Experimental/EditorInteractiveToolsFramework/Private/EdModeInteractiveToolsContext.cpp b/Engine/Source/Editor/Experimental/EditorInteractiveToolsFramework/Private/EdModeInteractiveToolsContext.cpp
--- a/Engine/Source/Editor/Experimental/EditorInteractiveToolsFramework/Private/EdModeInteractiveToolsContext.cpp
+++ b/Engine/Source/Editor/Experimental/EditorInteractiveToolsFramework/Private/EdModeInteractiveToolsContext.cpp
@@ -630,6 +630,16 @@ void UEdModeInteractiveToolsContext::InitializeContextWithEditorModeManager(FEdi
 
 	// set up standard materials
 	StandardVertexColorMaterial = GEngine->VertexColorMaterial;
+
+	if (UTypedElementSelectionSet* TypedElementSelectionSet = EditorModeManager->GetEditorSelectionSet())
+	{
+		TypedElementSelectionSet->OnChanged().AddUObject(this, &UEdModeInteractiveToolsContext::OnEditorSelectionSetChanged);
+	}
+	else
+	{
+		FLevelEditorModule& LevelEditor = FModuleManager::Get().LoadModuleChecked<FLevelEditorModule>("LevelEditor");
+		LevelEditor.OnLevelEditorCreated().AddUObject(this, &UEdModeInteractiveToolsContext::OnLevelEditorCreated);
+	}
 }
 
 
@@ -1274,8 +1284,24 @@ void UEdModeInteractiveToolsContext::OnToolPostBuild(UInteractiveToolManager* In
 	// todo: Add any shared tool targets for the mode toolkit
 }
 
-
 void UEdModeInteractiveToolsContext::SetEnableRenderingDuringHitProxyPass(bool bEnabled)
 {
 	bEnableRenderingDuringHitProxyPass = bEnabled;
-}
\ No newline at end of file
+}
+
+void UEdModeInteractiveToolsContext::OnLevelEditorCreated(TSharedPtr<ILevelEditor> InLevelEditor)
+{
+	if (UTypedElementSelectionSet* TypedElementSelectionSet = EditorModeManager->GetEditorSelectionSet())
+	{
+		TypedElementSelectionSet->OnChanged().AddUObject(this, &UEdModeInteractiveToolsContext::OnEditorSelectionSetChanged);
+	}
+}
+
+void UEdModeInteractiveToolsContext::OnEditorSelectionSetChanged(const UTypedElementSelectionSet* InSelectionSet)
+{
+	if (UEditorInteractiveGizmoManager* EditorGizmoManager = Cast<UEditorInteractiveGizmoManager>(GizmoManager))
+	{
+		EditorGizmoManager->HandleEditorSelectionSetChanged(InSelectionSet);
+	}
+}
+
diff --git a/Engine/Source/Editor/Experimental/EditorInteractiveToolsFramework/Private/EditorGizmos/EditorTransformGizmo.cpp b/Engine/Source/Editor/Experimental/EditorInteractiveToolsFramework/Private/EditorGizmos/EditorTransformGizmo.cpp
--- a/Engine/Source/Editor/Experimental/EditorInteractiveToolsFramework/Private/EditorGizmos/EditorTransformGizmo.cpp
+++ b/Engine/Source/Editor/Experimental/EditorInteractiveToolsFramework/Private/EditorGizmos/EditorTransformGizmo.cpp
@@ -10,6 +10,8 @@
 #include "BaseGizmos/TransformSources.h"
 #include "EditorGizmos/EditorAxisSources.h"
 #include "EditorGizmos/EditorParameterToTransformAdapters.h"
+#include "Elements/Interfaces/TypedElementObjectInterface.h"
+#include "Elements/Interfaces/TypedElementWorldInterface.h"
 #include "Engine/CollisionProfile.h"
 #include "Engine/Selection.h"
 #include "Engine/World.h"
@@ -47,61 +49,55 @@ constexpr FColor UEditorTransformGizmo::CurrentColor;
 
 UInteractiveGizmo* UEditorTransformGizmoBuilder::BuildGizmo(const FToolBuilderState& SceneState) const
 {
-	// @todo - remove global call
-	FEditorModeTools& ModeTools = GLevelEditorModeTools();
-
-	ETransformGizmoSubElements Elements  = ETransformGizmoSubElements::None;
-	bool bUseContextCoordinateSystem = true;
-	UE::Widget::EWidgetMode WidgetMode = ModeTools.GetWidgetMode();
-	switch ( WidgetMode )
-	{
-	case UE::Widget::EWidgetMode::WM_Translate:
-		Elements = ETransformGizmoSubElements::TranslateAllAxes | ETransformGizmoSubElements::TranslateAllPlanes;
-		break;
-	case UE::Widget::EWidgetMode::WM_Rotate:
-		Elements = ETransformGizmoSubElements::RotateAllAxes;
-		break;
-	case UE::Widget::EWidgetMode::WM_Scale:
-		Elements = ETransformGizmoSubElements::ScaleAllAxes | ETransformGizmoSubElements::ScaleAllPlanes;
-		bUseContextCoordinateSystem = false;
-		break;
-	case UE::Widget::EWidgetMode::WM_2D:
-		Elements = ETransformGizmoSubElements::RotateAxisY | ETransformGizmoSubElements::TranslatePlaneXZ;
-		break;
-	default:
-		Elements = ETransformGizmoSubElements::FullTranslateRotateScale;
-		break;
-	}
-
 	UEditorTransformGizmo* TransformGizmo = NewObject<UEditorTransformGizmo>(SceneState.GizmoManager);
 	TransformGizmo->Setup();
 
 	TransformGizmo->SetWorld(SceneState.World);
-	TransformGizmo->bUseContextCoordinateSystem = bUseContextCoordinateSystem;
 
-	// @todo - update to work with typed elements
-	TArray<AActor*> SelectedActors;
-	ModeTools.GetSelectedActors()->GetSelectedObjects<AActor>(SelectedActors);
-
-	UTransformProxy* TransformProxy = NewObject<UTransformProxy>();
-	for (auto Actor : SelectedActors)
-	{
-		USceneComponent* SceneComponent = Actor->GetRootComponent();
-		TransformProxy->AddComponent(SceneComponent);
-	}
-	TransformGizmo->SetActiveTarget( TransformProxy );
-	TransformGizmo->SetVisibility(SelectedActors.Num() > 0);
+	// @todo - remove this once transform is updated to take widget mode as a source
+	TransformGizmo->SetElements(ETransformGizmoSubElements::TranslateAllAxes | ETransformGizmoSubElements::TranslateAllPlanes);
+	TransformGizmo->bUseContextCoordinateSystem = false;
 
 	return TransformGizmo;
 }
 
-bool UEditorTransformGizmoBuilder::SatisfiesCondition(const FToolBuilderState& SceneState) const 
+void UEditorTransformGizmoBuilder::UpdateGizmoForSelection(UInteractiveGizmo* Gizmo, const FToolBuilderState& SceneState)
 {
-	if (UTypedElementSelectionSet* SelectionSet = SceneState.TypedElementSelectionSet.Get())
+	UEditorTransformGizmo* TransformGizmo = Cast<UEditorTransformGizmo>(Gizmo);
+
+	if (ensure(TransformGizmo))
 	{
-		return (SelectionSet->HasSelectedElements());
+		UTransformProxy* TransformProxy = NewObject<UTransformProxy>();
+		bool bVisibility = false;
+
+		ETypedElementWorldType WorldType = (SceneState.World->WorldType == EWorldType::PIE ? ETypedElementWorldType::Game : ETypedElementWorldType::Editor);
+
+		// @todo - once UTransformProxy supports typed elements, update this to use the normalized typed
+		// element selection set.
+		if (UTypedElementSelectionSet* SelectionSet = SceneState.TypedElementSelectionSet.Get())
+		{
+			SelectionSet->ForEachSelectedElement<UTypedElementWorldInterface>([TransformProxy, &bVisibility, SelectionSet, WorldType](const TTypedElement<UTypedElementWorldInterface>& InWorldElement)
+			{
+				if (InWorldElement.CanMoveElement(WorldType))
+				{
+					if (TTypedElement<UTypedElementObjectInterface> ObjectTypedElement = SelectionSet->GetElementList()->GetElement<UTypedElementObjectInterface>(InWorldElement))
+					{
+						if (AActor* Actor = ObjectTypedElement.GetObjectAs<AActor>())
+						{
+							USceneComponent* SceneComponent = Actor->GetRootComponent();
+							TransformProxy->AddComponent(SceneComponent);
+							bVisibility = true;
+						}
+						return true;
+					}
+				}
+				return false;
+			});
+		}
+
+		TransformGizmo->SetActiveTarget(TransformProxy);
+		TransformGizmo->SetVisibility(bVisibility);
 	}
-	return true;
 }
 
 void UEditorTransformGizmo::SetWorld(UWorld* WorldIn)
diff --git a/Engine/Source/Editor/Experimental/EditorInteractiveToolsFramework/Private/EditorInteractiveGizmoManager.cpp b/Engine/Source/Editor/Experimental/EditorInteractiveToolsFramework/Private/EditorInteractiveGizmoManager.cpp
--- a/Engine/Source/Editor/Experimental/EditorInteractiveToolsFramework/Private/EditorInteractiveGizmoManager.cpp
+++ b/Engine/Source/Editor/Experimental/EditorInteractiveToolsFramework/Private/EditorInteractiveGizmoManager.cpp
@@ -2,21 +2,24 @@
 
 
 #include "EditorInteractiveGizmoManager.h"
+#include "EditorInteractiveGizmoConditionalBuilder.h"
 #include "EditorInteractiveGizmoSelectionBuilder.h"
 #include "EditorInteractiveGizmoSubsystem.h"
 #include "EdModeInteractiveToolsContext.h"
 #include "EditorModeManager.h"
+#include "Elements/Framework/TypedElementSelectionSet.h"
 #include "HAL/IConsoleManager.h"
+#include "LevelEditor.h"
 #include "InputRouter.h"
 #include "InteractiveGizmo.h"
 #include "InteractiveToolChange.h"
 #include "InteractiveToolsContext.h"
+#include "Selection.h"
 #include "ToolContextInterfaces.h"
 
 
 #define LOCTEXT_NAMESPACE "UEditorInteractiveGizmoManager"
 
-#if 0
 static TAutoConsoleVariable<int32> CVarUseLegacyWidget(
 	TEXT("Gizmos.UseLegacyWidget"),
 	1,
@@ -24,7 +27,6 @@ static TAutoConsoleVariable<int32> CVarUseLegacyWidget(
 	TEXT("0 = enable UE5 transform and other selection-based gizmos.\n")
 	TEXT("1 = enable legacy UE4 transform widget."),
 	ECVF_RenderThreadSafe);
-#endif
 
 
 UEditorInteractiveGizmoManager::UEditorInteractiveGizmoManager() :
@@ -38,7 +40,6 @@ void UEditorInteractiveGizmoManager::InitializeWithEditorModeManager(IToolsConte
 {
 	Super::Initialize(QueriesAPIIn, TransactionsAPIIn, InputRouterIn);
 	EditorModeManager = InEditorModeManager;
-
 }
 
 
@@ -132,50 +133,41 @@ TArray<TObjectPtr<UEditorInteractiveGizmoSelectionBuilder>> UEditorInteractiveGi
 	return FoundBuilders;
 }
 
-TArray<UInteractiveGizmo*> UEditorInteractiveGizmoManager::CreateSelectionGizmos(void* Owner)
+TArray<UInteractiveGizmo*> UEditorInteractiveGizmoManager::CreateGizmosForCurrentSelectionState()
 {
-	// always destroy the previous active auto gizmo
-	DestroyAllSelectionGizmos();
-
 	if (bShowSelectionGizmos)
 	{
+		TArray<UInteractiveGizmo*> NewGizmos;
+
 		FToolBuilderState CurrentSceneState;
 		QueriesAPI->GetCurrentSelectionState(CurrentSceneState);
 
-		if (UTypedElementSelectionSet* SelectionSet = CurrentSceneState.TypedElementSelectionSet.Get())
+		TArray<TObjectPtr<UEditorInteractiveGizmoSelectionBuilder>> FoundBuilders = GetQualifiedGizmoSelectionBuilders(CurrentSceneState);
+
+		for (TObjectPtr<UEditorInteractiveGizmoSelectionBuilder> FoundBuilder : FoundBuilders)
 		{
-			if (SelectionSet->HasSelectedElements())
+			UInteractiveGizmo* NewGizmo = FoundBuilder->BuildGizmoForSelection(CurrentSceneState);
+			if (NewGizmo == nullptr)
 			{
-				TArray<UInteractiveGizmo*> NewGizmos;
-
-				TArray<TObjectPtr<UEditorInteractiveGizmoSelectionBuilder>> FoundBuilders = GetQualifiedGizmoSelectionBuilders(CurrentSceneState);
-
-				for (TObjectPtr<UEditorInteractiveGizmoSelectionBuilder> FoundBuilder : FoundBuilders)
-				{
-					UInteractiveGizmo* NewGizmo = FoundBuilder->BuildGizmo(CurrentSceneState);
-					if (NewGizmo == nullptr)
-					{
-						DisplayMessage(LOCTEXT("CreateGizmoReturnNullMessage", "UEditorInteractiveGizmoManager::CreateGizmo: BuildGizmo() returned null"), EToolMessageLevel::Internal);
-						return NewGizmos;
-					}
-
-					// register new active input behaviors
-					InputRouter->RegisterSource(NewGizmo);
+				DisplayMessage(LOCTEXT("CreateGizmoReturnNullMessage", "UEditorInteractiveGizmoManager::CreateGizmo: CreateGizmoForSelectionSet() returned null"), EToolMessageLevel::Internal);
+				return NewGizmos;
+			}
 
-					NewGizmos.Add(NewGizmo);
-				}
+			// register new active input behaviors
+			InputRouter->RegisterSource(NewGizmo);
 
-				PostInvalidation();
+			NewGizmos.Add(NewGizmo);
+		}
 
-				for (UInteractiveGizmo* NewGizmo : NewGizmos)
-				{
-					FActiveSelectionGizmo ActiveGizmo = { NewGizmo, Owner };
-					ActiveSelectionGizmos.Add(ActiveGizmo);
-				}
+		PostInvalidation();
 
-				return NewGizmos;
-			}
+		for (UInteractiveGizmo* NewGizmo : NewGizmos)
+		{
+			FActiveSelectionGizmo ActiveGizmo = { NewGizmo, nullptr };
+			ActiveSelectionGizmos.Add(ActiveGizmo);
 		}
+
+		return NewGizmos;
 	}
 
 	return TArray<UInteractiveGizmo*>();
@@ -219,15 +211,14 @@ void UEditorInteractiveGizmoManager::DestroyAllSelectionGizmos()
 	PostInvalidation();
 }
 
-
-void UEditorInteractiveGizmoManager::OnEditorSelectionChanged()
-{
-	CreateSelectionGizmos();
-}
-
-void UEditorInteractiveGizmoManager::OnEditorSelectNone()
+void UEditorInteractiveGizmoManager::HandleEditorSelectionSetChanged(const UTypedElementSelectionSet* InSelectionSet)
 {
 	DestroyAllSelectionGizmos();
+
+	if (InSelectionSet && InSelectionSet->HasSelectedElements())
+	{
+		CreateGizmosForCurrentSelectionState();
+	}
 }
 
 // @todo move this to a gizmo context object
@@ -246,20 +237,16 @@ bool UEditorInteractiveGizmoManager::GetShowSelectionGizmosForView(IToolsContext
 
 void UEditorInteractiveGizmoManager::UpdateActiveSelectionGizmos()
 {
-#if 0
 	const bool bEditorModeToolsSupportsWidgetDrawing = EditorModeManager ? EditorModeManager->GetShowWidget() : true;
 	const bool bEnableSelectionGizmos = (CVarUseLegacyWidget.GetValueOnGameThread() == 0);
 	const bool bNewShowSelectionGizmos = bEditorModeToolsSupportsWidgetDrawing && bEnableSelectionGizmos;
-#else
-	const bool bNewShowSelectionGizmos = EditorModeManager ? EditorModeManager->GetShowWidget() : true;
-#endif
 
 	if (bShowSelectionGizmos != bNewShowSelectionGizmos)
 	{
 		bShowSelectionGizmos = bNewShowSelectionGizmos;
 		if (bShowSelectionGizmos)
 		{
-			CreateSelectionGizmos();
+			CreateGizmosForCurrentSelectionState();
 		}
 		else
 		{
diff --git a/Engine/Source/Editor/Experimental/EditorInteractiveToolsFramework/Private/EditorInteractiveGizmoSubsystem.cpp b/Engine/Source/Editor/Experimental/EditorInteractiveToolsFramework/Private/EditorInteractiveGizmoSubsystem.cpp
--- a/Engine/Source/Editor/Experimental/EditorInteractiveToolsFramework/Private/EditorInteractiveGizmoSubsystem.cpp
+++ b/Engine/Source/Editor/Experimental/EditorInteractiveToolsFramework/Private/EditorInteractiveGizmoSubsystem.cpp
@@ -34,11 +34,9 @@ void UEditorInteractiveGizmoSubsystem::Deinitialize()
 
 void UEditorInteractiveGizmoSubsystem::RegisterBuiltinGizmoSelectionTypes()
 {
-#if 0
 	// Register built-in gizmo types here
 	TObjectPtr<UEditorTransformGizmoBuilder> EditorTransformBuilder = NewObject<UEditorTransformGizmoBuilder>();
 	RegisterGizmoSelectionType(EditorTransformBuilder);
-#endif
 
 	RegisterEditorGizmoSelectionTypesDelegate.Broadcast();
 }
diff --git a/Engine/Source/Editor/Experimental/EditorInteractiveToolsFramework/Public/EdModeInteractiveToolsContext.h b/Engine/Source/Editor/Experimental/EditorInteractiveToolsFramework/Public/EdModeInteractiveToolsContext.h
--- a/Engine/Source/Editor/Experimental/EditorInteractiveToolsFramework/Public/EdModeInteractiveToolsContext.h
+++ b/Engine/Source/Editor/Experimental/EditorInteractiveToolsFramework/Public/EdModeInteractiveToolsContext.h
@@ -19,6 +19,9 @@ class FViewport;
 class UMaterialInterface;
 class FPrimitiveDrawInterface;
 class FViewportClient;
+class ILevelEditor;
+class USelection;
+class UTypedElementSelectionSet;
 
 /**
  * EdModeInteractiveToolsContext is an extension/adapter of an InteractiveToolsContext which 
@@ -102,6 +105,19 @@ public:
 	/** @return true if HitProxy rendering will be allowed in ::Render() */
 	bool GetEnableRenderingDuringHitProxyPass() const { return bEnableRenderingDuringHitProxyPass; }
 
+	/**
+	 * Called upon Level Editor Created in order to bind to EditorElementSelectionPtr changed
+	 * event since global selection set is not initialized before the tools context.
+	 * @param InSelectionSet - typed element selection set which invoked this selection changed call
+	 */
+	void OnLevelEditorCreated(TSharedPtr<ILevelEditor> InLevelEditor);
+
+	/**
+	 * Handle Editor selection changes
+	 * @param InSelectionSet - typed element selection set which invoked this selection changed call
+	 */
+	void OnEditorSelectionSetChanged(const UTypedElementSelectionSet* InSelectionSet);
+
 protected:
 	// we hide these 
 	virtual void Initialize(IToolsContextQueriesAPI* QueriesAPI, IToolsContextTransactionsAPI* TransactionsAPI) override;
diff --git a/Engine/Source/Editor/Experimental/EditorInteractiveToolsFramework/Public/EditorGizmos/EditorTransformGizmo.h b/Engine/Source/Editor/Experimental/EditorInteractiveToolsFramework/Public/EditorGizmos/EditorTransformGizmo.h
--- a/Engine/Source/Editor/Experimental/EditorInteractiveToolsFramework/Public/EditorGizmos/EditorTransformGizmo.h
+++ b/Engine/Source/Editor/Experimental/EditorInteractiveToolsFramework/Public/EditorGizmos/EditorTransformGizmo.h
@@ -32,9 +32,15 @@ class EDITORINTERACTIVETOOLSFRAMEWORK_API UEditorTransformGizmoBuilder : public
 
 public:
 
+	// UEditorInteractiveGizmoSelectionBuilder interface 
 	virtual UInteractiveGizmo* BuildGizmo(const FToolBuilderState& SceneState) const override;
+	virtual void UpdateGizmoForSelection(UInteractiveGizmo* Gizmo, const FToolBuilderState& SceneState) override;
 
-	virtual bool SatisfiesCondition(const FToolBuilderState& SceneState) const override;
+	// UEditorInteractiveGizmoConditionalBuilder interface 
+	virtual bool SatisfiesCondition(const FToolBuilderState& SceneState) const override
+	{
+		return true;
+	}
 };
 
 
diff --git a/Engine/Source/Editor/Experimental/EditorInteractiveToolsFramework/Public/EditorInteractiveGizmoConditionalBuilder.h b/Engine/Source/Editor/Experimental/EditorInteractiveToolsFramework/Public/EditorInteractiveGizmoConditionalBuilder.h
--- a/Engine/Source/Editor/Experimental/EditorInteractiveToolsFramework/Public/EditorInteractiveGizmoConditionalBuilder.h
+++ b/Engine/Source/Editor/Experimental/EditorInteractiveToolsFramework/Public/EditorInteractiveGizmoConditionalBuilder.h
@@ -52,7 +52,13 @@ struct EDITORINTERACTIVETOOLSFRAMEWORK_API FEditorGizmoTypePriority
 	}
 };
 
-/** UInteractiveGizmoType provides actions and other information about gizmo types */
+/** 
+ * Gizmo builders which should be built once a condition is satisfied in the current scene state
+ * inherit from UEditorInteractiveGizmoConditionalBuilder or its derived classes. The SatisfiesCondition method
+ * should return true whenever the builder is buildable based on Editor scene state, usually based
+ * on selection. The builder's priority should be specified such that when more than one builder
+ * is discovered, the builder with highest priority will be built.
+ */
 UCLASS(Transient, Abstract)
 class EDITORINTERACTIVETOOLSFRAMEWORK_API UEditorInteractiveGizmoConditionalBuilder : public UInteractiveGizmoBuilder
 {
@@ -73,11 +79,7 @@ public:
 	}
 
 	/** Returns true if this gizmo is valid for creation based on the current state. */
-	virtual bool SatisfiesCondition(const FToolBuilderState& SceneState) const
-	{
-		unimplemented();
-		return false;
-	}
+	virtual bool SatisfiesCondition(const FToolBuilderState& SceneState) const  PURE_VIRTUAL(UEditorInteractiveGizmoConditionalBuilder::SatisfiesCondition, return false; );
 
 protected:
 
diff --git a/Engine/Source/Editor/Experimental/EditorInteractiveToolsFramework/Public/EditorInteractiveGizmoManager.h b/Engine/Source/Editor/Experimental/EditorInteractiveToolsFramework/Public/EditorInteractiveGizmoManager.h
--- a/Engine/Source/Editor/Experimental/EditorInteractiveToolsFramework/Public/EditorInteractiveGizmoManager.h
+++ b/Engine/Source/Editor/Experimental/EditorInteractiveToolsFramework/Public/EditorInteractiveGizmoManager.h
@@ -9,8 +9,9 @@
 class FEditorModeTools;
 class UEditorInteractiveGizmoSelectionBuilder;
 class UEdModeInteractiveToolsContext;
-class UinteractiveGizmo;
+class UInteractiveGizmo;
 class IToolsContextRenderAPI;
+class UTypedElementSelectionSet;
 
 USTRUCT()
 struct FActiveSelectionGizmo
@@ -92,22 +93,15 @@ public:
 
 	/**
 	 * Try to automatically activate a new Gizmo instance based on the current state
-	 * @param Owner void pointer to whatever "owns" this Gizmo. Allows Gizmo to later be deleted using DestroyAllGizmosByOwner()
 	 * @return array of new Gizmo instances that have been created and initialized
 	 */
-	virtual TArray<UInteractiveGizmo*> CreateSelectionGizmos(void* Owner = nullptr);
+	virtual TArray<UInteractiveGizmo*> CreateGizmosForCurrentSelectionState();
 
 	/**
 	 * Handle Editor selection changes
-	 * @param Tools - Mode Manager which invoked this selection changed call
-	 * @param NewSelection - Object which is undergoing selection change
+	 * @param InSelectionSet - typed element selection set which invoked this selection changed call
 	 */
-	virtual void OnEditorSelectionChanged();
-
-	/**
-	 * Handle case when selection has been cleared.
-	 */
-	virtual void OnEditorSelectNone();
+	void HandleEditorSelectionSetChanged(const UTypedElementSelectionSet* InSelectionSet);
 
 	/**
 	 * Shutdown and remove a selection-based Gizmo
@@ -116,7 +110,7 @@ public:
 	 */
 	virtual bool DestroySelectionGizmo(UInteractiveGizmo* Gizmo);
 
-		/**
+	/**
 	 * Shutdown and remove all active auto gizmos
 	 */
 	virtual void DestroyAllSelectionGizmos();
diff --git a/Engine/Source/Editor/Experimental/EditorInteractiveToolsFramework/Public/EditorInteractiveGizmoSelectionBuilder.h b/Engine/Source/Editor/Experimental/EditorInteractiveToolsFramework/Public/EditorInteractiveGizmoSelectionBuilder.h
--- a/Engine/Source/Editor/Experimental/EditorInteractiveToolsFramework/Public/EditorInteractiveGizmoSelectionBuilder.h
+++ b/Engine/Source/Editor/Experimental/EditorInteractiveToolsFramework/Public/EditorInteractiveGizmoSelectionBuilder.h
@@ -4,14 +4,15 @@
 
 #include "CoreMinimal.h"
 #include "EditorInteractiveGizmoConditionalBuilder.h"
+#include "InteractiveGizmo.h"
 #include "ToolContextInterfaces.h"
 #include "EditorInteractiveGizmoSelectionBuilder.generated.h"
 
-
 /** 
- * UEditorInteractiveGizmoSelectionBuilder provides a method for checking that the current selection and widget mode satisfy 
- * the conditions of this builder. Builders derived from this class should be registered in the gizmo subsystem, for gizmos
- * available globally in the Editor, or in the gizmo manager for gizmos only relevant to a particle ed mode or asset editor.
+ * UEditorInteractiveGizmoSelectionBuilder provides a method for building and updating gizmos based on the current Editor selection 
+ * and state. Builders derived from this class may be registered in one of the following places: 
+ *   1) the gizmo subsystem if the gizmo should be available throughout the Editor.
+ *   2) the gizmo manager if the gizmo is only used in a particular ed mode or in an asset editor.
  */
 UCLASS(Transient, Abstract)
 class EDITORINTERACTIVETOOLSFRAMEWORK_API UEditorInteractiveGizmoSelectionBuilder : public UEditorInteractiveGizmoConditionalBuilder
@@ -19,10 +20,27 @@ class EDITORINTERACTIVETOOLSFRAMEWORK_API UEditorInteractiveGizmoSelectionBuilde
 	GENERATED_BODY()
 
 public:
-
-	/** Returns true if this gizmo is valid for creation based on the current state. */
-	virtual bool SatisfiesCondition(const FToolBuilderState& SceneState) const override
+	/**
+	 * Build gizmo for the current Editor selection and state. The Editor gizmo manager calls this method to construct gizmos for 
+	 * for the current selection. This implementation calls BuildGizmo() then UpdateGizmoForSelection(). But derived classes may provide 
+	 * their own implementation which is expected to both build gizmo and set it up to manipulate the current selection. Note that 
+	 * when the selection changes, the gizmo manager may reuse the gizmo and only call UpdateGizmoForSelection() on the existing gizmo.
+	 */
+	virtual UInteractiveGizmo* BuildGizmoForSelection(const FToolBuilderState& SceneState)
 	{
-		return false;
+		if (UInteractiveGizmo* Gizmo = BuildGizmo(SceneState))
+		{
+			UpdateGizmoForSelection(Gizmo, SceneState);
+			return Gizmo;
+		}
+		return nullptr;
 	}
+
+	/**
+	 * Update the input gizmo's active target based on the current Editor selection and scene state.  Derived implementations
+	 * of this method should create a transform proxy for the current Editor selection and sets the gizmo's active target to the new transform proxy. 
+	 * The gizmo manager calls this method when reusing a gizmo, to update the gizmo for the current selection.
+	 */
+	virtual void UpdateGizmoForSelection(UInteractiveGizmo* Gizmo, const FToolBuilderState& SceneState) PURE_VIRTUAL(UEditorInteractiveGizmoSelectionBuilder::UpdateGizmoForSelection, return;);
+
 };
diff --git a/Engine/Source/Editor/UnrealEd/Private/UnrealWidgetRender.cpp b/Engine/Source/Editor/UnrealEd/Private/UnrealWidgetRender.cpp
--- a/Engine/Source/Editor/UnrealEd/Private/UnrealWidgetRender.cpp
+++ b/Engine/Source/Editor/UnrealEd/Private/UnrealWidgetRender.cpp
@@ -180,13 +180,17 @@ void FWidget::Render( const FSceneView* View,FPrimitiveDrawInterface* PDI, FEdit
 
 	const bool bShowFlagsSupportsWidgetDrawing = View->Family->EngineShowFlags.ModeWidgets;
 	const bool bEditorModeToolsSupportsWidgetDrawing = EditorModeTools ? EditorModeTools->GetShowWidget() : true;
+
+	static const auto UseLegacyWidgetCVar = IConsoleManager::Get().FindConsoleVariable(TEXT("Gizmos.UseLegacyWidget"));
+	const bool bUseLegacyWidget = UseLegacyWidgetCVar->GetInt() > 0;
+
 	bool bDrawWidget;
 
 	// Because the movement routines use the widget axis to determine how to transform mouse movement into
 	// editor object movement, we need to still run through the Render routine even though widget drawing may be
 	// disabled.  So we keep a flag that is used to determine whether or not to actually render anything.  This way
 	// we can still update the widget axis' based on the Context's transform matrices, even though drawing is disabled.
-	if(bDefaultVisibility && bShowFlagsSupportsWidgetDrawing && bEditorModeToolsSupportsWidgetDrawing)
+	if(bDefaultVisibility && bShowFlagsSupportsWidgetDrawing && bEditorModeToolsSupportsWidgetDrawing && bUseLegacyWidget)
 	{
 		bDrawWidget = true;
 
