Switch to Editor Grid/Rotation snapping settings and enable/disable state for all usage of TransformGizmo in ModelingTools
- TransformGizmo now defaults to Grid Position/Rotation Snapping enabled by default
- remove shift-modifier that toggled on 45-degree snapping in UAxisAngleGizmo
- update most Tools and Mechanics to remove explicit snapping properties/checkboxes and shift modifiers (PolyEdit and PolyDeform Tool still need to be updated)
- update FEdModeToolsContextQueriesImpl::ExecuteSceneSnapQueryPosition() and FEdModeToolsContextQueriesImpl::ExecuteSceneSnapQueryRotation() implementations to support snapping enable/disable via ULevelEditorViewportSettings
#rnx

diff --git a/Engine/Plugins/Experimental/MeshModelingToolset/Source/MeshModelingTools/Private/DrawAndRevolveTool.cpp b/Engine/Plugins/Experimental/MeshModelingToolset/Source/MeshModelingTools/Private/DrawAndRevolveTool.cpp
--- a/Engine/Plugins/Experimental/MeshModelingToolset/Source/MeshModelingTools/Private/DrawAndRevolveTool.cpp
+++ b/Engine/Plugins/Experimental/MeshModelingToolset/Source/MeshModelingTools/Private/DrawAndRevolveTool.cpp
@@ -188,7 +188,6 @@ void UDrawAndRevolveTool::Setup()
 		}
 		UpdateRevolutionAxis();
 		});
-	PlaneMechanic->SetEnableGridSnaping(Settings->bSnapToWorldGrid);
 
 	ControlPointsMechanic->SetPlane(PlaneMechanic->Plane);
 	ControlPointsMechanic->SetInteractiveInitialization(true);
@@ -294,7 +293,6 @@ void UDrawAndRevolveTool::OnPropertyModified(UObject* PropertySet, FProperty* Pr
 	PlaneMechanic->SetPlaneWithoutBroadcast(ProfileDrawPlane);
 	UpdateRevolutionAxis();
 
-	PlaneMechanic->SetEnableGridSnaping(Settings->bSnapToWorldGrid);
 	ControlPointsMechanic->SetSnappingEnabled(Settings->bEnableSnapping);
 
 	if (Preview)
diff --git a/Engine/Plugins/Experimental/MeshModelingToolset/Source/MeshModelingTools/Private/DrawPolyPathTool.cpp b/Engine/Plugins/Experimental/MeshModelingToolset/Source/MeshModelingTools/Private/DrawPolyPathTool.cpp
--- a/Engine/Plugins/Experimental/MeshModelingToolset/Source/MeshModelingTools/Private/DrawPolyPathTool.cpp
+++ b/Engine/Plugins/Experimental/MeshModelingToolset/Source/MeshModelingTools/Private/DrawPolyPathTool.cpp
@@ -381,7 +381,6 @@ void UDrawPolyPathTool::OnTick(float DeltaTime)
 {
 	if (PlaneMechanic != nullptr)
 	{
-		PlaneMechanic->SetEnableGridSnaping(TransformProps->bSnapToWorldGrid);
 		PlaneMechanic->Tick(DeltaTime);
 	}
 }
diff --git a/Engine/Plugins/Experimental/MeshModelingToolset/Source/MeshModelingTools/Private/DrawPolygonTool.cpp b/Engine/Plugins/Experimental/MeshModelingToolset/Source/MeshModelingTools/Private/DrawPolygonTool.cpp
--- a/Engine/Plugins/Experimental/MeshModelingToolset/Source/MeshModelingTools/Private/DrawPolygonTool.cpp
+++ b/Engine/Plugins/Experimental/MeshModelingToolset/Source/MeshModelingTools/Private/DrawPolygonTool.cpp
@@ -236,11 +236,6 @@ void UDrawPolygonTool::ApplyUndoPoints(const TArray<FVector3d>& ClickPointsIn, c
 
 void UDrawPolygonTool::OnTick(float DeltaTime)
 {
-	if (PlaneTransformGizmo != nullptr)
-	{
-		PlaneTransformGizmo->bSnapToWorldGrid =
-			SnapProperties->bEnableSnapping && SnapProperties->bSnapToWorldGrid && bIgnoreSnappingToggle == false;
-	}
 }
 
 
diff --git a/Engine/Plugins/Experimental/MeshModelingToolset/Source/MeshModelingTools/Private/DynamicMeshSculptTool.cpp b/Engine/Plugins/Experimental/MeshModelingToolset/Source/MeshModelingTools/Private/DynamicMeshSculptTool.cpp
--- a/Engine/Plugins/Experimental/MeshModelingToolset/Source/MeshModelingTools/Private/DynamicMeshSculptTool.cpp
+++ b/Engine/Plugins/Experimental/MeshModelingToolset/Source/MeshModelingTools/Private/DynamicMeshSculptTool.cpp
@@ -2674,8 +2674,6 @@ void UDynamicMeshSculptTool::UpdateFixedPlaneGizmoVisibility(bool bVisible)
 			PlaneTransformGizmo->SetActiveTarget(PlaneTransformProxy, GetToolManager());
 			PlaneTransformGizmo->ReinitializeGizmoTransform(FTransform(GizmoProperties->Rotation, GizmoProperties->Position));
 		}
-
-		PlaneTransformGizmo->bSnapToWorldGrid = GizmoProperties->bSnapToGrid;
 	}
 }
 
diff --git a/Engine/Plugins/Experimental/MeshModelingToolset/Source/MeshModelingTools/Private/MeshSpaceDeformerTool.cpp b/Engine/Plugins/Experimental/MeshModelingToolset/Source/MeshModelingTools/Private/MeshSpaceDeformerTool.cpp
--- a/Engine/Plugins/Experimental/MeshModelingToolset/Source/MeshModelingTools/Private/MeshSpaceDeformerTool.cpp
+++ b/Engine/Plugins/Experimental/MeshModelingToolset/Source/MeshModelingTools/Private/MeshSpaceDeformerTool.cpp
@@ -202,7 +202,6 @@ void UMeshSpaceDeformerTool::Setup()
 	TransformProxy->SetTransform(GizmoFrame.ToFTransform());
 	TransformGizmo = UE::TransformGizmoUtil::CreateCustomTransformGizmo(GetToolManager(),
 		ETransformGizmoSubElements::StandardTranslateRotate, this);
-	TransformGizmo->bSnapToWorldGrid = Settings->bSnapToWorldGrid;
 
 	TransformGizmo->SetActiveTarget(TransformProxy, GetToolManager());
 
@@ -359,8 +358,6 @@ void  UMeshSpaceDeformerTool::OnPropertyModified(UObject* PropertySet, FProperty
 
 	UpdatePreview();
 
-	TransformGizmo->bSnapToWorldGrid = Settings->bSnapToWorldGrid;
-
 	OriginalMeshPreview->SetVisible(Settings->bShowOriginalMesh);
 }
 
diff --git a/Engine/Plugins/Experimental/MeshModelingToolset/Source/MeshModelingTools/Private/MirrorTool.cpp b/Engine/Plugins/Experimental/MeshModelingToolset/Source/MeshModelingTools/Private/MirrorTool.cpp
--- a/Engine/Plugins/Experimental/MeshModelingToolset/Source/MeshModelingTools/Private/MirrorTool.cpp
+++ b/Engine/Plugins/Experimental/MeshModelingToolset/Source/MeshModelingTools/Private/MirrorTool.cpp
@@ -163,9 +163,6 @@ void UMirrorTool::OnTick(float DeltaTime)
 
 	if (PlaneMechanic != nullptr)
 	{
-		// Update snapping behavior based on modifier key.
-		PlaneMechanic->SetEnableGridSnaping(Settings->bSnapToWorldGrid ^ bSnappingToggle);
-
 		PlaneMechanic->Tick(DeltaTime);
 	}
 	
@@ -187,7 +184,7 @@ void UMirrorTool::Setup()
 
 	SetToolDisplayName(LOCTEXT("ToolName", "Mirror"));
 	GetToolManager()->DisplayMessage(
-		LOCTEXT("OnStartMirrorTool", "Mirror one or more meshes across a plane. Grid snapping behavior is swapped while the shift key is down. The plane can be set by using the preset buttons, moving the gizmo, or ctrl+clicking on a spot on the original mesh."),
+		LOCTEXT("OnStartMirrorTool", "Mirror one or more meshes across a plane. The plane can be set by using the preset buttons, moving the gizmo, or ctrl+clicking on a spot on the original mesh."),
 		EToolMessageLevel::UserNotification);
 
 	// Set up the properties
@@ -276,12 +273,6 @@ void UMirrorTool::Setup()
 		PlaneMechanic->SetPlaneCtrlClickBehaviorTarget->InvisibleComponentsToHitTest.Add(TargetComponent->GetOwnerComponent());
 	}
 
-
-	// Add modifier button for snapping
-	UKeyAsModifierInputBehavior* SnapToggleBehavior = NewObject<UKeyAsModifierInputBehavior>();
-	SnapToggleBehavior->Initialize(this, SnappingToggleModifierId, FInputDeviceState::IsShiftKeyDown);
-	AddInputBehavior(SnapToggleBehavior);
-
 	// Start the preview calculations
 	for (UMeshOpPreviewWithBackgroundCompute* Preview : Previews)
 	{
@@ -562,10 +553,6 @@ void UMirrorTool::ApplyAction(EMirrorToolAction ActionType)
 
 void UMirrorTool::OnUpdateModifierState(int ModifierID, bool bIsOn)
 {
-	if (ModifierID == SnappingToggleModifierId)
-	{
-		bSnappingToggle = bIsOn;
-	}
 }
 
 #undef LOCTEXT_NAMESPACE
diff --git a/Engine/Plugins/Experimental/MeshModelingToolset/Source/MeshModelingTools/Private/PlaneCutTool.cpp b/Engine/Plugins/Experimental/MeshModelingToolset/Source/MeshModelingTools/Private/PlaneCutTool.cpp
--- a/Engine/Plugins/Experimental/MeshModelingToolset/Source/MeshModelingTools/Private/PlaneCutTool.cpp
+++ b/Engine/Plugins/Experimental/MeshModelingToolset/Source/MeshModelingTools/Private/PlaneCutTool.cpp
@@ -9,8 +9,6 @@
 #include "DynamicMesh/DynamicMesh3.h"
 #include "DynamicMesh/DynamicMeshTriangleAttribute.h"
 #include "DynamicMeshEditor.h"
-#include "BaseBehaviors/MultiClickSequenceInputBehavior.h"
-#include "BaseBehaviors/KeyAsModifierInputBehavior.h"
 #include "Selection/SelectClickedAction.h"
 
 #include "MeshDescriptionToDynamicMesh.h"
@@ -92,11 +90,6 @@ void UPlaneCutTool::Setup()
 {
 	UInteractiveTool::Setup();
 
-	// add modifier button for snapping
-	UKeyAsModifierInputBehavior* SnapToggleBehavior = NewObject<UKeyAsModifierInputBehavior>();
-	SnapToggleBehavior->Initialize(this, SnappingModifier, FInputDeviceState::IsShiftKeyDown);
-	AddInputBehavior(SnapToggleBehavior);
-
 	// hide input StaticMeshComponents
 	for (int32 ComponentIdx = 0; ComponentIdx < Targets.Num(); ComponentIdx++)
 	{
@@ -180,7 +173,7 @@ void UPlaneCutTool::Setup()
 
 	SetToolDisplayName(LOCTEXT("ToolName", "Plane Cut"));
 	GetToolManager()->DisplayMessage(
-		LOCTEXT("OnStartPlaneCutTool", "Press 'T' or use the Cut button to cut the mesh without leaving the tool.  Press 'R' to flip the plane direction.  Hold 'Shift' to toggle grid snapping."),
+		LOCTEXT("OnStartPlaneCutTool", "Press 'T' or use the Cut button to cut the mesh without leaving the tool.  Press 'R' to flip the plane direction."),
 		EToolMessageLevel::UserNotification);
 }
 
@@ -353,8 +346,6 @@ void UPlaneCutTool::Render(IToolsContextRenderAPI* RenderAPI)
 void UPlaneCutTool::OnTick(float DeltaTime)
 {
 	PlaneMechanic->Tick(DeltaTime);
-	PlaneMechanic->SetEnableGridSnaping(BasicProperties->bSnapToWorldGrid ^ bSnappingToggle);
-	PlaneMechanic->SetEnableGridRotationSnapping(BasicProperties->bSnapRotationToWorldGrid ^ bSnappingToggle);
 
 	if (PendingAction != EPlaneCutToolActions::NoAction)
 	{
@@ -403,13 +394,6 @@ void UPlaneCutTool::OnPropertyModified(UObject* PropertySet, FProperty* Property
 }
 
 
-void UPlaneCutTool::OnUpdateModifierState(int ModifierID, bool bIsOn)
-{
-	if (ModifierID == SnappingModifier)
-	{
-		bSnappingToggle = bIsOn;
-	}
-}
 
 
 void UPlaneCutTool::InvalidatePreviews()
diff --git a/Engine/Plugins/Experimental/MeshModelingToolset/Source/MeshModelingTools/Private/RevolveBoundaryTool.cpp b/Engine/Plugins/Experimental/MeshModelingToolset/Source/MeshModelingTools/Private/RevolveBoundaryTool.cpp
--- a/Engine/Plugins/Experimental/MeshModelingToolset/Source/MeshModelingTools/Private/RevolveBoundaryTool.cpp
+++ b/Engine/Plugins/Experimental/MeshModelingToolset/Source/MeshModelingTools/Private/RevolveBoundaryTool.cpp
@@ -109,8 +109,6 @@ void URevolveBoundaryTool::Setup()
 		UpdateRevolutionAxis();
 		});
 
-	PlaneMechanic->SetEnableGridSnaping(Settings->bSnapToWorldGrid);
-
 	Cast<IPrimitiveComponentBackedTarget>(Target)->SetOwnerVisibility(Settings->bDisplayOriginalMesh);
 
 	SetToolDisplayName(LOCTEXT("ToolName", "Revolve Boundary"));
@@ -343,7 +341,6 @@ void URevolveBoundaryTool::OnPropertyModified(UObject* PropertySet, FProperty* P
 	UpdateRevolutionAxis();
 
 	Cast<IPrimitiveComponentBackedTarget>(Target)->SetOwnerVisibility(Settings->bDisplayOriginalMesh);
-	PlaneMechanic->SetEnableGridSnaping(Settings->bSnapToWorldGrid);
 
 	if (Preview)
 	{
diff --git a/Engine/Plugins/Experimental/MeshModelingToolset/Source/MeshModelingTools/Private/Sculpting/MeshSculptToolBase.cpp b/Engine/Plugins/Experimental/MeshModelingToolset/Source/MeshModelingTools/Private/Sculpting/MeshSculptToolBase.cpp
--- a/Engine/Plugins/Experimental/MeshModelingToolset/Source/MeshModelingTools/Private/Sculpting/MeshSculptToolBase.cpp
+++ b/Engine/Plugins/Experimental/MeshModelingToolset/Source/MeshModelingTools/Private/Sculpting/MeshSculptToolBase.cpp
@@ -1210,8 +1210,6 @@ void UMeshSculptToolBase::UpdateFixedPlaneGizmoVisibility(bool bVisible)
 			PlaneTransformGizmo->SetActiveTarget(PlaneTransformProxy, GetToolManager());
 			PlaneTransformGizmo->ReinitializeGizmoTransform(FTransform(GizmoProperties->Rotation, GizmoProperties->Position));
 		}
-
-		PlaneTransformGizmo->bSnapToWorldGrid = GizmoProperties->bSnapToGrid;
 	}
 }
 
diff --git a/Engine/Plugins/Experimental/MeshModelingToolset/Source/MeshModelingTools/Private/TransformMeshesTool.cpp b/Engine/Plugins/Experimental/MeshModelingToolset/Source/MeshModelingTools/Private/TransformMeshesTool.cpp
--- a/Engine/Plugins/Experimental/MeshModelingToolset/Source/MeshModelingTools/Private/TransformMeshesTool.cpp
+++ b/Engine/Plugins/Experimental/MeshModelingToolset/Source/MeshModelingTools/Private/TransformMeshesTool.cpp
@@ -92,8 +92,6 @@ void UTransformMeshesTool::Setup()
 	TransformProps->WatchProperty(TransformProps->TransformMode, [this](ETransformMeshesTransformMode NewMode) { UpdateTransformMode(NewMode); });
 	TransformProps->WatchProperty(TransformProps->bApplyToInstances, [this](bool bNewValue) { UpdateTransformMode(TransformProps->TransformMode); });
 	TransformProps->WatchProperty(TransformProps->bSetPivotMode, [this](bool bNewValue) { UpdateSetPivotModes(bNewValue); });
-	TransformProps->WatchProperty(TransformProps->bSnapToGrid, [this](bool bNewValue) { OnSnappingSettingsUpdated(); });
-	TransformProps->WatchProperty(TransformProps->bSnapToRotation, [this](bool bNewValue) { OnSnappingSettingsUpdated(); });
 	
 	// determine if we have any ISMCs
 	TransformProps->bHaveInstances = false;
@@ -213,8 +211,6 @@ void UTransformMeshesTool::UpdateTransformMode(ETransformMeshesTransformMode New
 			break;
 	}
 
-	OnSnappingSettingsUpdated();
-
 	CurTransformMode = NewMode;
 }
 
@@ -502,14 +498,5 @@ void UTransformMeshesTool::OnTerminateDragSequence()
 }
 
 
-void UTransformMeshesTool::OnSnappingSettingsUpdated()
-{
-	for (FTransformMeshesTarget& Target : ActiveGizmos)
-	{
-		Target.TransformGizmo->bSnapToWorldGrid = TransformProps->bSnapToGrid;
-		Target.TransformGizmo->bSnapToWorldRotGrid = TransformProps->bSnapToRotation;
-	}
-}
-
 
 #undef LOCTEXT_NAMESPACE
diff --git a/Engine/Plugins/Experimental/MeshModelingToolset/Source/MeshModelingTools/Public/DrawAndRevolveTool.h b/Engine/Plugins/Experimental/MeshModelingToolset/Source/MeshModelingTools/Public/DrawAndRevolveTool.h
--- a/Engine/Plugins/Experimental/MeshModelingToolset/Source/MeshModelingTools/Public/DrawAndRevolveTool.h
+++ b/Engine/Plugins/Experimental/MeshModelingToolset/Source/MeshModelingTools/Public/DrawAndRevolveTool.h
@@ -42,10 +42,6 @@ public:
 	UPROPERTY(EditAnywhere, Category = RevolveSettings, AdvancedDisplay)
 	bool bConnectOpenProfileToAxis = true;
 	
-	/** Determines whether plane control widget snaps to world grid (only relevant if world coordinate mode is active in viewport) .*/
-	UPROPERTY(EditAnywhere, Category = DrawPlane)
-	bool bSnapToWorldGrid = false;
-
 	/** Sets the draw plane origin. The revolution axis is the X axis in the plane. */
 	UPROPERTY(EditAnywhere, Category = DrawPlane, meta = (
 		EditCondition = "bAllowedToEditDrawPlane", HideEditConditionToggle))
diff --git a/Engine/Plugins/Experimental/MeshModelingToolset/Source/MeshModelingTools/Public/DynamicMeshSculptTool.h b/Engine/Plugins/Experimental/MeshModelingToolset/Source/MeshModelingTools/Public/DynamicMeshSculptTool.h
--- a/Engine/Plugins/Experimental/MeshModelingToolset/Source/MeshModelingTools/Public/DynamicMeshSculptTool.h
+++ b/Engine/Plugins/Experimental/MeshModelingToolset/Source/MeshModelingTools/Public/DynamicMeshSculptTool.h
@@ -199,14 +199,10 @@ public:
 	UPROPERTY( meta = (TransientToolProperty) )
 	bool bPropertySetEnabled = true;
 
-	/** Toggle whether Work Plane Positioing Gizmo is visible */
+	/** Toggle whether Work Plane Positioning Gizmo is visible */
 	UPROPERTY(EditAnywhere, Category = TargetPlane, meta = (HideEditConditionToggle, EditCondition = "bPropertySetEnabled == true"))
 	bool bShowGizmo = true;
 
-	/** Toggle whether Work Plane snaps to grid when using Gizmo */
-	UPROPERTY(EditAnywhere, Category = TargetPlane, meta = (HideEditConditionToggle, EditCondition = "bPropertySetEnabled == true"))
-	bool bSnapToGrid = true;
-
 	UPROPERTY(EditAnywhere, Category = TargetPlane, AdvancedDisplay, meta = (HideEditConditionToggle, EditCondition = "bPropertySetEnabled == true"))
 	FVector Position = FVector::ZeroVector;
 
diff --git a/Engine/Plugins/Experimental/MeshModelingToolset/Source/MeshModelingTools/Public/MeshSpaceDeformerTool.h b/Engine/Plugins/Experimental/MeshModelingToolset/Source/MeshModelingTools/Public/MeshSpaceDeformerTool.h
--- a/Engine/Plugins/Experimental/MeshModelingToolset/Source/MeshModelingTools/Public/MeshSpaceDeformerTool.h
+++ b/Engine/Plugins/Experimental/MeshModelingToolset/Source/MeshModelingTools/Public/MeshSpaceDeformerTool.h
@@ -137,10 +137,6 @@ public:
 	UPROPERTY(EditAnywhere, Category = Options, meta = (EditCondition = "SelectedOperationType == ENonlinearOperationType::Bend", EditConditionHides))
 	bool bDrawVisualization = true;
 
-	/** Snap the deformer gizmo to the world grid */
-	UPROPERTY(EditAnywhere, Category = Gizmo)
-	bool bSnapToWorldGrid = false;
-
 	/** When true, Ctrl+click not only moves the gizmo to the clicked location, but also aligns the Z axis with the normal at that point. */
 	UPROPERTY(EditAnywhere, Category = Gizmo)
 	bool bAlignToNormalOnCtrlClick = false;
diff --git a/Engine/Plugins/Experimental/MeshModelingToolset/Source/MeshModelingTools/Public/MirrorTool.h b/Engine/Plugins/Experimental/MeshModelingToolset/Source/MeshModelingTools/Public/MirrorTool.h
--- a/Engine/Plugins/Experimental/MeshModelingToolset/Source/MeshModelingTools/Public/MirrorTool.h
+++ b/Engine/Plugins/Experimental/MeshModelingToolset/Source/MeshModelingTools/Public/MirrorTool.h
@@ -94,10 +94,6 @@ public:
 		EditCondition = "bWeldVerticesOnMirrorPlane && OperationMode == EMirrorOperationMode::MirrorAndAppend", EditConditionHides))
 	bool bAllowBowtieVertexCreation = false;
 
-	/** Snap the mirror plane to the world grid. */
-	UPROPERTY(EditAnywhere, Category = RepositionOptions)
-	bool bSnapToWorldGrid = true;
-
 	/** What Ctrl + clicking does to the mirror plane. */
 	UPROPERTY(EditAnywhere, Category = RepositionOptions)
 	EMirrorCtrlClickBehavior CtrlClickBehavior = EMirrorCtrlClickBehavior::Reposition;
@@ -235,10 +231,6 @@ protected:
 
 	UWorld* TargetWorld;
 
-	// flags used to identify modifier keys/buttons
-	static const int SnappingToggleModifierId = 1;
-	bool bSnappingToggle = false;		// toggled by hotkey (shift)
-
 	UPROPERTY()
 	TObjectPtr<UConstructionPlaneMechanic> PlaneMechanic;
 
diff --git a/Engine/Plugins/Experimental/MeshModelingToolset/Source/MeshModelingTools/Public/PlaneCutTool.h b/Engine/Plugins/Experimental/MeshModelingToolset/Source/MeshModelingTools/Public/PlaneCutTool.h
--- a/Engine/Plugins/Experimental/MeshModelingToolset/Source/MeshModelingTools/Public/PlaneCutTool.h
+++ b/Engine/Plugins/Experimental/MeshModelingToolset/Source/MeshModelingTools/Public/PlaneCutTool.h
@@ -69,14 +69,6 @@ public:
 	UPlaneCutToolProperties()
 	{}
 
-	/** Snap the cut plane to the world grid */
-	UPROPERTY(EditAnywhere, Category = Snapping)
-	bool bSnapToWorldGrid = false;
-
-	/** Snap rotation of the cut plane to the world grid */
-	UPROPERTY(EditAnywhere, Category = Snapping)
-	bool bSnapRotationToWorldGrid = false;
-
 	/** If true, both halves of the cut are computed */
 	UPROPERTY(EditAnywhere, Category = Options)
 	bool bKeepBothHalves = false;
@@ -129,7 +121,7 @@ public:
  * Simple Mesh Plane Cutting Tool
  */
 UCLASS()
-class MESHMODELINGTOOLS_API UPlaneCutTool : public UMultiSelectionTool, public IModifierToggleBehaviorTarget
+class MESHMODELINGTOOLS_API UPlaneCutTool : public UMultiSelectionTool
 {
 	GENERATED_BODY()
 
@@ -159,10 +151,6 @@ public:
 
 	virtual void OnPropertyModified(UObject* PropertySet, FProperty* Property) override;
 
-	// IClickSequenceBehaviorTarget implementation
-	virtual void OnUpdateModifierState(int ModifierID, bool bIsOn) override;
-
-
 protected:
 
 	UPROPERTY()
@@ -211,10 +199,6 @@ protected:
 
 	FViewCameraState CameraState;
 
-	// flags used to identify modifier keys/buttons
-	static const int SnappingModifier = 1;
-	bool bSnappingToggle = false;		// toggled by hotkey (shift)
-
 	EPlaneCutToolActions PendingAction = EPlaneCutToolActions::NoAction;
 
 
diff --git a/Engine/Plugins/Experimental/MeshModelingToolset/Source/MeshModelingTools/Public/RevolveBoundaryTool.h b/Engine/Plugins/Experimental/MeshModelingToolset/Source/MeshModelingTools/Public/RevolveBoundaryTool.h
--- a/Engine/Plugins/Experimental/MeshModelingToolset/Source/MeshModelingTools/Public/RevolveBoundaryTool.h
+++ b/Engine/Plugins/Experimental/MeshModelingToolset/Source/MeshModelingTools/Public/RevolveBoundaryTool.h
@@ -65,10 +65,6 @@ public:
 	UPROPERTY(EditAnywhere, Category = RevolutionAxis, meta = (
 		UIMin = -180, UIMax = 180, ClampMin = -180000, ClampMax = 18000))
 	float AxisPitch = 0;
-
-	/** Determines whether the axis control widget snaps to world grid (only relevant if world coordinate mode is active in viewport) .*/
-	UPROPERTY(EditAnywhere, Category = RevolutionAxis)
-	bool bSnapToWorldGrid = false;
 };
 
 /** 
diff --git a/Engine/Plugins/Experimental/MeshModelingToolset/Source/MeshModelingTools/Public/Sculpting/MeshSculptToolBase.h b/Engine/Plugins/Experimental/MeshModelingToolset/Source/MeshModelingTools/Public/Sculpting/MeshSculptToolBase.h
--- a/Engine/Plugins/Experimental/MeshModelingToolset/Source/MeshModelingTools/Public/Sculpting/MeshSculptToolBase.h
+++ b/Engine/Plugins/Experimental/MeshModelingToolset/Source/MeshModelingTools/Public/Sculpting/MeshSculptToolBase.h
@@ -128,10 +128,6 @@ public:
 	UPROPERTY(EditAnywhere, Category = TargetPlane, meta = (HideEditConditionToggle, EditCondition = "bPropertySetEnabled == true"))
 	bool bShowGizmo = true;
 
-	/** Toggle whether Work Plane snaps to grid when using Gizmo */
-	UPROPERTY(EditAnywhere, Category = TargetPlane, meta = (HideEditConditionToggle, EditCondition = "bPropertySetEnabled == true"))
-	bool bSnapToGrid = true;
-
 	UPROPERTY(EditAnywhere, Category = TargetPlane, AdvancedDisplay, meta = (HideEditConditionToggle, EditCondition = "bPropertySetEnabled == true"))
 	FVector Position = FVector::ZeroVector;
 
diff --git a/Engine/Plugins/Experimental/MeshModelingToolset/Source/MeshModelingTools/Public/TransformMeshesTool.h b/Engine/Plugins/Experimental/MeshModelingToolset/Source/MeshModelingTools/Public/TransformMeshesTool.h
--- a/Engine/Plugins/Experimental/MeshModelingToolset/Source/MeshModelingTools/Public/TransformMeshesTool.h
+++ b/Engine/Plugins/Experimental/MeshModelingToolset/Source/MeshModelingTools/Public/TransformMeshesTool.h
@@ -106,15 +106,6 @@ public:
 	bool bApplyToInstances = true;
 
 
-	/** Snap to World Space grid during translation with World Axes */
-	UPROPERTY(EditAnywhere, Category = Snapping)
-	bool bSnapToGrid = false;
-
-	/** Snap to fixed rotation angles during axis rotation */
-	UPROPERTY(EditAnywhere, Category = Snapping)
-	bool bSnapToRotation = false;
-
-
 	/** When true, the Gizmo can be moved independently without affecting objects. This allows the Gizmo to be repositioned before transforming. */
 	UPROPERTY(EditAnywhere, Category = Pivot, meta = (TransientToolProperty, EditCondition = "TransformMode != ETransformMeshesTransformMode::PerObjectGizmo") )
 	bool bSetPivotMode = false;
@@ -214,6 +205,4 @@ protected:
 	int ActiveSnapDragIndex = -1;
 
 	void OnParametersUpdated();
-
-	void OnSnappingSettingsUpdated();
 };
diff --git a/Engine/Plugins/Experimental/MeshModelingToolset/Source/ModelingComponents/Private/BaseTools/BaseCreateFromSelectedTool.cpp b/Engine/Plugins/Experimental/MeshModelingToolset/Source/ModelingComponents/Private/BaseTools/BaseCreateFromSelectedTool.cpp
--- a/Engine/Plugins/Experimental/MeshModelingToolset/Source/ModelingComponents/Private/BaseTools/BaseCreateFromSelectedTool.cpp
+++ b/Engine/Plugins/Experimental/MeshModelingToolset/Source/ModelingComponents/Private/BaseTools/BaseCreateFromSelectedTool.cpp
@@ -158,11 +158,6 @@ void UBaseCreateFromSelectedTool::SetWorld(UWorld* World)
 
 void UBaseCreateFromSelectedTool::OnTick(float DeltaTime)
 {
-	for (UTransformGizmo* Gizmo : TransformGizmos)
-	{
-		Gizmo->bSnapToWorldGrid = TransformProperties->bSnapToWorldGrid;
-	}
-
 	Preview->Tick(DeltaTime);
 }
 
@@ -314,7 +309,6 @@ void UBaseCreateFromSelectedTool::OnPropertyModified(UObject* PropertySet, FProp
 	}
 	else if (Property && 
 		(  PropertySet == HandleSourcesProperties
-		|| Property->GetFName() == GET_MEMBER_NAME_CHECKED(UTransformInputsToolProperties, bSnapToWorldGrid)
 		))
 	{
 		// nothing
diff --git a/Engine/Plugins/Experimental/MeshModelingToolset/Source/ModelingComponents/Private/Mechanics/ConstructionPlaneMechanic.cpp b/Engine/Plugins/Experimental/MeshModelingToolset/Source/ModelingComponents/Private/Mechanics/ConstructionPlaneMechanic.cpp
--- a/Engine/Plugins/Experimental/MeshModelingToolset/Source/ModelingComponents/Private/Mechanics/ConstructionPlaneMechanic.cpp
+++ b/Engine/Plugins/Experimental/MeshModelingToolset/Source/ModelingComponents/Private/Mechanics/ConstructionPlaneMechanic.cpp
@@ -57,16 +57,6 @@ void UConstructionPlaneMechanic::Initialize(UWorld* TargetWorld, const FFrame3d&
 
 
 
-void UConstructionPlaneMechanic::SetEnableGridSnaping(bool bEnable)
-{
-	bEnableSnapToWorldGrid = bEnable;
-}
-
-void UConstructionPlaneMechanic::SetEnableGridRotationSnapping(bool bEnableRotation)
-{
-	bEnableSnapToWorldRotGrid = bEnableRotation;
-}
-
 void UConstructionPlaneMechanic::UpdateClickPriority(FInputCapturePriority NewPriority)
 {
 	ensure(ClickToSetPlaneBehavior != nullptr);
@@ -122,8 +112,6 @@ void UConstructionPlaneMechanic::Tick(float DeltaTime)
 {
 	if (PlaneTransformGizmo != nullptr)
 	{
-		PlaneTransformGizmo->bSnapToWorldGrid = bEnableSnapToWorldGrid;
-		PlaneTransformGizmo->bSnapToWorldRotGrid = bEnableSnapToWorldRotGrid;
 		PlaneTransformGizmo->SetVisibility(CanUpdatePlaneFunc());
 	}
 }
diff --git a/Engine/Plugins/Experimental/MeshModelingToolset/Source/ModelingComponents/Public/BaseTools/BaseCreateFromSelectedTool.h b/Engine/Plugins/Experimental/MeshModelingToolset/Source/ModelingComponents/Public/BaseTools/BaseCreateFromSelectedTool.h
--- a/Engine/Plugins/Experimental/MeshModelingToolset/Source/ModelingComponents/Public/BaseTools/BaseCreateFromSelectedTool.h
+++ b/Engine/Plugins/Experimental/MeshModelingToolset/Source/ModelingComponents/Public/BaseTools/BaseCreateFromSelectedTool.h
@@ -83,10 +83,6 @@ public:
 	/** Show UI to allow changing translation, rotation and scale of input meshes */
 	UPROPERTY(EditAnywhere, Category = Transform)
 	bool bShowTransformUI = true;
-
-	/** Snap the cut plane to the world grid */
-	UPROPERTY(EditAnywhere, Category = Transform, meta = (EditCondition = "bShowTransformUI == true"))
-	bool bSnapToWorldGrid = false;
 };
 
 
diff --git a/Engine/Plugins/Experimental/MeshModelingToolset/Source/ModelingComponents/Public/Mechanics/ConstructionPlaneMechanic.h b/Engine/Plugins/Experimental/MeshModelingToolset/Source/ModelingComponents/Public/Mechanics/ConstructionPlaneMechanic.h
--- a/Engine/Plugins/Experimental/MeshModelingToolset/Source/ModelingComponents/Public/Mechanics/ConstructionPlaneMechanic.h
+++ b/Engine/Plugins/Experimental/MeshModelingToolset/Source/ModelingComponents/Public/Mechanics/ConstructionPlaneMechanic.h
@@ -56,10 +56,6 @@ public:
 	 */
 	void SetPlaneWithoutBroadcast(const UE::Geometry::FFrame3d& Plane);
 
-	void SetEnableGridSnaping(bool bEnable);
-
-	void SetEnableGridRotationSnapping(bool bEnableRotationSnapping);
-
 	void UpdateClickPriority(FInputCapturePriority NewPriority);
 
 public:
@@ -77,9 +73,6 @@ public:
 	TUniquePtr<FSelectClickedAction> SetPlaneCtrlClickBehaviorTarget;
 
 protected:
-	bool bEnableSnapToWorldGrid = false;
-	bool bEnableSnapToWorldRotGrid = false;
-
 	UPROPERTY()
 	TObjectPtr<USingleClickInputBehavior> ClickToSetPlaneBehavior;
 
diff --git a/Engine/Source/Editor/Experimental/EditorInteractiveToolsFramework/Private/EdModeInteractiveToolsContext.cpp b/Engine/Source/Editor/Experimental/EditorInteractiveToolsFramework/Private/EdModeInteractiveToolsContext.cpp
--- a/Engine/Source/Editor/Experimental/EditorInteractiveToolsFramework/Private/EdModeInteractiveToolsContext.cpp
+++ b/Engine/Source/Editor/Experimental/EditorInteractiveToolsFramework/Private/EdModeInteractiveToolsContext.cpp
@@ -217,15 +217,18 @@ public:
 	{
 		if ((Request.TargetTypes & ESceneSnapQueryTargetType::Grid) != ESceneSnapQueryTargetType::None)
 		{
-			FRotator Rotator ( Request.DeltaRotation );
-			FRotator RotGrid = Request.RotGridSize.Get(GEditor->GetRotGridSize());
-			Rotator = Rotator.GridSnap( RotGrid );
-
-			FSceneSnapQueryResult SnapResult;
-			SnapResult.TargetType = ESceneSnapQueryTargetType::Grid;
-			SnapResult.DeltaRotation = Rotator.Quaternion();
-			Results.Add(SnapResult);
-			return true;
+			if (GetDefault<ULevelEditorViewportSettings>()->RotGridEnabled != 0)
+			{
+				FRotator Rotator(Request.DeltaRotation);
+				FRotator RotGrid = Request.RotGridSize.Get(GEditor->GetRotGridSize());
+				Rotator = Rotator.GridSnap(RotGrid);
+
+				FSceneSnapQueryResult SnapResult;
+				SnapResult.TargetType = ESceneSnapQueryTargetType::Grid;
+				SnapResult.DeltaRotation = Rotator.Quaternion();
+				Results.Add(SnapResult);
+				return true;
+			}
 		}
 		return false;
 	}
@@ -236,18 +239,21 @@ public:
 
 		if ((Request.TargetTypes & ESceneSnapQueryTargetType::Grid) != ESceneSnapQueryTargetType::None)
 		{
-			FSceneSnapQueryResult SnapResult;
-			SnapResult.TargetType = ESceneSnapQueryTargetType::Grid;
+			if (GetDefault<ULevelEditorViewportSettings>()->GridEnabled != 0)
+			{
+				FSceneSnapQueryResult SnapResult;
+				SnapResult.TargetType = ESceneSnapQueryTargetType::Grid;
 
-			float SnapSize = GEditor->GetGridSize();
-			FVector GridSize = Request.GridSize.Get(FVector(SnapSize, SnapSize, SnapSize));
+				float SnapSize = GEditor->GetGridSize();
+				FVector GridSize = Request.GridSize.Get(FVector(SnapSize, SnapSize, SnapSize));
 
-			SnapResult.Position.X = SnapToIncrement(Request.Position.X, GridSize.X);
-			SnapResult.Position.Y = SnapToIncrement(Request.Position.Y, GridSize.Y);
-			SnapResult.Position.Z = SnapToIncrement(Request.Position.Z, GridSize.Z);
+				SnapResult.Position.X = SnapToIncrement(Request.Position.X, GridSize.X);
+				SnapResult.Position.Y = SnapToIncrement(Request.Position.Y, GridSize.Y);
+				SnapResult.Position.Z = SnapToIncrement(Request.Position.Z, GridSize.Z);
 
-			Results.Add(SnapResult);
-			FoundResultCount++;
+				Results.Add(SnapResult);
+				FoundResultCount++;
+			}
 		}
 
 		//
diff --git a/Engine/Source/Runtime/Experimental/InteractiveToolsFramework/Private/BaseGizmos/AxisAngleGizmo.cpp b/Engine/Source/Runtime/Experimental/InteractiveToolsFramework/Private/BaseGizmos/AxisAngleGizmo.cpp
--- a/Engine/Source/Runtime/Experimental/InteractiveToolsFramework/Private/BaseGizmos/AxisAngleGizmo.cpp
+++ b/Engine/Source/Runtime/Experimental/InteractiveToolsFramework/Private/BaseGizmos/AxisAngleGizmo.cpp
@@ -23,7 +23,6 @@ void UAxisAngleGizmo::Setup()
 
 	// Add default mouse input behavior
 	MouseBehavior = NewObject<UClickDragInputBehavior>();
-	MouseBehavior->Modifiers.RegisterModifier(SnapAngleModifierID, FInputDeviceState::IsShiftKeyDown);
 	MouseBehavior->Initialize(this);
 	MouseBehavior->SetDefaultPriority(FInputCapturePriority(FInputCapturePriority::DEFAULT_GIZMO_PRIORITY));
 	AddInputBehavior(MouseBehavior);
@@ -43,10 +42,6 @@ void UAxisAngleGizmo::Setup()
 
 void UAxisAngleGizmo::OnUpdateModifierState(int ModifierID, bool bIsOn)
 {
-	if (ModifierID == SnapAngleModifierID)
-	{
-		bEnableSnapAngleModifier = bIsOn;
-	}
 }
 
 
@@ -151,11 +146,6 @@ void UAxisAngleGizmo::OnClickDrag(const FInputDeviceRay& DragPos)
 			RotationOrigin, RotationAxis, RotationPlaneX, RotationPlaneY);
 
 		DeltaAngle = InteractionCurAngle - InteractionStartAngle;
-		if (bEnableSnapAngleModifier)
-		{
-			DeltaAngle = GizmoMath::SnapToIncrement(DeltaAngle, 3.14159f / 4.0f);				// hardcoded 45-degree snap
-		}
-		
 	}
 
 	float NewAngle = InitialTargetAngle + DeltaAngle;
diff --git a/Engine/Source/Runtime/Experimental/InteractiveToolsFramework/Private/BaseGizmos/TransformGizmo.cpp b/Engine/Source/Runtime/Experimental/InteractiveToolsFramework/Private/BaseGizmos/TransformGizmo.cpp
--- a/Engine/Source/Runtime/Experimental/InteractiveToolsFramework/Private/BaseGizmos/TransformGizmo.cpp
+++ b/Engine/Source/Runtime/Experimental/InteractiveToolsFramework/Private/BaseGizmos/TransformGizmo.cpp
@@ -927,7 +927,8 @@ bool UTransformGizmo::PositionSnapFunction(const FVector& WorldPosition, FVector
 	}
 
 	// only snap to world grid when using world axes
-	if (GetGizmoManager()->GetContextQueriesAPI()->GetCurrentCoordinateSystem() != EToolContextCoordinateSystem::World)
+	EToolContextCoordinateSystem CoordSystem = GetGizmoManager()->GetContextQueriesAPI()->GetCurrentCoordinateSystem();
+	if (CoordSystem != EToolContextCoordinateSystem::World)
 	{
 		return false;
 	}
diff --git a/Engine/Source/Runtime/Experimental/InteractiveToolsFramework/Public/BaseGizmos/AxisAngleGizmo.h b/Engine/Source/Runtime/Experimental/InteractiveToolsFramework/Public/BaseGizmos/AxisAngleGizmo.h
--- a/Engine/Source/Runtime/Experimental/InteractiveToolsFramework/Public/BaseGizmos/AxisAngleGizmo.h
+++ b/Engine/Source/Runtime/Experimental/InteractiveToolsFramework/Public/BaseGizmos/AxisAngleGizmo.h
@@ -130,9 +130,5 @@ protected:
 	// The angle of the closest axis in the plane of the angle gizmo (will be 0, pi/2, pi, or 3pi/2).
 	// Used for snapping the nearest axis to a ray cast point when using CustomDestinationRayCaster.
 	float ClosestAxisStartAngle;
-
-	static const int SnapAngleModifierID = 1;		// identifier we associate with the shift key
-	bool bEnableSnapAngleModifier = false;			// flag we use to keep track of modifier state
-
 };
 
diff --git a/Engine/Source/Runtime/Experimental/InteractiveToolsFramework/Public/BaseGizmos/TransformGizmo.h b/Engine/Source/Runtime/Experimental/InteractiveToolsFramework/Public/BaseGizmos/TransformGizmo.h
--- a/Engine/Source/Runtime/Experimental/InteractiveToolsFramework/Public/BaseGizmos/TransformGizmo.h
+++ b/Engine/Source/Runtime/Experimental/InteractiveToolsFramework/Public/BaseGizmos/TransformGizmo.h
@@ -364,7 +364,7 @@ public:
 	 * If true, then when using world frame, Axis and Plane translation snap to the world grid via the ContextQueriesAPI (in PositionSnapFunction)
 	 */
 	UPROPERTY()
-	bool bSnapToWorldGrid = false;
+	bool bSnapToWorldGrid = true;
 
 	/**
 	 * Optional grid size which overrides the Context Grid
@@ -386,7 +386,7 @@ public:
 	 * If true, then when using world frame, Axis and Plane translation snap to the world grid via the ContextQueriesAPI (in RotationSnapFunction)
 	 */
 	UPROPERTY()
-	bool bSnapToWorldRotGrid = false;
+	bool bSnapToWorldRotGrid = true;
 
 	/**
 	 * Whether to use the World/Local coordinate system provided by the context via the ContextyQueriesAPI.
