BazleExecutor: Thread-safey for QueuedItems map

#rnx

diff --git a/Engine/Source/Editor/Experimental/BazelExecutor/Private/BazelCompletionQueueRunnable.cpp b/Engine/Source/Editor/Experimental/BazelExecutor/Private/BazelCompletionQueueRunnable.cpp
--- a/Engine/Source/Editor/Experimental/BazelExecutor/Private/BazelCompletionQueueRunnable.cpp
+++ b/Engine/Source/Editor/Experimental/BazelExecutor/Private/BazelCompletionQueueRunnable.cpp
@@ -24,7 +24,9 @@ FBazelCompletionQueueRunnable::~FBazelCompletionQueueRunnable()
 
 void FBazelCompletionQueueRunnable::ProcessNext(void* Tag, bool Ok)
 {
+	QueuedItemsMutex.lock();
 	FQueuedItem* Item = QueuedItems.Find(Tag);
+	QueuedItemsMutex.unlock();
 	if (!Item)
 	{
 		return;
@@ -78,7 +80,9 @@ void FBazelCompletionQueueRunnable::ProcessNext(void* Tag, bool Ok)
 			Item->Finish(Tag, Ok, *Item->Status.Get(), *Item->Message.Get());
 		}
 
+		QueuedItemsMutex.lock();
 		QueuedItems.Remove(Tag);
+		QueuedItemsMutex.unlock();
 		break;
 	}
 	}
@@ -121,6 +125,7 @@ void FBazelCompletionQueueRunnable::Exit()
 	bool Ok;
 	while (CompletionQueue->Next(&Tag, &Ok)) {}
 
+	QueuedItemsMutex.lock();
 	for (const TPair<void*,FQueuedItem>& Pair : QueuedItems)
 	{
 		const FQueuedItem& Item = Pair.Value;
@@ -131,6 +136,7 @@ void FBazelCompletionQueueRunnable::Exit()
 	}
 
 	QueuedItems.Empty();
+	QueuedItemsMutex.unlock();
 }
 
 FSingleThreadRunnable* FBazelCompletionQueueRunnable::GetSingleThreadInterface()
@@ -183,7 +189,9 @@ bool FBazelCompletionQueueRunnable::AddAsyncOperation(
 	Item.StartCall = MoveTemp(OnStartCall);
 	Item.Read = MoveTemp(OnRead);
 	Item.Finish = MoveTemp(OnFinish);
+	QueuedItemsMutex.lock();
 	QueuedItems.Add(AsyncReader, MoveTemp(Item));
+	QueuedItemsMutex.unlock();
 	AsyncReader->StartCall(AsyncReader);
 	return true;
 }
diff --git a/Engine/Source/Editor/Experimental/BazelExecutor/Private/BazelCompletionQueueRunnable.h b/Engine/Source/Editor/Experimental/BazelExecutor/Private/BazelCompletionQueueRunnable.h
--- a/Engine/Source/Editor/Experimental/BazelExecutor/Private/BazelCompletionQueueRunnable.h
+++ b/Engine/Source/Editor/Experimental/BazelExecutor/Private/BazelCompletionQueueRunnable.h
@@ -7,6 +7,8 @@
 #include "HAL/RunnableThread.h"
 #include "Misc/SingleThreadRunnable.h"
 
+#include <mutex>
+
 namespace grpc
 {
 	class ClientContext;
@@ -60,6 +62,7 @@ private:
 	bool Running;
 	TMap<void*, FQueuedItem> QueuedItems;
 	TUniquePtr<grpc::CompletionQueue> CompletionQueue;
+	std::mutex QueuedItemsMutex;
 
 	void ProcessNext(void* Tag, bool Ok);
 
diff --git a/Engine/Source/Editor/Experimental/BazelExecutor/Private/BazelCompletionQueueRunnable.inl b/Engine/Source/Editor/Experimental/BazelExecutor/Private/BazelCompletionQueueRunnable.inl
--- a/Engine/Source/Editor/Experimental/BazelExecutor/Private/BazelCompletionQueueRunnable.inl
+++ b/Engine/Source/Editor/Experimental/BazelExecutor/Private/BazelCompletionQueueRunnable.inl
@@ -35,7 +35,9 @@ bool FBazelCompletionQueueRunnable::AddAsyncResponse(
 	Item.Message.Reset(Response);
 	Item.ClientContext = MoveTemp(ClientContext);
 	Item.Finish = MoveTemp(OnFinishResponse);
+	QueuedItemsMutex.lock();
 	QueuedItems.Add(AsyncReader, MoveTemp(Item));
+	QueuedItemsMutex.unlock();
 	AsyncReader->StartCall();
 	AsyncReader->Finish(Response, Status, AsyncReader);
 	return true;
