Allow BazelExecutor to unity build

diff --git a/Engine/Source/Editor/Experimental/BazelExecutor/BazelExecutor.Build.cs b/Engine/Source/Editor/Experimental/BazelExecutor/BazelExecutor.Build.cs
--- a/Engine/Source/Editor/Experimental/BazelExecutor/BazelExecutor.Build.cs
+++ b/Engine/Source/Editor/Experimental/BazelExecutor/BazelExecutor.Build.cs
@@ -17,8 +17,6 @@ namespace UnrealBuildTool.Rules
 				return;
 			}
 
-			bUseUnity = false;
-
 			PrivateDefinitions.Add("GPR_FORBID_UNREACHABLE_CODE=0");
 
 			PrivateIncludePaths.Add(Path.Combine(ModuleDirectory, "Generated"));
