Reduce severity of logging messages emitted from BazelExecutor to avoid having errors reported for recoverable situations when cooking with DDC2 remote execution support.

diff --git a/Engine/Source/Editor/Experimental/BazelExecutor/Private/ContentAddressableStorage.cpp b/Engine/Source/Editor/Experimental/BazelExecutor/Private/ContentAddressableStorage.cpp
--- a/Engine/Source/Editor/Experimental/BazelExecutor/Private/ContentAddressableStorage.cpp
+++ b/Engine/Source/Editor/Experimental/BazelExecutor/Private/ContentAddressableStorage.cpp
@@ -138,7 +138,7 @@ TFuture<TPair<FStatus, FFindMissingBlobsResponse>> FContentAddressableStorage::F
 		}
 		else
 		{
-			UE_LOG(LogBazelExecutor, Error, TEXT("FindMissingBlobsAsync: Finish !Ok"));
+			UE_LOG(LogBazelExecutor, Verbose, TEXT("FindMissingBlobsAsync: Finish !Ok"));
 			Status.Code = EStatusCode::ABORTED;
 		}
 		ReturnPromise->EmplaceValue(TPair<FStatus, FFindMissingBlobsResponse>(MoveTemp(Status), MoveTemp(Response)));
@@ -183,7 +183,7 @@ TFuture<TPair<FStatus, FBatchUpdateBlobsResponse>> FContentAddressableStorage::B
 		}
 		else
 		{
-			UE_LOG(LogBazelExecutor, Error, TEXT("BatchUpdateBlobsAsync: Finish !Ok"));
+			UE_LOG(LogBazelExecutor, Verbose, TEXT("BatchUpdateBlobsAsync: Finish !Ok"));
 			Status.Code = EStatusCode::ABORTED;
 		}
 		ReturnPromise->EmplaceValue(TPair<FStatus, FBatchUpdateBlobsResponse>(MoveTemp(Status), MoveTemp(Response)));
@@ -227,7 +227,7 @@ TFuture<TPair<FStatus, FBatchReadBlobsResponse>> FContentAddressableStorage::Bat
 		}
 		else
 		{
-			UE_LOG(LogBazelExecutor, Error, TEXT("BatchReadBlobsAsync: Finish !Ok"));
+			UE_LOG(LogBazelExecutor, Verbose, TEXT("BatchReadBlobsAsync: Finish !Ok"));
 			Status.Code = EStatusCode::ABORTED;
 		}
 		ReturnPromise->EmplaceValue(TPair<FStatus, FBatchReadBlobsResponse>(MoveTemp(Status), MoveTemp(Response)));
diff --git a/Engine/Source/Editor/Experimental/BazelExecutor/Private/Execution.cpp b/Engine/Source/Editor/Experimental/BazelExecutor/Private/Execution.cpp
--- a/Engine/Source/Editor/Experimental/BazelExecutor/Private/Execution.cpp
+++ b/Engine/Source/Editor/Experimental/BazelExecutor/Private/Execution.cpp
@@ -57,18 +57,18 @@ bool FExecution::Execute(const FExecuteRequest& Request, FExecuteResponse& Respo
 		google::longrunning::Operation Operation;
 		if (!Call->Read(&Operation))
 		{
-			UE_LOG(LogBazelExecutor, Error, TEXT("Execute: Unable to read operation"));
+			UE_LOG(LogBazelExecutor, Verbose, TEXT("Execute: Unable to read operation"));
 			break;
 		}
 
 		build::bazel::remote::execution::v2::ExecuteOperationMetadata Metadata;
 		if (!Operation.metadata().UnpackTo(&Metadata))
 		{
-			UE_LOG(LogBazelExecutor, Error, TEXT("Execute: %s Unable to decode metadata"), UTF8_TO_TCHAR(Operation.name().c_str()));
+			UE_LOG(LogBazelExecutor, Verbose, TEXT("Execute: %s Unable to decode metadata"), UTF8_TO_TCHAR(Operation.name().c_str()));
 			break;
 		}
 
-		UE_LOG(LogBazelExecutor, Display, TEXT("Execute: %s Execution state: %s"),
+		UE_LOG(LogBazelExecutor, Verbose, TEXT("Execute: %s Execution state: %s"),
 			UTF8_TO_TCHAR(Operation.name().c_str()),
 			UTF8_TO_TCHAR(build::bazel::remote::execution::v2::ExecutionStage_Value_Name(Metadata.stage()).c_str()));
 
@@ -77,25 +77,15 @@ bool FExecution::Execute(const FExecuteRequest& Request, FExecuteResponse& Respo
 			build::bazel::remote::execution::v2::ExecuteResponse ProtoResponse;
 			if (!Operation.response().UnpackTo(&ProtoResponse))
 			{
-				UE_LOG(LogBazelExecutor, Error, TEXT("Execute: %s Unable to decode response"), UTF8_TO_TCHAR(Operation.name().c_str()));
+				UE_LOG(LogBazelExecutor, Verbose, TEXT("Execute: %s Unable to decode response"), UTF8_TO_TCHAR(Operation.name().c_str()));
 				break;
 			}
 			ProtoConverter::FromProto(ProtoResponse, Response);
 			if (!Response.Status.Ok())
 			{
-				if ((Response.Status.Code == EStatusCode::RESOURCE_EXHAUSTED) ||
-					(Response.Status.Code == EStatusCode::UNAVAILABLE))
-				{
-					UE_LOG(LogBazelExecutor, Display, TEXT("Execute: %s Info: %s"),
-						UTF8_TO_TCHAR(Operation.name().c_str()),
-						*Response.Status.Message);
-				}
-				else
-				{
-					UE_LOG(LogBazelExecutor, Error, TEXT("Execute: %s Error: %s"),
-						UTF8_TO_TCHAR(Operation.name().c_str()),
-						*Response.Status.Message);
-				}
+				UE_LOG(LogBazelExecutor, Verbose, TEXT("Execute: %s Info: %s"),
+					UTF8_TO_TCHAR(Operation.name().c_str()),
+					*Response.Status.Message);
 				break;
 			}
 			return true;
@@ -122,7 +112,7 @@ TFuture<FExecuteResponse> FExecution::ExecuteAsync(const FExecuteRequest& Reques
 	{
 		if (!Ok)
 		{
-			UE_LOG(LogBazelExecutor, Error, TEXT("ExecuteAsync: Call Started !Ok"));
+			UE_LOG(LogBazelExecutor, Verbose, TEXT("ExecuteAsync: Call Started !Ok"));
 			return;
 		}
 	};
@@ -131,18 +121,18 @@ TFuture<FExecuteResponse> FExecution::ExecuteAsync(const FExecuteRequest& Reques
 	{
 		if (!Ok)
 		{
-			UE_LOG(LogBazelExecutor, Error, TEXT("ExecuteAsync: Read !Ok"));
+			UE_LOG(LogBazelExecutor, Verbose, TEXT("ExecuteAsync: Read !Ok"));
 			return;
 		}
 
 		build::bazel::remote::execution::v2::ExecuteOperationMetadata Metadata;
 		if (!Operation.metadata().UnpackTo(&Metadata))
 		{
-			UE_LOG(LogBazelExecutor, Error, TEXT("ExecuteAsync: %s Unable to decode metadata"), UTF8_TO_TCHAR(Operation.name().c_str()));
+			UE_LOG(LogBazelExecutor, Verbose, TEXT("ExecuteAsync: %s Unable to decode metadata"), UTF8_TO_TCHAR(Operation.name().c_str()));
 			return;
 		}
 
-		UE_LOG(LogBazelExecutor, Display, TEXT("ExecuteAsync: %s Execution state: %s"),
+		UE_LOG(LogBazelExecutor, Verbose, TEXT("ExecuteAsync: %s Execution state: %s"),
 			UTF8_TO_TCHAR(Operation.name().c_str()),
 			UTF8_TO_TCHAR(build::bazel::remote::execution::v2::ExecutionStage_Value_Name(Metadata.stage()).c_str()));
 	};
@@ -157,7 +147,7 @@ TFuture<FExecuteResponse> FExecution::ExecuteAsync(const FExecuteRequest& Reques
 		}
 		else
 		{
-			UE_LOG(LogBazelExecutor, Error, TEXT("ExecuteAsync: %s Finish !Ok"), UTF8_TO_TCHAR(Operation.name().c_str()));
+			UE_LOG(LogBazelExecutor, Verbose, TEXT("ExecuteAsync: %s Finish !Ok"), UTF8_TO_TCHAR(Operation.name().c_str()));
 			Status.Code = EStatusCode::ABORTED;
 		}
 
@@ -173,7 +163,7 @@ TFuture<FExecuteResponse> FExecution::ExecuteAsync(const FExecuteRequest& Reques
 		build::bazel::remote::execution::v2::ExecuteResponse ProtoResponse;
 		if (!Operation.response().UnpackTo(&ProtoResponse))
 		{
-			UE_LOG(LogBazelExecutor, Error, TEXT("ExecuteAsync: %s Unable to decode response"), UTF8_TO_TCHAR(Operation.name().c_str()));
+			UE_LOG(LogBazelExecutor, Verbose, TEXT("ExecuteAsync: %s Unable to decode response"), UTF8_TO_TCHAR(Operation.name().c_str()));
 			FExecuteResponse Response;
 			Response.Status.Code = EStatusCode::INTERNAL;
 			Response.Status.Message = TEXT("Unable to decode response");
@@ -185,19 +175,9 @@ TFuture<FExecuteResponse> FExecution::ExecuteAsync(const FExecuteRequest& Reques
 		ProtoConverter::FromProto(ProtoResponse, Response);
 		if (!Response.Status.Ok())
 		{
-			if ((Response.Status.Code == EStatusCode::RESOURCE_EXHAUSTED) ||
-				(Response.Status.Code == EStatusCode::UNAVAILABLE))
-			{
-				UE_LOG(LogBazelExecutor, Display, TEXT("ExecuteAsync: %s Info: %s"),
-					UTF8_TO_TCHAR(Operation.name().c_str()),
-					*Response.Status.Message);
-			}
-			else
-			{
-				UE_LOG(LogBazelExecutor, Error, TEXT("ExecuteAsync: %s Error: %s"),
-					UTF8_TO_TCHAR(Operation.name().c_str()),
-					*Response.Status.Message);
-			}
+			UE_LOG(LogBazelExecutor, Verbose, TEXT("ExecuteAsync: %s Info: %s"),
+				UTF8_TO_TCHAR(Operation.name().c_str()),
+				*Response.Status.Message);
 		}
 
 		ReturnPromise->EmplaceValue(MoveTemp(Response));
