Add Accessory and Primary gizmo support to Editor gizmo manager and subsystem.

diff --git a/Engine/Source/Editor/Experimental/EditorInteractiveToolsFramework/Private/EditorGizmos/EditorTransformGizmo.cpp b/Engine/Source/Editor/Experimental/EditorInteractiveToolsFramework/Private/EditorGizmos/EditorTransformGizmo.cpp
--- a/Engine/Source/Editor/Experimental/EditorInteractiveToolsFramework/Private/EditorGizmos/EditorTransformGizmo.cpp
+++ b/Engine/Source/Editor/Experimental/EditorInteractiveToolsFramework/Private/EditorGizmos/EditorTransformGizmo.cpp
@@ -63,40 +63,18 @@ UInteractiveGizmo* UEditorTransformGizmoBuilder::BuildGizmo(const FToolBuilderSt
 
 void UEditorTransformGizmoBuilder::UpdateGizmoForSelection(UInteractiveGizmo* Gizmo, const FToolBuilderState& SceneState)
 {
-	UEditorTransformGizmo* TransformGizmo = Cast<UEditorTransformGizmo>(Gizmo);
-
-	if (ensure(TransformGizmo))
+	if (UEditorTransformGizmo* TransformGizmo = Cast<UEditorTransformGizmo>(Gizmo))
 	{
-		UTransformProxy* TransformProxy = NewObject<UTransformProxy>();
-		bool bVisibility = false;
-
-		ETypedElementWorldType WorldType = (SceneState.World->WorldType == EWorldType::PIE ? ETypedElementWorldType::Game : ETypedElementWorldType::Editor);
-
-		// @todo - once UTransformProxy supports typed elements, update this to use the normalized typed
-		// element selection set.
-		if (UTypedElementSelectionSet* SelectionSet = SceneState.TypedElementSelectionSet.Get())
+		if (UTransformProxy* TransformProxy = FEditorGizmoSelectionBuilderHelper::CreateTransformProxyForSelection(SceneState))
 		{
-			SelectionSet->ForEachSelectedElement<UTypedElementWorldInterface>([TransformProxy, &bVisibility, SelectionSet, WorldType](const TTypedElement<UTypedElementWorldInterface>& InWorldElement)
-			{
-				if (InWorldElement.CanMoveElement(WorldType))
-				{
-					if (TTypedElement<UTypedElementObjectInterface> ObjectTypedElement = SelectionSet->GetElementList()->GetElement<UTypedElementObjectInterface>(InWorldElement))
-					{
-						if (AActor* Actor = ObjectTypedElement.GetObjectAs<AActor>())
-						{
-							USceneComponent* SceneComponent = Actor->GetRootComponent();
-							TransformProxy->AddComponent(SceneComponent);
-							bVisibility = true;
-						}
-						return true;
-					}
-				}
-				return false;
-			});
+			TransformGizmo->SetActiveTarget(TransformProxy);
+			TransformGizmo->SetVisibility(true);
+		}
+		else
+		{
+			TransformGizmo->SetActiveTarget(nullptr);
+			TransformGizmo->SetVisibility(false);
 		}
-
-		TransformGizmo->SetActiveTarget(TransformProxy);
-		TransformGizmo->SetVisibility(bVisibility);
 	}
 }
 
diff --git a/Engine/Source/Editor/Experimental/EditorInteractiveToolsFramework/Private/EditorInteractiveGizmoManager.cpp b/Engine/Source/Editor/Experimental/EditorInteractiveToolsFramework/Private/EditorInteractiveGizmoManager.cpp
--- a/Engine/Source/Editor/Experimental/EditorInteractiveToolsFramework/Private/EditorInteractiveGizmoManager.cpp
+++ b/Engine/Source/Editor/Experimental/EditorInteractiveToolsFramework/Private/EditorInteractiveGizmoManager.cpp
@@ -2,7 +2,6 @@
 
 
 #include "EditorInteractiveGizmoManager.h"
-#include "EditorInteractiveGizmoConditionalBuilder.h"
 #include "EditorInteractiveGizmoSelectionBuilder.h"
 #include "EditorInteractiveGizmoSubsystem.h"
 #include "EdModeInteractiveToolsContext.h"
@@ -17,7 +16,6 @@
 #include "Selection.h"
 #include "ToolContextInterfaces.h"
 
-
 #define LOCTEXT_NAMESPACE "UEditorInteractiveGizmoManager"
 
 static TAutoConsoleVariable<int32> CVarUseLegacyWidget(
@@ -32,7 +30,7 @@ static TAutoConsoleVariable<int32> CVarUseLegacyWidget(
 UEditorInteractiveGizmoManager::UEditorInteractiveGizmoManager() :
 	UInteractiveGizmoManager()
 {
-
+	Registry = NewObject<UEditorInteractiveGizmoRegistry>();
 }
 
 
@@ -45,126 +43,84 @@ void UEditorInteractiveGizmoManager::InitializeWithEditorModeManager(IToolsConte
 
 void UEditorInteractiveGizmoManager::Shutdown()
 {
-	DestroyAllSelectionGizmos();
+	Registry->Shutdown();
 	Super::Shutdown();
 }
 
-void UEditorInteractiveGizmoManager::RegisterGizmoSelectionType(const TObjectPtr<UEditorInteractiveGizmoSelectionBuilder> InGizmoSelectionBuilder)
+void UEditorInteractiveGizmoManager::RegisterEditorGizmoType(EEditorGizmoCategory InGizmoCategory, UInteractiveGizmoBuilder* InGizmoBuilder)
 {
-	if (ensure(InGizmoSelectionBuilder))
-	{
-		if (GizmoSelectionBuilders.Contains(InGizmoSelectionBuilder))
-		{
-			DisplayMessage(
-				FText::Format(LOCTEXT("DeregisterFailedMessage", "UInteractiveGizmoSubsystem::DeregisterGizmoSelectionType: type has already been registered {0}"), FText::FromName(InGizmoSelectionBuilder->GetFName())),
-				EToolMessageLevel::Internal);
-			return;
-		}
-
-		GizmoSelectionBuilders.Add(InGizmoSelectionBuilder);
-		GizmoSelectionBuilders.StableSort(
-			[](UEditorInteractiveGizmoSelectionBuilder& A, UEditorInteractiveGizmoSelectionBuilder& B) {
-				return (A).GetPriority() > (B).GetPriority();
-			});
-	}
+	check(Registry);
+	Registry->RegisterEditorGizmoType(InGizmoCategory, InGizmoBuilder);
 }
 
-bool UEditorInteractiveGizmoManager::DeregisterGizmoSelectionType(const TObjectPtr<UEditorInteractiveGizmoSelectionBuilder> InGizmoSelectionBuilder)
+void UEditorInteractiveGizmoManager::DeregisterEditorGizmoType(EEditorGizmoCategory InGizmoCategory, UInteractiveGizmoBuilder* InGizmoBuilder)
 {
-	if (ensure(InGizmoSelectionBuilder))
-	{
-		if (GizmoSelectionBuilders.Contains(InGizmoSelectionBuilder) == false)
-		{
-			DisplayMessage(
-				FText::Format(LOCTEXT("DeregisterFailedMessage", "UInteractiveGizmoSubsystem::DeregisterGizmoSelectionType: could not find requested type {0}"), FText::FromName(InGizmoSelectionBuilder->GetFName())),
-				EToolMessageLevel::Internal);
-			return false;
-		}
-		GizmoSelectionBuilders.Remove(InGizmoSelectionBuilder);
-
-		return true;
-	}
-
-	return false;
-
+	check(Registry);
+	Registry->DeregisterEditorGizmoType(InGizmoCategory, InGizmoBuilder);
 }
 
-TArray<TObjectPtr<UEditorInteractiveGizmoSelectionBuilder>> UEditorInteractiveGizmoManager::GetQualifiedGizmoSelectionBuilders(const FToolBuilderState& InToolBuilderState)
+void UEditorInteractiveGizmoManager::GetQualifiedEditorGizmoBuilders(EEditorGizmoCategory InGizmoCategory, const FToolBuilderState& InToolBuilderState, TArray<UInteractiveGizmoBuilder*>& InFoundBuilders)
 {
-	TArray<TObjectPtr<UEditorInteractiveGizmoSelectionBuilder>> FoundBuilders;
+	check(Registry);
+	Registry->GetQualifiedEditorGizmoBuilders(InGizmoCategory, InToolBuilderState, InFoundBuilders);
 
 	FEditorGizmoTypePriority FoundPriority = 0;
 
-	for (TObjectPtr<UEditorInteractiveGizmoSelectionBuilder> Builder : GizmoSelectionBuilders)
-	{
-		if (Builder->GetPriority() < FoundPriority)
-		{
-			break;
-		}
-
-		if (Builder->SatisfiesCondition(InToolBuilderState))
-		{
-			FoundBuilders.Add(Builder);
-			FoundPriority = Builder->GetPriority();
-		}
-	}
-
 	if (!bSearchLocalBuildersOnly)
 	{
 		UEditorInteractiveGizmoSubsystem* GizmoSubsystem = GEditor->GetEditorSubsystem<UEditorInteractiveGizmoSubsystem>();
 		if (ensure(GizmoSubsystem))
 		{
-			TArray<TObjectPtr<UEditorInteractiveGizmoSelectionBuilder>> FoundSubsystemBuilders = GizmoSubsystem->GetQualifiedGizmoSelectionBuilders(InToolBuilderState);
-
-			FEditorGizmoTypePriority FoundPriority0 = FoundBuilders.Num() > 0 ? FoundBuilders[0]->GetPriority() : 0;
-			FEditorGizmoTypePriority FoundPriority1 = FoundSubsystemBuilders.Num() > 0 ? FoundSubsystemBuilders[0]->GetPriority() : 0;
-
-			if (FoundPriority0 == FoundPriority1)
-			{
-				FoundBuilders.Append(FoundSubsystemBuilders);
-			}
-			else if (FoundPriority0 < FoundPriority1)
-			{
-				FoundBuilders = FoundSubsystemBuilders;
-			}
+			GizmoSubsystem->GetQualifiedGlobalEditorGizmoBuilders(InGizmoCategory, InToolBuilderState, InFoundBuilders);
 		}
 	}
-
-	return FoundBuilders;
 }
 
-TArray<UInteractiveGizmo*> UEditorInteractiveGizmoManager::CreateGizmosForCurrentSelectionState()
+TArray<UInteractiveGizmo*> UEditorInteractiveGizmoManager::CreateGizmosForCurrentSelectionSet()
 {
-	if (bShowSelectionGizmos)
+	if (bShowEditorGizmos)
 	{
-		TArray<UInteractiveGizmo*> NewGizmos;
-
 		FToolBuilderState CurrentSceneState;
 		QueriesAPI->GetCurrentSelectionState(CurrentSceneState);
 
-		TArray<TObjectPtr<UEditorInteractiveGizmoSelectionBuilder>> FoundBuilders = GetQualifiedGizmoSelectionBuilders(CurrentSceneState);
+		TArray<TObjectPtr<UInteractiveGizmoBuilder>> FoundBuilders;
+		GetQualifiedEditorGizmoBuilders(EEditorGizmoCategory::Primary, CurrentSceneState, FoundBuilders);
+
+		bool bHasPrimaryBuilder = (FoundBuilders.Num() > 0);
+
+		if (!bHasPrimaryBuilder)
+		{
+			GetQualifiedEditorGizmoBuilders(EEditorGizmoCategory::Accessory, CurrentSceneState, FoundBuilders);
+
+			UInteractiveGizmoBuilder* TransformBuilder = GetTransformGizmoBuilder();
+
+			FoundBuilders.Add(TransformBuilder);
+		}
+
+		TArray<TObjectPtr<UInteractiveGizmo>> NewGizmos;
 
-		for (TObjectPtr<UEditorInteractiveGizmoSelectionBuilder> FoundBuilder : FoundBuilders)
+		for (UInteractiveGizmoBuilder* FoundBuilder : FoundBuilders)
 		{
-			UInteractiveGizmo* NewGizmo = FoundBuilder->BuildGizmoForSelection(CurrentSceneState);
-			if (NewGizmo == nullptr)
+			if (UInteractiveGizmo* NewGizmo = FoundBuilder->BuildGizmo(CurrentSceneState))
 			{
-				DisplayMessage(LOCTEXT("CreateGizmoReturnNullMessage", "UEditorInteractiveGizmoManager::CreateGizmo: CreateGizmoForSelectionSet() returned null"), EToolMessageLevel::Internal);
-				return NewGizmos;
-			}
+				if (UEditorInteractiveGizmoSelectionBuilder* SelectionBuilder = Cast<UEditorInteractiveGizmoSelectionBuilder>(FoundBuilder))
+				{
+					SelectionBuilder->UpdateGizmoForSelection(NewGizmo, CurrentSceneState);
+				}
 
-			// register new active input behaviors
-			InputRouter->RegisterSource(NewGizmo);
+				// register new active input behaviors
+				InputRouter->RegisterSource(NewGizmo);
 
-			NewGizmos.Add(NewGizmo);
+				NewGizmos.Add(NewGizmo);
+			}
 		}
 
 		PostInvalidation();
 
 		for (UInteractiveGizmo* NewGizmo : NewGizmos)
 		{
-			FActiveSelectionGizmo ActiveGizmo = { NewGizmo, nullptr };
-			ActiveSelectionGizmos.Add(ActiveGizmo);
+			FActiveEditorGizmo ActiveGizmo = { NewGizmo, nullptr };
+			ActiveEditorGizmos.Add(ActiveGizmo);
 		}
 
 		return NewGizmos;
@@ -173,11 +129,21 @@ TArray<UInteractiveGizmo*> UEditorInteractiveGizmoManager::CreateGizmosForCurren
 	return TArray<UInteractiveGizmo*>();
 }
 
+UInteractiveGizmoBuilder* UEditorInteractiveGizmoManager::GetTransformGizmoBuilder()
+{
+	UEditorInteractiveGizmoSubsystem* GizmoSubsystem = GEditor->GetEditorSubsystem<UEditorInteractiveGizmoSubsystem>();
+	if (ensure(GizmoSubsystem))
+	{
+		return GizmoSubsystem->GetTransformGizmoBuilder();
+	}
+
+	return nullptr;
+}
 
-bool UEditorInteractiveGizmoManager::DestroySelectionGizmo(UInteractiveGizmo* Gizmo)
+bool UEditorInteractiveGizmoManager::DestroyEditorGizmo(UInteractiveGizmo* Gizmo)
 {
-	auto Pred = [Gizmo](const FActiveSelectionGizmo& ActiveSelectionGizmo) {return ActiveSelectionGizmo.Gizmo == Gizmo; };
-	if (!ensure(ActiveSelectionGizmos.FindByPredicate(Pred)))
+	auto Pred = [Gizmo](const FActiveEditorGizmo& ActiveEditorGizmo) {return ActiveEditorGizmo.Gizmo == Gizmo; };
+	if (!ensure(ActiveEditorGizmos.FindByPredicate(Pred)))
 	{
 		return false;
 	}
@@ -188,69 +154,69 @@ bool UEditorInteractiveGizmoManager::DestroySelectionGizmo(UInteractiveGizmo* Gi
 
 	InputRouter->DeregisterSource(Gizmo);
 
-	ActiveSelectionGizmos.RemoveAll(Pred);
+	ActiveEditorGizmos.RemoveAll(Pred);
 
 	PostInvalidation();
 
 	return true;
 }
 
-void UEditorInteractiveGizmoManager::DestroyAllSelectionGizmos()
+void UEditorInteractiveGizmoManager::DestroyAllEditorGizmos()
 {
-	for (int i = 0; i < ActiveSelectionGizmos.Num(); i++)
+	for (int i = 0; i < ActiveEditorGizmos.Num(); i++)
 	{
-		UInteractiveGizmo* Gizmo = ActiveSelectionGizmos[i].Gizmo;
+		UInteractiveGizmo* Gizmo = ActiveEditorGizmos[i].Gizmo;
 		if (ensure(Gizmo))
 		{
-			DestroySelectionGizmo(Gizmo);
+			DestroyEditorGizmo(Gizmo);
 		}
 	}
 
-	ActiveSelectionGizmos.Reset();
+	ActiveEditorGizmos.Reset();
 
 	PostInvalidation();
 }
 
 void UEditorInteractiveGizmoManager::HandleEditorSelectionSetChanged(const UTypedElementSelectionSet* InSelectionSet)
 {
-	DestroyAllSelectionGizmos();
+	DestroyAllEditorGizmos();
 
 	if (InSelectionSet && InSelectionSet->HasSelectedElements())
 	{
-		CreateGizmosForCurrentSelectionState();
+		CreateGizmosForCurrentSelectionSet();
 	}
 }
 
 // @todo move this to a gizmo context object
-bool UEditorInteractiveGizmoManager::GetShowSelectionGizmos()
+bool UEditorInteractiveGizmoManager::GetShowEditorGizmos()
 {
-	return bShowSelectionGizmos;
+	return bShowEditorGizmos;
 }
 
-bool UEditorInteractiveGizmoManager::GetShowSelectionGizmosForView(IToolsContextRenderAPI* RenderAPI)
+bool UEditorInteractiveGizmoManager::GetShowEditorGizmosForView(IToolsContextRenderAPI* RenderAPI)
 {
 	const bool bEngineShowFlagsModeWidget = (RenderAPI && RenderAPI->GetSceneView() && 
 											 RenderAPI->GetSceneView()->Family &&
 											 RenderAPI->GetSceneView()->Family->EngineShowFlags.ModeWidgets);
-	return bShowSelectionGizmos && bEngineShowFlagsModeWidget;
+	return bShowEditorGizmos && bEngineShowFlagsModeWidget;
 }
 
-void UEditorInteractiveGizmoManager::UpdateActiveSelectionGizmos()
+void UEditorInteractiveGizmoManager::UpdateActiveEditorGizmos()
 {
 	const bool bEditorModeToolsSupportsWidgetDrawing = EditorModeManager ? EditorModeManager->GetShowWidget() : true;
-	const bool bEnableSelectionGizmos = (CVarUseLegacyWidget.GetValueOnGameThread() == 0);
-	const bool bNewShowSelectionGizmos = bEditorModeToolsSupportsWidgetDrawing && bEnableSelectionGizmos;
+	const bool bEnableEditorGizmos = (CVarUseLegacyWidget.GetValueOnGameThread() == 0);
+	const bool bNewShowEditorGizmos = bEditorModeToolsSupportsWidgetDrawing && bEnableEditorGizmos;
 
-	if (bShowSelectionGizmos != bNewShowSelectionGizmos)
+	if (bShowEditorGizmos != bNewShowEditorGizmos)
 	{
-		bShowSelectionGizmos = bNewShowSelectionGizmos;
-		if (bShowSelectionGizmos)
+		bShowEditorGizmos = bNewShowEditorGizmos;
+		if (bShowEditorGizmos)
 		{
-			CreateGizmosForCurrentSelectionState();
+			CreateGizmosForCurrentSelectionSet();
 		}
 		else
 		{
-			DestroyAllSelectionGizmos();
+			DestroyAllEditorGizmos();
 		}
 	}
 }
@@ -259,11 +225,11 @@ void UEditorInteractiveGizmoManager::Tick(float DeltaTime)
 {
 	Super::Tick(DeltaTime);
 
-	UpdateActiveSelectionGizmos();
+	UpdateActiveEditorGizmos();
 
-	for (FActiveSelectionGizmo& ActiveSelectionGizmo : ActiveSelectionGizmos)
+	for (FActiveEditorGizmo& ActiveEditorGizmo : ActiveEditorGizmos)
 	{
-		ActiveSelectionGizmo.Gizmo->Tick(DeltaTime);
+		ActiveEditorGizmo.Gizmo->Tick(DeltaTime);
 	}
 }
 
@@ -272,11 +238,11 @@ void UEditorInteractiveGizmoManager::Render(IToolsContextRenderAPI* RenderAPI)
 {
 	Super::Render(RenderAPI);
 
-	if (GetShowSelectionGizmosForView(RenderAPI))
+	if (GetShowEditorGizmosForView(RenderAPI))
 	{
-		for (FActiveSelectionGizmo& ActiveSelectionGizmo : ActiveSelectionGizmos)
+		for (FActiveEditorGizmo& ActiveEditorGizmo : ActiveEditorGizmos)
 		{
-			ActiveSelectionGizmo.Gizmo->Render(RenderAPI);
+			ActiveEditorGizmo.Gizmo->Render(RenderAPI);
 		}
 	}
 }
@@ -285,11 +251,11 @@ void UEditorInteractiveGizmoManager::DrawHUD(FCanvas* Canvas, IToolsContextRende
 {
 	Super::DrawHUD(Canvas, RenderAPI);
 
-	if (GetShowSelectionGizmosForView(RenderAPI))
+	if (GetShowEditorGizmosForView(RenderAPI))
 	{
-		for (FActiveSelectionGizmo& ActiveSelectionGizmo : ActiveSelectionGizmos)
+		for (FActiveEditorGizmo& ActiveEditorGizmo : ActiveEditorGizmos)
 		{
-			ActiveSelectionGizmo.Gizmo->DrawHUD(Canvas, RenderAPI);
+			ActiveEditorGizmo.Gizmo->DrawHUD(Canvas, RenderAPI);
 		}
 	}
 }
diff --git a/Engine/Source/Editor/Experimental/EditorInteractiveToolsFramework/Private/EditorInteractiveGizmoRegistry.cpp b/Engine/Source/Editor/Experimental/EditorInteractiveToolsFramework/Private/EditorInteractiveGizmoRegistry.cpp
new file mode 100644
--- /dev/null
+++ b/Engine/Source/Editor/Experimental/EditorInteractiveToolsFramework/Private/EditorInteractiveGizmoRegistry.cpp
@@ -0,0 +1,207 @@
+// Copyright Epic Games, Inc. All Rights Reserved.
+
+#include "EditorInteractiveGizmoRegistry.h"
+#include "EditorInteractiveGizmoConditionalBuilder.h"
+#include "EditorInteractiveGizmoSelectionBuilder.h"
+
+
+#define LOCTEXT_NAMESPACE "UEditorInteractiveGizmoRegistry"
+
+DEFINE_LOG_CATEGORY_STATIC(LogEditorInteractiveGizmoRegistry, Log, All);
+
+void UEditorGizmoRegistryCategoryEntry::RegisterGizmoType(UInteractiveGizmoBuilder* InGizmoBuilder)
+{
+	check(InGizmoBuilder);
+
+	if (!InGizmoBuilder->IsA(BaseGizmoBuilderType))
+	{
+		UE_LOG(LogEditorInteractiveGizmoRegistry, Warning,
+			TEXT("%s gizmo builder '%s' of type '%s' does not derive from '%s'! Skipping registration."), *CategoryName, *InGizmoBuilder->GetPathName(), *InGizmoBuilder->GetClass()->GetName(), *BaseGizmoBuilderType->GetName());
+		return;
+	}
+
+	if (GizmoTypes.Contains(InGizmoBuilder))
+	{
+		UE_LOG(LogEditorInteractiveGizmoRegistry, Warning,
+			TEXT("%s gizmo builder type %s has already been registered!"), *InGizmoBuilder->GetPathName());
+		return;
+	}
+
+	GizmoTypes.Add(InGizmoBuilder);
+}
+
+void UEditorGizmoRegistryCategoryEntry::DeregisterGizmoType(UInteractiveGizmoBuilder* InGizmoBuilder)
+{
+	check(InGizmoBuilder);
+
+	if (GizmoTypes.Contains(InGizmoBuilder))
+	{
+		GizmoTypes.Remove(InGizmoBuilder);
+	}
+	else
+	{
+		UE_LOG(LogEditorInteractiveGizmoRegistry, Warning,
+			TEXT("UInteractiveGizmoRegistry::DeregisterEditorGizmoType: gizmo type not found in registry %s for category %d"), *InGizmoBuilder->GetName(), *CategoryName);
+	}
+}
+
+void UEditorGizmoRegistryCategoryEntry::ClearGizmoTypes()
+{
+	GizmoTypes.Reset();
+}
+
+void UEditorGizmoRegistryCategoryEntry_Conditional::RegisterGizmoType(UInteractiveGizmoBuilder* InGizmoBuilder)
+{
+	Super::RegisterGizmoType(InGizmoBuilder);
+
+	GizmoTypes.StableSort(
+		[](UInteractiveGizmoBuilder& A, UInteractiveGizmoBuilder& B) {
+			UEditorInteractiveGizmoConditionalBuilder* AA = Cast<UEditorInteractiveGizmoConditionalBuilder>(&A);
+			UEditorInteractiveGizmoConditionalBuilder* BB = Cast<UEditorInteractiveGizmoConditionalBuilder>(&B);
+			return (AA)->GetPriority() > (BB)->GetPriority();
+		});
+}
+
+UEditorGizmoRegistryCategoryEntry_Primary::UEditorGizmoRegistryCategoryEntry_Primary()
+{
+	CategoryName = TEXT("Primary");
+	BaseGizmoBuilderType = UEditorInteractiveGizmoSelectionBuilder::StaticClass();
+}
+
+void UEditorGizmoRegistryCategoryEntry_Primary::GetQualifiedGizmoBuilders(const FToolBuilderState& InToolBuilderState, TArray<UInteractiveGizmoBuilder*>& FoundBuilders)
+{
+	FEditorGizmoTypePriority FoundPriority = 0;
+
+	if (FoundBuilders.Num() > 0)
+	{
+		if (UEditorInteractiveGizmoConditionalBuilder* Builder = Cast<UEditorInteractiveGizmoConditionalBuilder>(FoundBuilders[0]))
+		{
+			FoundPriority = Builder->GetPriority();
+		}
+	}
+
+	for (TObjectPtr<UInteractiveGizmoBuilder> GizmoBuilder : GizmoTypes)
+	{
+		if (UEditorInteractiveGizmoConditionalBuilder* Builder = Cast<UEditorInteractiveGizmoConditionalBuilder>(GizmoBuilder))
+		{
+			if (Builder->GetPriority() < FoundPriority)
+			{
+				break;
+			}
+
+			if (Builder->SatisfiesCondition(InToolBuilderState))
+			{
+				// Reset found builders since only one primary builder should be buildable.
+				FoundBuilders.Reset();
+				FoundBuilders.Add(Builder);
+				break;
+			}
+		}
+	}
+}
+
+UEditorGizmoRegistryCategoryEntry_Accessory::UEditorGizmoRegistryCategoryEntry_Accessory()
+{
+	CategoryName = TEXT("Accessory");
+	BaseGizmoBuilderType = UEditorInteractiveGizmoSelectionBuilder::StaticClass();
+}
+
+void UEditorGizmoRegistryCategoryEntry_Accessory::GetQualifiedGizmoBuilders(const FToolBuilderState& InToolBuilderState, TArray<UInteractiveGizmoBuilder*>& FoundBuilders)
+{
+	FEditorGizmoTypePriority FoundPriority = 0;
+
+	if (FoundBuilders.Num() > 0)
+	{
+		if (UEditorInteractiveGizmoConditionalBuilder* Builder = Cast<UEditorInteractiveGizmoConditionalBuilder>(FoundBuilders[0]))
+		{
+			FoundPriority = Builder->GetPriority();
+		}
+	}
+
+	for (TObjectPtr<UInteractiveGizmoBuilder> GizmoBuilder : GizmoTypes)
+	{
+		if (UEditorInteractiveGizmoConditionalBuilder* Builder = Cast<UEditorInteractiveGizmoConditionalBuilder>(GizmoBuilder))
+		{
+			if (Builder->GetPriority() < FoundPriority)
+			{
+				break;
+			}
+
+			if (Builder->SatisfiesCondition(InToolBuilderState))
+			{
+				// If priority is greater than previous found priority, reset found builders. If it is the same simply append 
+				// to the builders since more than one accessory builder can be built at a time.
+				if (Builder->GetPriority() > FoundPriority)
+				{
+					FoundBuilders.Reset();
+				}
+				FoundBuilders.Add(Builder);
+			}
+		}
+	}
+}
+
+UEditorInteractiveGizmoRegistry::UEditorInteractiveGizmoRegistry()
+{
+	GizmoCategoryMap.Add(EEditorGizmoCategory::Primary, NewObject<UEditorGizmoRegistryCategoryEntry_Primary>());
+	GizmoCategoryMap.Add(EEditorGizmoCategory::Accessory, NewObject<UEditorGizmoRegistryCategoryEntry_Accessory>());
+}
+
+void UEditorInteractiveGizmoRegistry::Shutdown()
+{
+	ClearEditorGizmoTypes();
+}
+
+void UEditorInteractiveGizmoRegistry::RegisterEditorGizmoType(EEditorGizmoCategory InGizmoCategory, UInteractiveGizmoBuilder* InGizmoBuilder)
+{
+	if (ensure(InGizmoBuilder))
+	{
+		if (UEditorGizmoRegistryCategoryEntry** RegEntry = GizmoCategoryMap.Find(InGizmoCategory))
+		{
+			if (*RegEntry)
+			{
+				(*RegEntry)->RegisterGizmoType(InGizmoBuilder);
+			}
+		}
+	}
+}
+
+void UEditorInteractiveGizmoRegistry::GetQualifiedEditorGizmoBuilders(EEditorGizmoCategory InGizmoCategory, const FToolBuilderState& InToolBuilderState, TArray<UInteractiveGizmoBuilder*>& InFoundBuilders)
+{
+	if (UEditorGizmoRegistryCategoryEntry** RegEntry = GizmoCategoryMap.Find(InGizmoCategory))
+	{
+		if (*RegEntry)
+		{
+			(*RegEntry)->GetQualifiedGizmoBuilders(InToolBuilderState, InFoundBuilders);
+		}
+	}
+}
+
+void UEditorInteractiveGizmoRegistry::DeregisterEditorGizmoType(EEditorGizmoCategory InGizmoCategory, UInteractiveGizmoBuilder* InGizmoBuilder)
+{
+	if (ensure(InGizmoBuilder))
+	{
+		if (UEditorGizmoRegistryCategoryEntry** RegEntry = GizmoCategoryMap.Find(InGizmoCategory))
+		{
+			if (*RegEntry)
+			{
+				(*RegEntry)->DeregisterGizmoType(InGizmoBuilder);
+			}
+		}
+	}
+
+}
+
+void UEditorInteractiveGizmoRegistry::ClearEditorGizmoTypes()
+{
+	for (auto& Elem : GizmoCategoryMap)
+	{
+		if (Elem.Value)
+		{
+			Elem.Value->ClearGizmoTypes();
+		}
+	}
+}
+
+#undef LOCTEXT_NAMESPACE
+
diff --git a/Engine/Source/Editor/Experimental/EditorInteractiveToolsFramework/Private/EditorInteractiveGizmoSelectionBuilder.cpp b/Engine/Source/Editor/Experimental/EditorInteractiveToolsFramework/Private/EditorInteractiveGizmoSelectionBuilder.cpp
new file mode 100644
--- /dev/null
+++ b/Engine/Source/Editor/Experimental/EditorInteractiveToolsFramework/Private/EditorInteractiveGizmoSelectionBuilder.cpp
@@ -0,0 +1,47 @@
+// Copyright Epic Games, Inc. All Rights Reserved.
+
+#include "EditorInteractiveGizmoSelectionBuilder.h"
+#include "BaseGizmos/TransformProxy.h"
+#include "Elements/Interfaces/TypedElementObjectInterface.h"
+#include "Elements/Interfaces/TypedElementWorldInterface.h"
+#include "Engine/Selection.h"
+#include "Engine/World.h"
+
+UTransformProxy* FEditorGizmoSelectionBuilderHelper::CreateTransformProxyForSelection(const FToolBuilderState& SceneState)
+{
+	// @todo - once UTransformProxy supports typed elements, update this to use the normalized typed
+	// element selection set.
+	if (UTypedElementSelectionSet* SelectionSet = SceneState.TypedElementSelectionSet.Get())
+	{
+		if (SelectionSet->GetNumSelectedElements() > 0)
+		{
+			bool bHasSelectedElements = false;
+			UTransformProxy* TransformProxy = NewObject<UTransformProxy>();
+
+			SelectionSet->ForEachSelectedElement<UTypedElementWorldInterface>([TransformProxy, SelectionSet, &bHasSelectedElements](const TTypedElement<UTypedElementWorldInterface>& InWorldElement)
+				{
+					if (InWorldElement.CanMoveElement(ETypedElementWorldType::Editor))
+					{
+						if (TTypedElement<UTypedElementObjectInterface> ObjectTypedElement = SelectionSet->GetElementList()->GetElement<UTypedElementObjectInterface>(InWorldElement))
+						{
+							if (AActor* Actor = ObjectTypedElement.GetObjectAs<AActor>())
+							{
+								USceneComponent* SceneComponent = Actor->GetRootComponent();
+								TransformProxy->AddComponent(SceneComponent);
+								bHasSelectedElements = true;
+							}
+						}
+					}
+					return true;
+				});
+
+			if (bHasSelectedElements)
+			{
+				return TransformProxy;
+			}
+		}
+	}
+
+	return nullptr;
+}
+	
\ No newline at end of file
diff --git a/Engine/Source/Editor/Experimental/EditorInteractiveToolsFramework/Private/EditorInteractiveGizmoSubsystem.cpp b/Engine/Source/Editor/Experimental/EditorInteractiveToolsFramework/Private/EditorInteractiveGizmoSubsystem.cpp
--- a/Engine/Source/Editor/Experimental/EditorInteractiveToolsFramework/Private/EditorInteractiveGizmoSubsystem.cpp
+++ b/Engine/Source/Editor/Experimental/EditorInteractiveToolsFramework/Private/EditorInteractiveGizmoSubsystem.cpp
@@ -3,7 +3,6 @@
 #include "EditorInteractiveGizmoSubsystem.h"
 #include "EditorGizmos/EditorTransformGizmo.h"
 
-
 #define LOCTEXT_NAMESPACE "UEditorInteractiveGizmoSubsystem"
 
 DEFINE_LOG_CATEGORY_STATIC(LogEditorInteractiveGizmoSubsystem, Log, All);
@@ -12,104 +11,61 @@ DEFINE_LOG_CATEGORY_STATIC(LogEditorInteractiveGizmoSubsystem, Log, All);
 UEditorInteractiveGizmoSubsystem::UEditorInteractiveGizmoSubsystem()
 	: UEditorSubsystem()
 {
-
+	Registry = NewObject<UEditorInteractiveGizmoRegistry>();
 }
 
 void UEditorInteractiveGizmoSubsystem::Initialize(FSubsystemCollectionBase& Collection)
 {
 	if (GEngine && GEngine->IsInitialized())
 	{
-		RegisterBuiltinGizmoSelectionTypes();
+		RegisterBuiltinEditorGizmoTypes();
 	}
 	else
 	{
-		FCoreDelegates::OnPostEngineInit.AddUObject(this, &UEditorInteractiveGizmoSubsystem::RegisterBuiltinGizmoSelectionTypes);
+		FCoreDelegates::OnPostEngineInit.AddUObject(this, &UEditorInteractiveGizmoSubsystem::RegisterBuiltinEditorGizmoTypes);
 	}
 }
 
 void UEditorInteractiveGizmoSubsystem::Deinitialize()
 {
-	DeregisterBuiltinGizmoSelectionTypes();
+	check(Registry);
+	DeregisterBuiltinEditorGizmoTypes();
+	Registry->Shutdown();
 }
 
-void UEditorInteractiveGizmoSubsystem::RegisterBuiltinGizmoSelectionTypes()
+void UEditorInteractiveGizmoSubsystem::RegisterBuiltinEditorGizmoTypes()
 {
-	// Register built-in gizmo types here
-	TObjectPtr<UEditorTransformGizmoBuilder> EditorTransformBuilder = NewObject<UEditorTransformGizmoBuilder>();
-	RegisterGizmoSelectionType(EditorTransformBuilder);
+	// Setup gizmo transform builder
+	TransformGizmoBuilder = NewObject<UEditorTransformGizmoBuilder>();
 
-	RegisterEditorGizmoSelectionTypesDelegate.Broadcast();
-}
+	// Register built-in gizmo types here
 
-void UEditorInteractiveGizmoSubsystem::DeregisterBuiltinGizmoSelectionTypes()
-{
-	DeregisterEditorGizmoSelectionTypesDelegate.Broadcast();
-	ClearGizmoSelectionTypeRegistry();
+	RegisterGlobalEditorGizmoTypesDelegate.Broadcast();
 }
 
-void UEditorInteractiveGizmoSubsystem::RegisterGizmoSelectionType(const TObjectPtr<UEditorInteractiveGizmoSelectionBuilder> InGizmoSelectionBuilder)
+void UEditorInteractiveGizmoSubsystem::DeregisterBuiltinEditorGizmoTypes()
 {
-	if (ensure(InGizmoSelectionBuilder != nullptr) == false)
-	{
-		return;
-	}
-
-	if (GizmoSelectionBuilders.Contains(InGizmoSelectionBuilder))
-	{
-		UE_LOG(LogEditorInteractiveGizmoSubsystem, Warning,
-			TEXT("UInteractiveGizmoSubsystem::RegisterGizmoSelectionType: type has already been registered %s"), *InGizmoSelectionBuilder->GetName());
-		return;
-	}
+	DeregisterGlobalEditorGizmoTypesDelegate.Broadcast();
 
-	GizmoSelectionBuilders.Add(InGizmoSelectionBuilder);
-	GizmoSelectionBuilders.StableSort(
-		[](UEditorInteractiveGizmoSelectionBuilder& A, UEditorInteractiveGizmoSelectionBuilder& B) {
-			return (A).GetPriority() > (B).GetPriority();
-		});
+	TransformGizmoBuilder = nullptr;
 }
 
-TArray<TObjectPtr<UEditorInteractiveGizmoSelectionBuilder>> UEditorInteractiveGizmoSubsystem::GetQualifiedGizmoSelectionBuilders(const FToolBuilderState& InToolBuilderState)
+void UEditorInteractiveGizmoSubsystem::RegisterGlobalEditorGizmoType(EEditorGizmoCategory InGizmoCategory, UInteractiveGizmoBuilder* InGizmoBuilder)
 {
-	TArray<TObjectPtr<UEditorInteractiveGizmoSelectionBuilder>> FoundBuilders;
-	FEditorGizmoTypePriority FoundPriority = 0;
-
-	for (TObjectPtr<UEditorInteractiveGizmoSelectionBuilder> Builder : GizmoSelectionBuilders)
-	{
-		if (Builder->GetPriority() < FoundPriority)
-		{
-			break;
-		}
-
-		if (Builder->SatisfiesCondition(InToolBuilderState))
-		{
-			FoundBuilders.Add(Builder);
-			FoundPriority = Builder->GetPriority();
-		}
-	}
-
-	return FoundBuilders;
+	check(Registry);
+	Registry->RegisterEditorGizmoType(InGizmoCategory, InGizmoBuilder);
 }
 
-bool UEditorInteractiveGizmoSubsystem::DeregisterGizmoSelectionType(const TObjectPtr<UEditorInteractiveGizmoSelectionBuilder> InGizmoSelectionBuilder)
+void UEditorInteractiveGizmoSubsystem::GetQualifiedGlobalEditorGizmoBuilders(EEditorGizmoCategory InGizmoCategory, const FToolBuilderState& InToolBuilderState, TArray<UInteractiveGizmoBuilder*>& InFoundBuilders)
 {
-	if (ensure(InGizmoSelectionBuilder != nullptr) == false)
-	{
-		return false;
-	}
-
-	if (GizmoSelectionBuilders.Contains(InGizmoSelectionBuilder) == false)
-	{
-		UE_LOG(LogEditorInteractiveGizmoSubsystem, Warning,
-			TEXT("UInteractiveGizmoSubsystem::DeregisterGizmoSelectionType: type has already been registered %s"), *InGizmoSelectionBuilder->GetName());
-		return false;
-	}
-	GizmoSelectionBuilders.Remove(InGizmoSelectionBuilder);
-	return true;
+	check(Registry);
+	Registry->GetQualifiedEditorGizmoBuilders(InGizmoCategory, InToolBuilderState, InFoundBuilders);
 }
 
-void UEditorInteractiveGizmoSubsystem::ClearGizmoSelectionTypeRegistry()
+void UEditorInteractiveGizmoSubsystem::DeregisterGlobalEditorGizmoType(EEditorGizmoCategory InGizmoCategory, UInteractiveGizmoBuilder* InGizmoBuilder)
 {
-	GizmoSelectionBuilders.Reset();
+	check(Registry);
+	Registry->DeregisterEditorGizmoType(InGizmoCategory, InGizmoBuilder);
 }
 
 #undef LOCTEXT_NAMESPACE
diff --git a/Engine/Source/Editor/Experimental/EditorInteractiveToolsFramework/Public/EditorInteractiveGizmoConditionalBuilder.h b/Engine/Source/Editor/Experimental/EditorInteractiveToolsFramework/Public/EditorInteractiveGizmoConditionalBuilder.h
--- a/Engine/Source/Editor/Experimental/EditorInteractiveToolsFramework/Public/EditorInteractiveGizmoConditionalBuilder.h
+++ b/Engine/Source/Editor/Experimental/EditorInteractiveToolsFramework/Public/EditorInteractiveGizmoConditionalBuilder.h
@@ -54,8 +54,8 @@ struct EDITORINTERACTIVETOOLSFRAMEWORK_API FEditorGizmoTypePriority
 
 /** 
  * Gizmo builders which should be built once a condition is satisfied in the current scene state
- * inherit from UEditorInteractiveGizmoConditionalBuilder or its derived classes. The SatisfiesCondition method
- * should return true whenever the builder is buildable based on Editor scene state, usually based
+ * inherit from UEditorInteractiveGizmoConditionalBuilder. The SatisfiesCondition method
+ * should return true whenever the builder is buildable based on scene state, most commonly based
  * on selection. The builder's priority should be specified such that when more than one builder
  * is discovered, the builder with highest priority will be built.
  */
diff --git a/Engine/Source/Editor/Experimental/EditorInteractiveToolsFramework/Public/EditorInteractiveGizmoManager.h b/Engine/Source/Editor/Experimental/EditorInteractiveToolsFramework/Public/EditorInteractiveGizmoManager.h
--- a/Engine/Source/Editor/Experimental/EditorInteractiveToolsFramework/Public/EditorInteractiveGizmoManager.h
+++ b/Engine/Source/Editor/Experimental/EditorInteractiveToolsFramework/Public/EditorInteractiveGizmoManager.h
@@ -3,18 +3,17 @@
 #pragma once
 
 #include "CoreMinimal.h"
+#include "EditorInteractiveGizmoRegistry.h"
 #include "InteractiveGizmoManager.h"
 #include "EditorInteractiveGizmoManager.generated.h"
 
 class FEditorModeTools;
-class UEditorInteractiveGizmoSelectionBuilder;
 class UEdModeInteractiveToolsContext;
 class UInteractiveGizmo;
 class IToolsContextRenderAPI;
-class UTypedElementSelectionSet;
 
 USTRUCT()
-struct FActiveSelectionGizmo
+struct FActiveEditorGizmo
 {
 	GENERATED_BODY()
 
@@ -53,32 +52,37 @@ public:
 	virtual void DrawHUD(FCanvas* Canvas, IToolsContextRenderAPI* RenderAPI) override;
 
 	/**
-	 * Register a new selection gizmo type
-	 * @param InGizmoSelectionBuilder new auto gizmo builder
+	 * Register a new Editor gizmo type.
+	 * @param InGizmoCategory category in which to register gizmo builder
+	 * @param InGizmonBuilder new Editor gizmo builder
+	 * - Accessory gizmo builders must be inherited from UEditorInteractiveGizmoAccessoryBuilder.
+	 * - Primary gizmo builders must be inherited from UEditorInteractiveGizmoPrimaryBuilder.
 	 */
-	void RegisterGizmoSelectionType(const TObjectPtr<UEditorInteractiveGizmoSelectionBuilder> InGizmoSelectionBuilder);
+	void RegisterEditorGizmoType(EEditorGizmoCategory InGizmoCategory, UInteractiveGizmoBuilder* InGizmoBuilder);
 
 	/**
-	* Remove a gizmo auto type from the set of known gizmo auto types
-	* @param InGizmoSelectionBuilder same object pointer that was passed to RegisterGizmoSelectionType()
+	* Remove an Editor gizmo type from the set of known Editor gizmo types
+	* @param InGizmoBuilder same object pointer that was passed to RegisterEditorGizmoType()
 	* @return true if gizmo type was found and deregistered
 	*/
-	bool DeregisterGizmoSelectionType(const TObjectPtr<UEditorInteractiveGizmoSelectionBuilder> InGizmoSelectionBuilder);
+	void DeregisterEditorGizmoType(EEditorGizmoCategory InGizmoCategory, UInteractiveGizmoBuilder* InGizmoBuilder);
 
 	/**
-	 * Get all qualified gizmo auto builders based on the current state. Qualification is determined by the gizmo builder
+	 * Get all qualified Editor gizmo builders for the specified category, based on the current state. Qualification is determined by the gizmo builder
 	 * returning true from SatisfiesCondition() and relative priority. All qualified builders at the highest found priority
 	 * will be returned.
-	 * @return array of qualified Gizmo auto builders based on current state
+	 * @param InGizmoCategory category in which to search for qualified builders
+	 * @param InToolBuilderState current selection and other state
+	 * @return array of qualified Gizmo selection builders based on current state
 	 */
-	virtual TArray<TObjectPtr<UEditorInteractiveGizmoSelectionBuilder>> GetQualifiedGizmoSelectionBuilders(const FToolBuilderState& InToolBuilderState);
+	void GetQualifiedEditorGizmoBuilders(EEditorGizmoCategory InGizmoCategory, const FToolBuilderState& InToolBuilderState, TArray<UInteractiveGizmoBuilder*>& InFoundBuilders);
 
 	/** 
-	 * Set how auto gizmo resolution should occur when CreateSelectionGizmo is invoked. If bSearchLocalOnly is true, only the current
+	 * Set how auto gizmo resolution should occur when CreateGizmosForCurrentState is invoked. If bSearchLocalOnly is true, only the current
 	 * @param bLocalOnly - if true, only the current gizmo manager registry will be searched for candidate gizmos. If false,
 	 *   both the gizmo manager registry and any higher gizmo manager or gizmo subsystem (in the case of selection builders) will be searched
 	 */
-	virtual void SetGizmoSelectionBuilderResolution(bool bLocalOnly)
+	virtual void SetEditorGizmoBuilderResolution(bool bLocalOnly)
 	{
 		bSearchLocalBuildersOnly = bLocalOnly;
 	}
@@ -86,16 +90,22 @@ public:
 	/**
 	 * Returns the current auto gizmo resolution setting 
 	 */
-	virtual bool GetGizmoSelectionBuilderResolution() const
+	virtual bool GetEditorGizmoBuilderResolution() const
 	{
 		return bSearchLocalBuildersOnly;
 	}
 
+	/**
+	 * Try to automatically activate a new Gizmo instance based on the current selection set state
+	 * @return array of new Gizmo instances that have been created and initialized
+	 */
+	virtual TArray<UInteractiveGizmo*> CreateGizmosForCurrentSelectionSet();
+
 	/**
 	 * Try to automatically activate a new Gizmo instance based on the current state
 	 * @return array of new Gizmo instances that have been created and initialized
 	 */
-	virtual TArray<UInteractiveGizmo*> CreateGizmosForCurrentSelectionState();
+	virtual UInteractiveGizmoBuilder* GetTransformGizmoBuilder();
 
 	/**
 	 * Handle Editor selection changes
@@ -108,12 +118,12 @@ public:
 	 * @param Gizmo the Gizmo to shutdown and remove
 	 * @return true if the Gizmo was found and removed
 	 */
-	virtual bool DestroySelectionGizmo(UInteractiveGizmo* Gizmo);
+	virtual bool DestroyEditorGizmo(UInteractiveGizmo* Gizmo);
 
 	/**
 	 * Shutdown and remove all active auto gizmos
 	 */
-	virtual void DestroyAllSelectionGizmos();
+	virtual void DestroyAllEditorGizmos();
 
 protected:
 
@@ -121,36 +131,38 @@ protected:
 	 * Returns true if selection gizmos should be visible. 
 	 * @todo move this to a gizmo context object
 	 */
-	virtual bool GetShowSelectionGizmos();
+	virtual bool GetShowEditorGizmos();
 
 	/**
 	 * Returns true if gizmos should be visible based on the current view's engine show flag.
 	 * @todo move this to a gizmo context object
 	 */
-	virtual bool GetShowSelectionGizmosForView(IToolsContextRenderAPI* RenderAPI);
+	virtual bool GetShowEditorGizmosForView(IToolsContextRenderAPI* RenderAPI);
 
 	/**
 	 * Updates active selection gizmos when show selection state changes
 	 */
-	void UpdateActiveSelectionGizmos();
+	void UpdateActiveEditorGizmos();
 
-protected:
 
-	/** set of Currently-active Gizmos */
+	/** Actual registry */
 	UPROPERTY()
-	TArray<FActiveSelectionGizmo> ActiveSelectionGizmos;
+	TObjectPtr<UEditorInteractiveGizmoRegistry> Registry;
 
-	/** Current set of GizmoSelectionBuilders */
+	/** set of Currently-active Gizmos */
 	UPROPERTY()
-	TArray<TObjectPtr<UEditorInteractiveGizmoSelectionBuilder>> GizmoSelectionBuilders;
+	TArray<FActiveEditorGizmo> ActiveEditorGizmos;
 
 	/** If false, only search gizmo builders in current gizmo manager. If true, also search gizmo subsystem */
 	bool bSearchLocalBuildersOnly = false;
 
-private:
-	/** @todo: remove when GetShowSelectionGizmos() is moved to gizmo context object */
+	/** Cache for already built gizmos, currently this only caches the transform */
+	UPROPERTY()
+	TMap<TObjectPtr<UInteractiveGizmoBuilder>, TObjectPtr<UInteractiveGizmo>> CachedGizmoMap;
+
+	/** @todo: remove when GetShowEditorGizmos() is moved to gizmo context object */
 	FEditorModeTools* EditorModeManager = nullptr;
 
-	/** Whether selection gizmos are enabled. UpdateActiveSelectionGizmos() determines this value each tick and updates if it has changed. */
-	bool bShowSelectionGizmos = false;
+	/** Whether Editor gizmos are enabled. UpdateActiveEditorGizmos() determines this value each tick and updates if it has changed. */
+	bool bShowEditorGizmos = false;
 };
diff --git a/Engine/Source/Editor/Experimental/EditorInteractiveToolsFramework/Public/EditorInteractiveGizmoRegistry.h b/Engine/Source/Editor/Experimental/EditorInteractiveToolsFramework/Public/EditorInteractiveGizmoRegistry.h
new file mode 100644
--- /dev/null
+++ b/Engine/Source/Editor/Experimental/EditorInteractiveToolsFramework/Public/EditorInteractiveGizmoRegistry.h
@@ -0,0 +1,139 @@
+// Copyright Epic Games, Inc. All Rights Reserved.
+
+#pragma once
+
+#include "CoreMinimal.h"
+#include "Delegates/Delegate.h"
+#include "EditorInteractiveGizmoSelectionBuilder.h"
+#include "EditorInteractiveGizmoRegistry.generated.h"
+
+/** Gizmo category used for registering Editor gizmo selection-based builders */
+UENUM()
+enum class EEditorGizmoCategory :uint8
+{
+	/** Accessory gizmos, built simultaneously with the Level Editor TRS gizmo. */
+	Accessory,
+
+	/** Primary gizmos, built in place of the Level Editor TRS gizmo.  */
+	Primary
+};
+
+UCLASS()
+class UEditorGizmoRegistryCategoryEntry : public UObject
+{
+	GENERATED_BODY()
+public:
+
+	/** 
+	 * Gets qualified gizmo builders, replacing those already found if necessary.
+	 * @param FoundBuilders inputs qualified builders found so far, and outputs qualified builders for this category
+	 *                      replacing 
+	 */
+	virtual void GetQualifiedGizmoBuilders(const FToolBuilderState& InToolBuilderState, TArray<UInteractiveGizmoBuilder*>& InFoundBuilders) PURE_VIRTUAL(UEditorGizmoRegistryCategoryEntry::GetQualifiedGizmoBuilders, return;);
+
+	virtual void RegisterGizmoType(UInteractiveGizmoBuilder* InGizmoBuilder);
+
+	virtual void DeregisterGizmoType(UInteractiveGizmoBuilder* InGizmoBuilder);
+
+	virtual void ClearGizmoTypes();
+
+	UPROPERTY()
+	TArray<TObjectPtr<UInteractiveGizmoBuilder>> GizmoTypes;
+
+protected:
+
+	UPROPERTY()
+	FString CategoryName;
+
+	UPROPERTY()
+	UClass* BaseGizmoBuilderType;
+};
+
+UCLASS()
+class UEditorGizmoRegistryCategoryEntry_Conditional : public UEditorGizmoRegistryCategoryEntry
+{
+	GENERATED_BODY()
+public:
+
+	virtual void RegisterGizmoType(UInteractiveGizmoBuilder* InGizmoBuilder);
+};
+
+
+UCLASS()
+class UEditorGizmoRegistryCategoryEntry_Primary : public UEditorGizmoRegistryCategoryEntry_Conditional
+{
+	GENERATED_BODY()
+public:
+	UEditorGizmoRegistryCategoryEntry_Primary();
+
+	virtual void GetQualifiedGizmoBuilders(const FToolBuilderState& InToolBuilderState, TArray<UInteractiveGizmoBuilder*>& FoundBuilders);
+};
+
+UCLASS()
+class UEditorGizmoRegistryCategoryEntry_Accessory : public UEditorGizmoRegistryCategoryEntry_Conditional
+{
+	GENERATED_BODY()
+public:
+	UEditorGizmoRegistryCategoryEntry_Accessory();
+
+	virtual void GetQualifiedGizmoBuilders(const FToolBuilderState& InToolBuilderState, TArray<UInteractiveGizmoBuilder*>& FoundBuilders);
+};
+
+/**
+ * Gizmo types should be registered in either UEditorInteractiveGizmoSubsystem or
+ * UEditorInteractiveGizmoManager. This registry class is used internally by the
+ * subsystem and manager which each maintain its own registry at different scopes:
+ * the subystem is global to the Editor, the manager is local to the Interactive
+ * Tools Context.
+ */
+UCLASS()
+class EDITORINTERACTIVETOOLSFRAMEWORK_API UEditorInteractiveGizmoRegistry : public UObject
+{
+	GENERATED_BODY()
+
+public:
+
+	UEditorInteractiveGizmoRegistry();
+
+	/** Shutdown the registry, called by the gizmo subsystem and gizmo manager when they are shutdown/deinitialized. */
+	virtual void Shutdown();
+
+	/**
+	 * Register a new Editor gizmo type.
+	 * @param InGizmoCategory category in which to register gizmo builder
+	 * @param InGizmonBuilder new Editor gizmo builder
+	 * - Accessory gizmo builders must be inherited from UEditorInteractiveGizmoAccessoryBuilder.
+	 * - Primary gizmo builders must be inherited from UEditorInteractiveGizmoPrimaryBuilder.
+	 */
+	void RegisterEditorGizmoType(EEditorGizmoCategory InGizmoCategory, UInteractiveGizmoBuilder* InGizmoBuilder);
+
+	/**
+	* Remove an Editor gizmo type from the set of known Editor gizmo types
+	* @param InGizmoBuilder same object pointer that was passed to RegisterEditorGizmoType()
+	* @return true if gizmo type was found and deregistered
+	*/
+	void DeregisterEditorGizmoType(EEditorGizmoCategory InGizmoCategory, UInteractiveGizmoBuilder* InGizmoBuilder);
+
+	/**
+	 * Clear all registered gizmo types
+	 */
+	void ClearEditorGizmoTypes();
+
+	/**
+	 * Get all qualified Editor gizmo builders for the specified category, based on the current state. Qualification is determined by the gizmo builder
+	 * returning true from SatisfiesCondition() and relative priority. All qualified builders at the highest found priority
+	 * will be returned.
+	 * @param InGizmoCategory category in which to search for qualified builders
+	 * @param InToolBuilderState current selection and other state
+	 * @param FoundBuilders input and output parameter, inputs currently qualified builders, outputs resulting qualified 
+	 *                      builders which may add to or replace the input found builders.
+	 * @return array of qualified Gizmo selection builders based on current state
+	 */
+	void GetQualifiedEditorGizmoBuilders(EEditorGizmoCategory InGizmoCategory, const FToolBuilderState& InToolBuilderState, TArray<UInteractiveGizmoBuilder*>& FoundBuilders);
+
+private:
+
+	/** Current set of Gizmo Builders */
+	UPROPERTY()
+	TMap<EEditorGizmoCategory, UEditorGizmoRegistryCategoryEntry*> GizmoCategoryMap;
+};
\ No newline at end of file
diff --git a/Engine/Source/Editor/Experimental/EditorInteractiveToolsFramework/Public/EditorInteractiveGizmoSelectionBuilder.h b/Engine/Source/Editor/Experimental/EditorInteractiveToolsFramework/Public/EditorInteractiveGizmoSelectionBuilder.h
--- a/Engine/Source/Editor/Experimental/EditorInteractiveToolsFramework/Public/EditorInteractiveGizmoSelectionBuilder.h
+++ b/Engine/Source/Editor/Experimental/EditorInteractiveToolsFramework/Public/EditorInteractiveGizmoSelectionBuilder.h
@@ -4,42 +4,32 @@
 
 #include "CoreMinimal.h"
 #include "EditorInteractiveGizmoConditionalBuilder.h"
-#include "InteractiveGizmo.h"
 #include "ToolContextInterfaces.h"
 #include "EditorInteractiveGizmoSelectionBuilder.generated.h"
 
-/** 
- * UEditorInteractiveGizmoSelectionBuilder provides a method for building and updating gizmos based on the current Editor selection 
- * and state. Builders derived from this class may be registered in one of the following places: 
- *   1) the gizmo subsystem if the gizmo should be available throughout the Editor.
- *   2) the gizmo manager if the gizmo is only used in a particular ed mode or in an asset editor.
- */
+class UTransformProxy;
+
+class EDITORINTERACTIVETOOLSFRAMEWORK_API FEditorGizmoSelectionBuilderHelper
+{
+public:
+	/**
+	 * Utility method that creates a transform proxy based on the current selection.
+	 */
+	static UTransformProxy* CreateTransformProxyForSelection(const FToolBuilderState& SceneState);
+};
+
 UCLASS(Transient, Abstract)
 class EDITORINTERACTIVETOOLSFRAMEWORK_API UEditorInteractiveGizmoSelectionBuilder : public UEditorInteractiveGizmoConditionalBuilder
 {
 	GENERATED_BODY()
 
 public:
-	/**
-	 * Build gizmo for the current Editor selection and state. The Editor gizmo manager calls this method to construct gizmos for 
-	 * for the current selection. This implementation calls BuildGizmo() then UpdateGizmoForSelection(). But derived classes may provide 
-	 * their own implementation which is expected to both build gizmo and set it up to manipulate the current selection. Note that 
-	 * when the selection changes, the gizmo manager may reuse the gizmo and only call UpdateGizmoForSelection() on the existing gizmo.
-	 */
-	virtual UInteractiveGizmo* BuildGizmoForSelection(const FToolBuilderState& SceneState)
-	{
-		if (UInteractiveGizmo* Gizmo = BuildGizmo(SceneState))
-		{
-			UpdateGizmoForSelection(Gizmo, SceneState);
-			return Gizmo;
-		}
-		return nullptr;
-	}
 
 	/**
-	 * Update the input gizmo's active target based on the current Editor selection and scene state.  Derived implementations
-	 * of this method should create a transform proxy for the current Editor selection and sets the gizmo's active target to the new transform proxy. 
-	 * The gizmo manager calls this method when reusing a gizmo, to update the gizmo for the current selection.
+	 * Update gizmo's active target based on the current Editor selection and scene state.  This method creates 
+	 * a transform proxy for the current selection and sets the gizmo's active target to the new transform proxy. 
+	 * This method is called after a gizmo is automatically built based upon selection and also to update the 
+	 * existing gizmo when the selection changes.
 	 */
 	virtual void UpdateGizmoForSelection(UInteractiveGizmo* Gizmo, const FToolBuilderState& SceneState) PURE_VIRTUAL(UEditorInteractiveGizmoSelectionBuilder::UpdateGizmoForSelection, return;);
 
diff --git a/Engine/Source/Editor/Experimental/EditorInteractiveToolsFramework/Public/EditorInteractiveGizmoSubsystem.h b/Engine/Source/Editor/Experimental/EditorInteractiveToolsFramework/Public/EditorInteractiveGizmoSubsystem.h
--- a/Engine/Source/Editor/Experimental/EditorInteractiveToolsFramework/Public/EditorInteractiveGizmoSubsystem.h
+++ b/Engine/Source/Editor/Experimental/EditorInteractiveToolsFramework/Public/EditorInteractiveGizmoSubsystem.h
@@ -6,21 +6,28 @@
 #include "Delegates/Delegate.h"
 #include "EditorSubsystem.h"
 #include "EditorInteractiveGizmoSelectionBuilder.h"
-
+#include "EditorInteractiveGizmoRegistry.h"
 #include "EditorInteractiveGizmoSubsystem.generated.h"
 
 class FSubsystemCollectionBase;
 
 /**
  * The InteractiveGizmoSubsystem provides methods for registering and unregistering 
- * selection-based gizmos builders. Editor gizmo managers which are not marked local-only, 
- * will query this subsystem for qualified builders based on the current selection.
+ * selection-based gizmo builders. This subsystem will be queried for qualified 
+ * builders based on the current selection.
  *
- * This subsystem should also be used to register gizmo selection builders from plugins
- * by binding to the delegates returned from OnEditorGizmoSubsystemRegisterGizmoSelectionTypes() and 
- * OnEditorGizmoSubsystemDeregisterGizmoSelectionTypes().
+ * This subsystem should be used to register gizmo selection-based builders which are not specific
+ * to an ed mode or asset editor. 
+ * For gizmo selection-based builders which are specific to an ed mode or asset editor,
+ * register with the UEditorinteractiveGizmoManager instead, when the ed mode or asset editor
+ * starts up (and deregister when the mode or asset editor shuts down).
+ * 
+ * Plugins registering gizmo types should bind to the delegates returned by:
+ * - OnEditorGizmoSubsystemRegisterEditorGizmoTypes()
+ * - OnEditorGizmoSubsystemDeregisterEditorGizmoTypes()
+ * to register and dergister their gizmo builders.
  */
-UCLASS()
+ UCLASS()
 class EDITORINTERACTIVETOOLSFRAMEWORK_API UEditorInteractiveGizmoSubsystem : public UEditorSubsystem
 {
 	GENERATED_BODY()
@@ -38,62 +45,73 @@ public:
 	/**
 	 * Event which is broadcast just after default types are registered in the gizmo subsystem
 	 */
-	DECLARE_EVENT(UEditorInteractiveSelectionGizmoSubsystem, FOnEditorGizmoSubsystemRegisterGizmoSelectionTypes);
-	FOnEditorGizmoSubsystemRegisterGizmoSelectionTypes& OnEditorGizmoSubsystemRegisterGizmoSelectionTypes() { return RegisterEditorGizmoSelectionTypesDelegate; }
+	DECLARE_EVENT(UEditorInteractiveSelectionGizmoSubsystem, FOnEditorGizmoSubsystemRegisterGlobalEditorGizmoTypes);
+	FOnEditorGizmoSubsystemRegisterGlobalEditorGizmoTypes& OnEditorGizmoSubsystemRegisterGlobalEditorGizmoTypes() { return RegisterGlobalEditorGizmoTypesDelegate; }
 
 	/**
 	 * Event which is broadcast just before default types are deregistered in the gizmo subsystem
 	 */
-	DECLARE_EVENT(UEditorInteractiveSelectionGizmoSubsystem, FOnEditorGizmoSubsystemDeregisterGizmoSelectionTypes);
-	FOnEditorGizmoSubsystemDeregisterGizmoSelectionTypes& OnEditorGizmoSubsystemDeregisterGizmoSelectionTypes() { return DeregisterEditorGizmoSelectionTypesDelegate; }
+	DECLARE_EVENT(UEditorInteractiveSelectionGizmoSubsystem, FOnEditorGizmoSubsystemDeregisterGlobalEditorGizmoTypes);
+	FOnEditorGizmoSubsystemDeregisterGlobalEditorGizmoTypes& OnEditorGizmoSubsystemDeregisterGlobalEditorGizmoTypes() { return DeregisterGlobalEditorGizmoTypesDelegate; }
 
 	/**
-	 * Registers all built-in gizmo selection types and broadcast registration event.
+	 * Registers all built-in Editor gizmo types and broadcast registration event.
 	 */
-	void RegisterBuiltinGizmoSelectionTypes();
+	void RegisterBuiltinEditorGizmoTypes();
 
 	/**
-	 * Removes all built-in gizmo selection types and broadcast deregistration event.
+	 * Removes all built-in Editor gizmo types and broadcast deregistration event.
 	 */
-	void DeregisterBuiltinGizmoSelectionTypes();
+	void DeregisterBuiltinEditorGizmoTypes();
 
 	/**
-	 * Register a new selection gizmo type
-	 * @param InGizmoSelectionBuilder new selection gizmo builder
+	 * Register a new Editor gizmo type which will be global to the Editor.
+	 * @param InGizmoCategory category in which to register gizmo builder
+	 * @param InGizmonBuilder new Editor gizmo builder
+	 * - Accessory gizmo builders must be inherited from UEditorInteractiveGizmoAccessoryBuilder.
+	 * - Primary gizmo builders must be inherited from UEditorInteractiveGizmoPrimaryBuilder.
 	 */
-	void RegisterGizmoSelectionType(const TObjectPtr<UEditorInteractiveGizmoSelectionBuilder> InGizmoSelectionBuilder);
+	void RegisterGlobalEditorGizmoType(EEditorGizmoCategory InGizmoCategory, UInteractiveGizmoBuilder* InGizmoBuilder);
 
 	/**
-	* Remove a gizmo selection type from the set of known gizmo selection types
-	* @param InGizmoSelectionBuilder same object pointer that was passed to RegisterGizmoSelectionType()
+	* Remove an Editor gizmo type from the set of known Editor gizmo types
+	* @param InGizmoBuilder same object pointer that was passed to RegisterEditorGizmoType()
 	* @return true if gizmo type was found and deregistered
 	*/
-	bool DeregisterGizmoSelectionType(const TObjectPtr<UEditorInteractiveGizmoSelectionBuilder> InGizmoSelectionBuilder);
-
-	/**
-	 * Clear all registered gizmo types
-	 */
-	void ClearGizmoSelectionTypeRegistry();
+	void DeregisterGlobalEditorGizmoType(EEditorGizmoCategory InGizmoCategory, UInteractiveGizmoBuilder* InGizmoBuilder);
 
 	/**
-	 * Get all qualified gizmo selection builders based on the current state. Qualification is determined by the gizmo builder
+	 * Get all qualified Editor gizmo builders for the specified category, based on the current state. Qualification is determined by the gizmo builder
 	 * returning true from SatisfiesCondition() and relative priority. All qualified builders at the highest found priority
 	 * will be returned.
+	 * @param InGizmoCategory category in which to search for qualified builders
+	 * @param InToolBuilderState current selection and other state
 	 * @return array of qualified Gizmo selection builders based on current state
 	 */
-	virtual TArray<TObjectPtr<UEditorInteractiveGizmoSelectionBuilder>> GetQualifiedGizmoSelectionBuilders(const FToolBuilderState& InToolBuilderState);
+	void GetQualifiedGlobalEditorGizmoBuilders(EEditorGizmoCategory InGizmoCategory, const FToolBuilderState& InToolBuilderState, TArray<UInteractiveGizmoBuilder*>& InFoundBuilders);
 
+	/** 
+	 * Get transform gizmo builder used to build the Level Editor TRS gizmo. 
+	 */
+	UEditorInteractiveGizmoSelectionBuilder* GetTransformGizmoBuilder()
+	{
+		return TransformGizmoBuilder;
+	}
 
 private:
 
-	/** Current set of GizmoSelectionBuilders */
+	/** TRS gizmo builder */
+	UPROPERTY()
+	TObjectPtr<UEditorInteractiveGizmoSelectionBuilder> TransformGizmoBuilder;
+
+	/** Actual registry */
 	UPROPERTY()
-	TArray<TObjectPtr<UEditorInteractiveGizmoSelectionBuilder> > GizmoSelectionBuilders;
+	TObjectPtr<UEditorInteractiveGizmoRegistry> Registry;
 
 	/** Call to register gizmo types */
-	FOnEditorGizmoSubsystemRegisterGizmoSelectionTypes RegisterEditorGizmoSelectionTypesDelegate;
+	FOnEditorGizmoSubsystemRegisterGlobalEditorGizmoTypes RegisterGlobalEditorGizmoTypesDelegate;
 
 	/** Call to deregister gizmo types */
-	FOnEditorGizmoSubsystemDeregisterGizmoSelectionTypes DeregisterEditorGizmoSelectionTypesDelegate;
+	FOnEditorGizmoSubsystemDeregisterGlobalEditorGizmoTypes DeregisterGlobalEditorGizmoTypesDelegate;
 
 };
\ No newline at end of file
