Fix non-unity compile error.

diff --git a/Engine/Source/Editor/Experimental/EditorInteractiveToolsFramework/Public/EditorInteractiveGizmoRegistry.h b/Engine/Source/Editor/Experimental/EditorInteractiveToolsFramework/Public/EditorInteractiveGizmoRegistry.h
--- a/Engine/Source/Editor/Experimental/EditorInteractiveToolsFramework/Public/EditorInteractiveGizmoRegistry.h
+++ b/Engine/Source/Editor/Experimental/EditorInteractiveToolsFramework/Public/EditorInteractiveGizmoRegistry.h
@@ -18,7 +18,7 @@ enum class EEditorGizmoCategory :uint8
 	Primary
 };
 
-UCLASS()
+UCLASS(Abstract)
 class UEditorGizmoRegistryCategoryEntry : public UObject
 {
 	GENERATED_BODY()
@@ -29,7 +29,7 @@ public:
 	 * @param FoundBuilders inputs qualified builders found so far, and outputs qualified builders for this category
 	 *                      replacing 
 	 */
-	virtual void GetQualifiedGizmoBuilders(const FToolBuilderState& InToolBuilderState, TArray<UInteractiveGizmoBuilder*>& InFoundBuilders) PURE_VIRTUAL(UEditorGizmoRegistryCategoryEntry::GetQualifiedGizmoBuilders, return;);
+	virtual void GetQualifiedGizmoBuilders(const FToolBuilderState& InToolBuilderState, TArray<UInteractiveGizmoBuilder*>& InFoundBuilders) PURE_VIRTUAL(UEditorGizmoRegistryCategoryEntry::GetQualifiedGizmoBuilders, );
 
 	virtual void RegisterGizmoType(UInteractiveGizmoBuilder* InGizmoBuilder);
 
@@ -49,12 +49,14 @@ protected:
 	UClass* BaseGizmoBuilderType;
 };
 
-UCLASS()
+UCLASS(Abstract)
 class UEditorGizmoRegistryCategoryEntry_Conditional : public UEditorGizmoRegistryCategoryEntry
 {
 	GENERATED_BODY()
 public:
 
+	virtual void GetQualifiedGizmoBuilders(const FToolBuilderState& InToolBuilderState, TArray<UInteractiveGizmoBuilder*>& InFoundBuilders) PURE_VIRTUAL(UEditorGizmoRegistryCategoryEntry::GetQualifiedGizmoBuilders, );
+
 	virtual void RegisterGizmoType(UInteractiveGizmoBuilder* InGizmoBuilder);
 };
 
