RemoteExecution: Change array usage to use uint8 instead of char to match the data types used by other engine FileIO
Use FCriticalSection\FScopeLock instead of mutex for CompletionQueue

#rnx

diff --git a/Engine/Source/Editor/Experimental/BazelExecutor/Private/BazelCompletionQueueRunnable.cpp b/Engine/Source/Editor/Experimental/BazelExecutor/Private/BazelCompletionQueueRunnable.cpp
--- a/Engine/Source/Editor/Experimental/BazelExecutor/Private/BazelCompletionQueueRunnable.cpp
+++ b/Engine/Source/Editor/Experimental/BazelExecutor/Private/BazelCompletionQueueRunnable.cpp
@@ -15,6 +15,7 @@ THIRD_PARTY_INCLUDES_END
 
 FBazelCompletionQueueRunnable::FBazelCompletionQueueRunnable()
 	: CompletionQueue(new grpc::CompletionQueue())
+	, QueuedItemsCriticalSection(new FCriticalSection())
 {
 }
 
@@ -24,9 +25,11 @@ FBazelCompletionQueueRunnable::~FBazelCompletionQueueRunnable()
 
 void FBazelCompletionQueueRunnable::ProcessNext(void* Tag, bool Ok)
 {
-	QueuedItemsMutex.lock();
-	FQueuedItem* Item = QueuedItems.Find(Tag);
-	QueuedItemsMutex.unlock();
+	FQueuedItem* Item;
+	{
+		FScopeLock Lock(QueuedItemsCriticalSection.Get());
+		Item = QueuedItems.Find(Tag);
+	}
 	if (!Item)
 	{
 		return;
@@ -80,9 +83,10 @@ void FBazelCompletionQueueRunnable::ProcessNext(void* Tag, bool Ok)
 			Item->Finish(Tag, Ok, *Item->Status.Get(), *Item->Message.Get());
 		}
 
-		QueuedItemsMutex.lock();
-		QueuedItems.Remove(Tag);
-		QueuedItemsMutex.unlock();
+		{
+			FScopeLock Lock(QueuedItemsCriticalSection.Get());
+			QueuedItems.Remove(Tag);
+		}
 		break;
 	}
 	}
@@ -125,8 +129,8 @@ void FBazelCompletionQueueRunnable::Exit()
 	bool Ok;
 	while (CompletionQueue->Next(&Tag, &Ok)) {}
 
-	QueuedItemsMutex.lock();
-	for (const TPair<void*,FQueuedItem>& Pair : QueuedItems)
+	FScopeLock Lock(QueuedItemsCriticalSection.Get());
+	for (const TPair<void*, FQueuedItem>& Pair : QueuedItems)
 	{
 		const FQueuedItem& Item = Pair.Value;
 		if (Item.Finish)
@@ -136,7 +140,6 @@ void FBazelCompletionQueueRunnable::Exit()
 	}
 
 	QueuedItems.Empty();
-	QueuedItemsMutex.unlock();
 }
 
 FSingleThreadRunnable* FBazelCompletionQueueRunnable::GetSingleThreadInterface()
@@ -189,9 +192,12 @@ bool FBazelCompletionQueueRunnable::AddAsyncOperation(
 	Item.StartCall = MoveTemp(OnStartCall);
 	Item.Read = MoveTemp(OnRead);
 	Item.Finish = MoveTemp(OnFinish);
-	QueuedItemsMutex.lock();
-	QueuedItems.Add(AsyncReader, MoveTemp(Item));
-	QueuedItemsMutex.unlock();
+	
+	{
+		FScopeLock Lock(QueuedItemsCriticalSection.Get());
+		QueuedItems.Add(AsyncReader, MoveTemp(Item));
+	}
+		
 	AsyncReader->StartCall(AsyncReader);
 	return true;
 }
diff --git a/Engine/Source/Editor/Experimental/BazelExecutor/Private/BazelCompletionQueueRunnable.h b/Engine/Source/Editor/Experimental/BazelExecutor/Private/BazelCompletionQueueRunnable.h
--- a/Engine/Source/Editor/Experimental/BazelExecutor/Private/BazelCompletionQueueRunnable.h
+++ b/Engine/Source/Editor/Experimental/BazelExecutor/Private/BazelCompletionQueueRunnable.h
@@ -6,8 +6,7 @@
 #include "HAL/Runnable.h"
 #include "HAL/RunnableThread.h"
 #include "Misc/SingleThreadRunnable.h"
-
-#include <mutex>
+#include "Misc/ScopeLock.h"
 
 namespace grpc
 {
@@ -62,7 +61,7 @@ private:
 	bool Running;
 	TMap<void*, FQueuedItem> QueuedItems;
 	TUniquePtr<grpc::CompletionQueue> CompletionQueue;
-	std::mutex QueuedItemsMutex;
+	TUniquePtr<FCriticalSection> QueuedItemsCriticalSection;
 
 	void ProcessNext(void* Tag, bool Ok);
 
diff --git a/Engine/Source/Editor/Experimental/BazelExecutor/Private/BazelCompletionQueueRunnable.inl b/Engine/Source/Editor/Experimental/BazelExecutor/Private/BazelCompletionQueueRunnable.inl
--- a/Engine/Source/Editor/Experimental/BazelExecutor/Private/BazelCompletionQueueRunnable.inl
+++ b/Engine/Source/Editor/Experimental/BazelExecutor/Private/BazelCompletionQueueRunnable.inl
@@ -35,9 +35,10 @@ bool FBazelCompletionQueueRunnable::AddAsyncResponse(
 	Item.Message.Reset(Response);
 	Item.ClientContext = MoveTemp(ClientContext);
 	Item.Finish = MoveTemp(OnFinishResponse);
-	QueuedItemsMutex.lock();
-	QueuedItems.Add(AsyncReader, MoveTemp(Item));
-	QueuedItemsMutex.unlock();
+	{
+		FScopeLock Lock(QueuedItemsCriticalSection.Get());
+		QueuedItems.Add(AsyncReader, MoveTemp(Item));
+	}
 	AsyncReader->StartCall();
 	AsyncReader->Finish(Response, Status, AsyncReader);
 	return true;
diff --git a/Engine/Source/Editor/Experimental/BazelExecutor/Private/ContentAddressableStorage.cpp b/Engine/Source/Editor/Experimental/BazelExecutor/Private/ContentAddressableStorage.cpp
--- a/Engine/Source/Editor/Experimental/BazelExecutor/Private/ContentAddressableStorage.cpp
+++ b/Engine/Source/Editor/Experimental/BazelExecutor/Private/ContentAddressableStorage.cpp
@@ -42,22 +42,22 @@ FContentAddressableStorage::~FContentAddressableStorage()
 {
 }
 
-bool FContentAddressableStorage::ToDigest(const TArray<char>& InData, FDigest& OutDigest)
+bool FContentAddressableStorage::ToDigest(const TArray<uint8>& InData, FDigest& OutDigest)
 {
 	return ProtoConverter::ToDigest(InData, OutDigest);
 }
 
-bool FContentAddressableStorage::ToBlob(const FDirectory& InDirectory, TArray<char>& OutData, FDigest& OutDigest)
+bool FContentAddressableStorage::ToBlob(const FDirectory& InDirectory, TArray<uint8>& OutData, FDigest& OutDigest)
 {
 	return ProtoConverter::ToBlob(InDirectory, OutData, OutDigest);
 }
 
-bool FContentAddressableStorage::ToBlob(const FCommand& InCommand, TArray<char>& OutData, FDigest& OutDigest)
+bool FContentAddressableStorage::ToBlob(const FCommand& InCommand, TArray<uint8>& OutData, FDigest& OutDigest)
 {
 	return ProtoConverter::ToBlob(InCommand, OutData, OutDigest);
 }
 
-bool FContentAddressableStorage::ToBlob(const FAction& InAction, TArray<char>& OutData, FDigest& OutDigest)
+bool FContentAddressableStorage::ToBlob(const FAction& InAction, TArray<uint8>& OutData, FDigest& OutDigest)
 {
 	return ProtoConverter::ToBlob(InAction, OutData, OutDigest);
 }
diff --git a/Engine/Source/Editor/Experimental/BazelExecutor/Private/ContentAddressableStorage.h b/Engine/Source/Editor/Experimental/BazelExecutor/Private/ContentAddressableStorage.h
--- a/Engine/Source/Editor/Experimental/BazelExecutor/Private/ContentAddressableStorage.h
+++ b/Engine/Source/Editor/Experimental/BazelExecutor/Private/ContentAddressableStorage.h
@@ -28,10 +28,10 @@ public:
 	FContentAddressableStorage(const std::shared_ptr<grpc::Channel>& Channel, TSharedPtr<FBazelCompletionQueueRunnable> CompletionQueueRunnable, const TMap<FString, FString>& Headers);
 	~FContentAddressableStorage();
 
-	bool ToDigest(const TArray<char>& InData, FDigest& OutDigest) override;
-	bool ToBlob(const FDirectory& InDirectory, TArray<char>& OutData, FDigest& OutDigest) override;
-	bool ToBlob(const FCommand& InCommand, TArray<char>& OutData, FDigest& OutDigest) override;
-	bool ToBlob(const FAction& InAction, TArray<char>& OutData, FDigest& OutDigest) override;
+	bool ToDigest(const TArray<uint8>& InData, FDigest& OutDigest) override;
+	bool ToBlob(const FDirectory& InDirectory, TArray<uint8>& OutData, FDigest& OutDigest) override;
+	bool ToBlob(const FCommand& InCommand, TArray<uint8>& OutData, FDigest& OutDigest) override;
+	bool ToBlob(const FAction& InAction, TArray<uint8>& OutData, FDigest& OutDigest) override;
 
 	FStatus FindMissingBlobs(const FFindMissingBlobsRequest& Request, FFindMissingBlobsResponse& Response, int64 TimeoutMs = 0) override;
 	FStatus BatchUpdateBlobs(const FBatchUpdateBlobsRequest& Request, FBatchUpdateBlobsResponse& Response, int64 TimeoutMs = 0) override;
diff --git a/Engine/Source/Editor/Experimental/BazelExecutor/Private/ProtoConverter.cpp b/Engine/Source/Editor/Experimental/BazelExecutor/Private/ProtoConverter.cpp
--- a/Engine/Source/Editor/Experimental/BazelExecutor/Private/ProtoConverter.cpp
+++ b/Engine/Source/Editor/Experimental/BazelExecutor/Private/ProtoConverter.cpp
@@ -183,7 +183,7 @@ void ProtoConverter::FromProto(const std::string& In, FIoHash& Out)
 	Out = FIoHash(ByteArray);
 }
 
-void ProtoConverter::FromProto(const std::string& In, TArray<char>& Out)
+void ProtoConverter::FromProto(const std::string& In, TArray<uint8>& Out)
 {
 	Out.Empty();
 	if (In.size() > 0)
@@ -336,7 +336,7 @@ void ProtoConverter::ToProto(const FBatchUpdateBlobsRequest& In, build::bazel::r
 	{
 		build::bazel::remote::execution::v2::BatchUpdateBlobsRequest_Request* OutRequest = Out.add_requests();
 		ToProto(InRequest.Digest, *OutRequest->mutable_digest());
-		OutRequest->mutable_data()->assign(InRequest.Data.GetData(), InRequest.Data.Num());
+		OutRequest->mutable_data()->assign((const char *)InRequest.Data.GetData(), InRequest.Data.Num());
 	}
 }
 
@@ -393,7 +393,7 @@ void ProtoConverter::FromProto(const build::bazel::remote::execution::v2::BatchR
 		FBatchReadBlobsResponse::FResponse OutResponse;
 		FromProto(InResponse.digest(), OutResponse.Digest);
 		OutResponse.Data.Empty();
-		OutResponse.Data.Append(InResponse.data().c_str(), InResponse.data().size());
+		OutResponse.Data.Append((const uint8 *)InResponse.data().c_str(), InResponse.data().size());
 		FromProto(InResponse.status(), OutResponse.Status);
 		Out.Responses.Add(MoveTemp(OutResponse));
 	}
@@ -417,14 +417,14 @@ void ProtoConverter::FromProto(const build::bazel::remote::execution::v2::Execut
 	FromProto(In.message(), Out.Message);
 }
 
-bool ProtoConverter::ToDigest(const TArray<char>& ToData, FDigest& OutDigest)
+bool ProtoConverter::ToDigest(const TArray<uint8>& ToData, FDigest& OutDigest)
 {
 	OutDigest.Hash = FIoHash::HashBuffer(ToData.GetData(), ToData.Num());
 	OutDigest.SizeBytes = ToData.Num();
 	return true;
 }
 
-bool ProtoConverter::ToBlob(const google::protobuf::Message& InMessage, TArray<char>& OutData, FDigest& OutDigest)
+bool ProtoConverter::ToBlob(const google::protobuf::Message& InMessage, TArray<uint8>& OutData, FDigest& OutDigest)
 {
 	OutData.Empty();
 	OutData.AddUninitialized(InMessage.ByteSizeLong());
@@ -436,7 +436,7 @@ bool ProtoConverter::ToBlob(const google::protobuf::Message& InMessage, TArray<c
 	return ToDigest(OutData, OutDigest);
 }
 
-bool ProtoConverter::ToBlob(const FDirectory& InDirectory, TArray<char>& OutData, FDigest& OutDigest)
+bool ProtoConverter::ToBlob(const FDirectory& InDirectory, TArray<uint8>& OutData, FDigest& OutDigest)
 {
 	// Convert to temporary proto Directory
 	build::bazel::remote::execution::v2::Directory OutDirectory;
@@ -444,7 +444,7 @@ bool ProtoConverter::ToBlob(const FDirectory& InDirectory, TArray<char>& OutData
 	return ToBlob(OutDirectory, OutData, OutDigest);
 }
 
-bool ProtoConverter::ToBlob(const FCommand& InCommand, TArray<char>& OutData, FDigest& OutDigest)
+bool ProtoConverter::ToBlob(const FCommand& InCommand, TArray<uint8>& OutData, FDigest& OutDigest)
 {
 	// Convert to temporary proto Action
 	build::bazel::remote::execution::v2::Command OutCommand;
@@ -452,7 +452,7 @@ bool ProtoConverter::ToBlob(const FCommand& InCommand, TArray<char>& OutData, FD
 	return ToBlob(OutCommand, OutData, OutDigest);
 }
 
-bool ProtoConverter::ToBlob(const FAction& InAction, TArray<char>& OutData, FDigest& OutDigest)
+bool ProtoConverter::ToBlob(const FAction& InAction, TArray<uint8>& OutData, FDigest& OutDigest)
 {
 	// Convert to temporary proto Action
 	build::bazel::remote::execution::v2::Action OutAction;
diff --git a/Engine/Source/Editor/Experimental/BazelExecutor/Private/ProtoConverter.h b/Engine/Source/Editor/Experimental/BazelExecutor/Private/ProtoConverter.h
--- a/Engine/Source/Editor/Experimental/BazelExecutor/Private/ProtoConverter.h
+++ b/Engine/Source/Editor/Experimental/BazelExecutor/Private/ProtoConverter.h
@@ -118,7 +118,7 @@ private:
 
 	static void FromProto(const std::string& In, FString& Out);
 	static void FromProto(const std::string& In, FIoHash& Out);
-	static void FromProto(const std::string& In, TArray<char>& Out);
+	static void FromProto(const std::string& In, TArray<uint8>& Out);
 
 	static void FromProto(const google::protobuf::Timestamp& In, FDateTime& Out);
 
@@ -134,7 +134,7 @@ private:
 	static void FromProto(const build::bazel::remote::execution::v2::ExecutedActionMetadata& In, FExecutedActionMetadata& Out);
 	static void FromProto(const build::bazel::remote::execution::v2::ActionResult& In, FActionResult& Out);
 
-	static bool ToBlob(const google::protobuf::Message& InMessage, TArray<char>& OutData, FDigest& OutDigest);
+	static bool ToBlob(const google::protobuf::Message& InMessage, TArray<uint8>& OutData, FDigest& OutDigest);
 
 public:
 	static void ToProto(const FExecuteRequest& In, build::bazel::remote::execution::v2::ExecuteRequest& Out);
@@ -149,9 +149,9 @@ public:
 	static void FromProto(const build::bazel::remote::execution::v2::BatchReadBlobsResponse& In, FBatchReadBlobsResponse& Out);
 	static void FromProto(const build::bazel::remote::execution::v2::ExecuteResponse& In, FExecuteResponse& Out);
 
-	static bool ToDigest(const TArray<char>& InData, FDigest& OutDigest);
+	static bool ToDigest(const TArray<uint8>& InData, FDigest& OutDigest);
 
-	static bool ToBlob(const FDirectory& InDirectory, TArray<char>& OutData, FDigest& OutDigest);
-	static bool ToBlob(const FCommand& InCommand, TArray<char>& OutData, FDigest& OutDigest);
-	static bool ToBlob(const FAction& InAction, TArray<char>& OutData, FDigest& OutDigest);
+	static bool ToBlob(const FDirectory& InDirectory, TArray<uint8>& OutData, FDigest& OutDigest);
+	static bool ToBlob(const FCommand& InCommand, TArray<uint8>& OutData, FDigest& OutDigest);
+	static bool ToBlob(const FAction& InAction, TArray<uint8>& OutData, FDigest& OutDigest);
 };
diff --git a/Engine/Source/Editor/Experimental/RemoteExecution/Public/IContentAddressableStorage.h b/Engine/Source/Editor/Experimental/RemoteExecution/Public/IContentAddressableStorage.h
--- a/Engine/Source/Editor/Experimental/RemoteExecution/Public/IContentAddressableStorage.h
+++ b/Engine/Source/Editor/Experimental/RemoteExecution/Public/IContentAddressableStorage.h
@@ -24,10 +24,10 @@ public:
 	/** Virtual destructor */
 	virtual ~IContentAddressableStorage() {}
 
-	virtual bool ToDigest(const TArray<char>& InData, FDigest& OutDigest) = 0;
-	virtual bool ToBlob(const FDirectory& InDirectory, TArray<char>& OutData, FDigest& OutDigest) = 0;
-	virtual bool ToBlob(const FCommand& InCommand, TArray<char>& OutData, FDigest& OutDigest) = 0;
-	virtual bool ToBlob(const FAction& InAction, TArray<char>& OutData, FDigest& OutDigest) = 0;
+	virtual bool ToDigest(const TArray<uint8>& InData, FDigest& OutDigest) = 0;
+	virtual bool ToBlob(const FDirectory& InDirectory, TArray<uint8>& OutData, FDigest& OutDigest) = 0;
+	virtual bool ToBlob(const FCommand& InCommand, TArray<uint8>& OutData, FDigest& OutDigest) = 0;
+	virtual bool ToBlob(const FAction& InAction, TArray<uint8>& OutData, FDigest& OutDigest) = 0;
 
 	virtual FStatus FindMissingBlobs(const FFindMissingBlobsRequest& Request, FFindMissingBlobsResponse& Response, int64 TimeoutMs = 0) = 0;
 	virtual FStatus BatchUpdateBlobs(const FBatchUpdateBlobsRequest& Request, FBatchUpdateBlobsResponse& Response, int64 TimeoutMs = 0) = 0;
diff --git a/Engine/Source/Editor/Experimental/RemoteExecution/Public/Messages.h b/Engine/Source/Editor/Experimental/RemoteExecution/Public/Messages.h
--- a/Engine/Source/Editor/Experimental/RemoteExecution/Public/Messages.h
+++ b/Engine/Source/Editor/Experimental/RemoteExecution/Public/Messages.h
@@ -192,7 +192,7 @@ struct FBatchUpdateBlobsRequest
 	struct FRequest
 	{
 		FDigest Digest;
-		TArray<char> Data;
+		TArray<uint8> Data;
 	};
 
 	FString InstanceName;
@@ -223,7 +223,7 @@ struct FBatchReadBlobsResponse
 	struct FResponse
 	{
 		FDigest Digest;
-		TArray<char> Data;
+		TArray<uint8> Data;
 		FStatus Status;
 	};
 
@@ -235,7 +235,7 @@ struct FOutputFile
 	FString Path;
 	FDigest Digest;
 	bool IsExecutable;
-	TArray<char> Contents;
+	TArray<uint8> Contents;
 	FNodeProperties NodeProperties;
 };
 
@@ -279,9 +279,9 @@ struct FActionResult
 	TArray<FOutputSymlink> OutputSymlinks;
 	TArray<FOutputDirectory> OutputDirectories;
 	int32 ExitCode = 0;
-	TArray<char> StdoutRaw;
+	TArray<uint8> StdoutRaw;
 	FDigest StdoutDigest;
-	TArray<char> StderrRaw;
+	TArray<uint8> StderrRaw;
 	FDigest StderrDigest;
 	FExecutedActionMetadata ExecutionMetadata;
 };
diff --git a/Engine/Source/Editor/UnrealEd/Private/DerivedDataBuildRemoteExecutor.cpp b/Engine/Source/Editor/UnrealEd/Private/DerivedDataBuildRemoteExecutor.cpp
--- a/Engine/Source/Editor/UnrealEd/Private/DerivedDataBuildRemoteExecutor.cpp
+++ b/Engine/Source/Editor/UnrealEd/Private/DerivedDataBuildRemoteExecutor.cpp
@@ -127,7 +127,7 @@ private:
 		FStringView Name;
 		FDirectory Directory;
 		TOptional<FDigest> Digest;
-		TArray<char> ContentBytes;
+		TArray<uint8> ContentBytes;
 		TArray<int32> SubDirIndices;
 	};
 
@@ -152,9 +152,9 @@ private:
 
 		// Unique items in the tree
 		FCommand Command;
-		TArray<char> CommandContentBytes;
+		TArray<uint8> CommandContentBytes;
 		FAction Action;
-		TArray<char> ActionContentBytes;
+		TArray<uint8> ActionContentBytes;
 		FDigest ActionDigest;
 		FCompositeBuffer BuildActionContentBytes {};
 		FDigest BuildActionDigest;
@@ -918,7 +918,7 @@ TFuture<TPair<FStatus, FBatchUpdateBlobsResponse>> FRemoteBuildExecutionRequest:
 					NewRequest.Data.Reserve(FileBuffer.GetSize());
 					for (const FSharedBuffer& Segment : State.Files[VariantIndex.Index].ContentBytes.GetSegments())
 					{
-						NewRequest.Data.Append((const char *)Segment.GetData(), Segment.GetSize());
+						NewRequest.Data.Append((const uint8 *)Segment.GetData(), Segment.GetSize());
 					}
 					FileBuffer.Reset();
 					UE_LOG(LogDerivedDataBuildRemoteExecutor, Verbose, TEXT("Uploading file '%s' (hash: %s, type: %s) of upload size %d."), *State.Files[VariantIndex.Index].File.Name, *::LexToString(NewRequest.Digest.Hash), LexToString(State.Files[VariantIndex.Index].Type), NewRequest.Data.Num());
