Wrapped ObjectPtr upgrade refresh for Engine + ShooterGame

diff --git a/Engine/Plugins/Animation/LiveLink/Source/LiveLinkSequencer/Public/MovieSceneLiveLinkControllerTrackRecorder.h b/Engine/Plugins/Animation/LiveLink/Source/LiveLinkSequencer/Public/MovieSceneLiveLinkControllerTrackRecorder.h
--- a/Engine/Plugins/Animation/LiveLink/Source/LiveLinkSequencer/Public/MovieSceneLiveLinkControllerTrackRecorder.h
+++ b/Engine/Plugins/Animation/LiveLink/Source/LiveLinkSequencer/Public/MovieSceneLiveLinkControllerTrackRecorder.h
@@ -24,5 +24,5 @@ public:
 protected:
 	/** The LiveLink controller that this track record will record properties from */
 	UPROPERTY(Transient)
-	ULiveLinkControllerBase* LiveLinkControllerToRecord;
+	TObjectPtr<ULiveLinkControllerBase> LiveLinkControllerToRecord;
 };
diff --git a/Engine/Plugins/Editor/SequencerAnimTools/Source/SequencerAnimTools/Public/EditPivotTool.h b/Engine/Plugins/Editor/SequencerAnimTools/Source/SequencerAnimTools/Public/EditPivotTool.h
--- a/Engine/Plugins/Editor/SequencerAnimTools/Source/SequencerAnimTools/Public/EditPivotTool.h
+++ b/Engine/Plugins/Editor/SequencerAnimTools/Source/SequencerAnimTools/Public/EditPivotTool.h
@@ -144,13 +144,13 @@ public:
 protected:
 
 	UPROPERTY()
-	USingleClickInputBehavior* ClickBehavior = nullptr;
+	TObjectPtr<USingleClickInputBehavior> ClickBehavior = nullptr;
 
 	UPROPERTY()
-	UTransformProxy* TransformProxy = nullptr;
+	TObjectPtr<UTransformProxy> TransformProxy = nullptr;
 
 	UPROPERTY()
-	UTransformGizmo* TransformGizmo = nullptr;
+	TObjectPtr<UTransformGizmo> TransformGizmo = nullptr;
 
 protected:
 	bool bShiftPressedWhenStarted = false;
diff --git a/Engine/Plugins/Editor/SequencerAnimTools/Source/SequencerAnimTools/Public/MotionTrailTool.h b/Engine/Plugins/Editor/SequencerAnimTools/Source/SequencerAnimTools/Public/MotionTrailTool.h
--- a/Engine/Plugins/Editor/SequencerAnimTools/Source/SequencerAnimTools/Public/MotionTrailTool.h
+++ b/Engine/Plugins/Editor/SequencerAnimTools/Source/SequencerAnimTools/Public/MotionTrailTool.h
@@ -157,9 +157,9 @@ public:
 	// Support for gizmo. Since the points aren't individual components, we don't actually use UTransformProxy
 	// for the transform forwarding- we just use it for the callbacks.
 	UPROPERTY()
-	UTransformProxy* TransformProxy;
+	TObjectPtr<UTransformProxy> TransformProxy;
 	UPROPERTY()
-	UTransformGizmo* TransformGizmo;
+	TObjectPtr<UTransformGizmo> TransformGizmo;
 
 	FTransform StartDragTransform;
 	bool bGizmoBeingDragged = false;
@@ -205,10 +205,10 @@ protected:
 	FLeftRghtClickBehaviorTarget RightTarget;
 
 	UPROPERTY()
-	USingleClickInputBehavior* LeftClickBehavior = nullptr;
+	TObjectPtr<USingleClickInputBehavior> LeftClickBehavior = nullptr;
 
 	UPROPERTY()
-	USingleClickInputBehavior* RightClickBehavior = nullptr;
+	TObjectPtr<USingleClickInputBehavior> RightClickBehavior = nullptr;
 
 
 
diff --git a/Engine/Plugins/Experimental/MeshModelingToolsetExp/Source/MeshModelingToolsExp/Public/ToolActivities/PolyEditActivityContext.h b/Engine/Plugins/Experimental/MeshModelingToolsetExp/Source/MeshModelingToolsExp/Public/ToolActivities/PolyEditActivityContext.h
--- a/Engine/Plugins/Experimental/MeshModelingToolsetExp/Source/MeshModelingToolsExp/Public/ToolActivities/PolyEditActivityContext.h
+++ b/Engine/Plugins/Experimental/MeshModelingToolsetExp/Source/MeshModelingToolsExp/Public/ToolActivities/PolyEditActivityContext.h
@@ -56,7 +56,7 @@ public:
 	 * through EmitCurrentMeshChangeAndUpdate so that it is lumped with the same undo transaction.
 	 */
 	UPROPERTY()
-	UPolygonSelectionMechanic* SelectionMechanic;
+	TObjectPtr<UPolygonSelectionMechanic> SelectionMechanic;
 
 	/**
 	 * Given a change to CurrentMesh (already applied and stored in MeshChangeIn), use this
diff --git a/Engine/Plugins/Experimental/MeshModelingToolsetExp/Source/MeshModelingToolsExp/Public/ToolActivities/PolyEditCutFacesActivity.h b/Engine/Plugins/Experimental/MeshModelingToolsetExp/Source/MeshModelingToolsExp/Public/ToolActivities/PolyEditCutFacesActivity.h
--- a/Engine/Plugins/Experimental/MeshModelingToolsetExp/Source/MeshModelingToolsExp/Public/ToolActivities/PolyEditCutFacesActivity.h
+++ b/Engine/Plugins/Experimental/MeshModelingToolsetExp/Source/MeshModelingToolsExp/Public/ToolActivities/PolyEditCutFacesActivity.h
@@ -72,13 +72,13 @@ protected:
 	void ApplyCutFaces();
 
 	UPROPERTY()
-	UPolyEditCutProperties* CutProperties;
+	TObjectPtr<UPolyEditCutProperties> CutProperties;
 
 	UPROPERTY()
-	UPolyEditPreviewMesh* EditPreview;
+	TObjectPtr<UPolyEditPreviewMesh> EditPreview;
 
 	UPROPERTY()
-	UCollectSurfacePathMechanic* SurfacePathMechanic = nullptr;
+	TObjectPtr<UCollectSurfacePathMechanic> SurfacePathMechanic = nullptr;
 
 	UPROPERTY()
 	TObjectPtr<UPolyEditActivityContext> ActivityContext;
diff --git a/Engine/Plugins/Experimental/MeshModelingToolsetExp/Source/MeshModelingToolsExp/Public/ToolActivities/PolyEditExtrudeActivity.h b/Engine/Plugins/Experimental/MeshModelingToolsetExp/Source/MeshModelingToolsExp/Public/ToolActivities/PolyEditExtrudeActivity.h
--- a/Engine/Plugins/Experimental/MeshModelingToolsetExp/Source/MeshModelingToolsExp/Public/ToolActivities/PolyEditExtrudeActivity.h
+++ b/Engine/Plugins/Experimental/MeshModelingToolsetExp/Source/MeshModelingToolsExp/Public/ToolActivities/PolyEditExtrudeActivity.h
@@ -90,7 +90,7 @@ public:
 	virtual void OnEndHover() override {}
 
 	UPROPERTY()
-	UPolyEditExtrudeProperties* ExtrudeProperties;
+	TObjectPtr<UPolyEditExtrudeProperties> ExtrudeProperties;
 
 protected:
 	FVector3d GetExtrudeDirection() const;
@@ -99,10 +99,10 @@ protected:
 	void Clear();
 
 	UPROPERTY()
-	UPolyEditPreviewMesh* EditPreview;
+	TObjectPtr<UPolyEditPreviewMesh> EditPreview;
 
 	UPROPERTY()
-	UPlaneDistanceFromHitMechanic* ExtrudeHeightMechanic = nullptr;
+	TObjectPtr<UPlaneDistanceFromHitMechanic> ExtrudeHeightMechanic = nullptr;
 
 	UPROPERTY()
 	TObjectPtr<UPolyEditActivityContext> ActivityContext;
diff --git a/Engine/Plugins/Experimental/MeshModelingToolsetExp/Source/MeshModelingToolsExp/Public/ToolActivities/PolyEditInsertEdgeActivity.h b/Engine/Plugins/Experimental/MeshModelingToolsetExp/Source/MeshModelingToolsExp/Public/ToolActivities/PolyEditInsertEdgeActivity.h
--- a/Engine/Plugins/Experimental/MeshModelingToolsetExp/Source/MeshModelingToolsExp/Public/ToolActivities/PolyEditInsertEdgeActivity.h
+++ b/Engine/Plugins/Experimental/MeshModelingToolsetExp/Source/MeshModelingToolsExp/Public/ToolActivities/PolyEditInsertEdgeActivity.h
@@ -103,7 +103,7 @@ protected:
 	TObjectPtr<UGroupEdgeInsertionProperties> Settings = nullptr;
 
 	UPROPERTY()
-	UPolyEditActivityContext* ActivityContext;
+	TObjectPtr<UPolyEditActivityContext> ActivityContext;
 
 	bool bIsRunning = false;
 
diff --git a/Engine/Plugins/Experimental/MeshModelingToolsetExp/Source/MeshModelingToolsExp/Public/ToolActivities/PolyEditInsertEdgeLoopActivity.h b/Engine/Plugins/Experimental/MeshModelingToolsetExp/Source/MeshModelingToolsExp/Public/ToolActivities/PolyEditInsertEdgeLoopActivity.h
--- a/Engine/Plugins/Experimental/MeshModelingToolsetExp/Source/MeshModelingToolsExp/Public/ToolActivities/PolyEditInsertEdgeLoopActivity.h
+++ b/Engine/Plugins/Experimental/MeshModelingToolsetExp/Source/MeshModelingToolsExp/Public/ToolActivities/PolyEditInsertEdgeLoopActivity.h
@@ -132,10 +132,10 @@ public:
 
 protected:
 	UPROPERTY()
-	UEdgeLoopInsertionProperties* Settings = nullptr;
+	TObjectPtr<UEdgeLoopInsertionProperties> Settings = nullptr;
 
 	UPROPERTY()
-	UPolyEditActivityContext* ActivityContext;
+	TObjectPtr<UPolyEditActivityContext> ActivityContext;
 
 	bool bIsRunning = false;
 
diff --git a/Engine/Plugins/Experimental/MeshModelingToolsetExp/Source/MeshModelingToolsExp/Public/ToolActivities/PolyEditInsetOutsetActivity.h b/Engine/Plugins/Experimental/MeshModelingToolsetExp/Source/MeshModelingToolsExp/Public/ToolActivities/PolyEditInsetOutsetActivity.h
--- a/Engine/Plugins/Experimental/MeshModelingToolsetExp/Source/MeshModelingToolsExp/Public/ToolActivities/PolyEditInsetOutsetActivity.h
+++ b/Engine/Plugins/Experimental/MeshModelingToolsetExp/Source/MeshModelingToolsExp/Public/ToolActivities/PolyEditInsetOutsetActivity.h
@@ -79,7 +79,7 @@ public:
 	virtual void OnEndHover() override {}
 
 	UPROPERTY()
-	UPolyEditInsetOutsetProperties* Settings;
+	TObjectPtr<UPolyEditInsetOutsetProperties> Settings;
 
 protected:
 	void Clear();
@@ -90,10 +90,10 @@ protected:
 	bool bPreviewUpdatePending = false;
 
 	UPROPERTY()
-	UPolyEditPreviewMesh* EditPreview;
+	TObjectPtr<UPolyEditPreviewMesh> EditPreview;
 
 	UPROPERTY()
-	USpatialCurveDistanceMechanic* CurveDistMechanic = nullptr;
+	TObjectPtr<USpatialCurveDistanceMechanic> CurveDistMechanic = nullptr;
 
 	UPROPERTY()
 	TObjectPtr<UPolyEditActivityContext> ActivityContext;
diff --git a/Engine/Plugins/Experimental/MeshModelingToolsetExp/Source/MeshModelingToolsExp/Public/ToolActivities/PolyEditPlanarProjectionUVActivity.h b/Engine/Plugins/Experimental/MeshModelingToolsetExp/Source/MeshModelingToolsExp/Public/ToolActivities/PolyEditPlanarProjectionUVActivity.h
--- a/Engine/Plugins/Experimental/MeshModelingToolsetExp/Source/MeshModelingToolsExp/Public/ToolActivities/PolyEditPlanarProjectionUVActivity.h
+++ b/Engine/Plugins/Experimental/MeshModelingToolsetExp/Source/MeshModelingToolsExp/Public/ToolActivities/PolyEditPlanarProjectionUVActivity.h
@@ -65,13 +65,13 @@ protected:
 	void ApplySetUVs();
 
 	UPROPERTY()
-	UPolyEditSetUVProperties* SetUVProperties;
+	TObjectPtr<UPolyEditSetUVProperties> SetUVProperties;
 
 	UPROPERTY()
-	UPolyEditPreviewMesh* EditPreview;
+	TObjectPtr<UPolyEditPreviewMesh> EditPreview;
 
 	UPROPERTY()
-	UCollectSurfacePathMechanic* SurfacePathMechanic = nullptr;
+	TObjectPtr<UCollectSurfacePathMechanic> SurfacePathMechanic = nullptr;
 
 	UPROPERTY()
 	TObjectPtr<UPolyEditActivityContext> ActivityContext;
diff --git a/Engine/Source/Editor/Blutility/Public/AsyncCaptureScene.h b/Engine/Source/Editor/Blutility/Public/AsyncCaptureScene.h
--- a/Engine/Source/Editor/Blutility/Public/AsyncCaptureScene.h
+++ b/Engine/Source/Editor/Blutility/Public/AsyncCaptureScene.h
@@ -44,8 +44,8 @@ private:
 
 private:
 	UPROPERTY()
-	ASceneCapture2D* SceneCapture;
+	TObjectPtr<ASceneCapture2D> SceneCapture;
 
 	UPROPERTY()
-	UTextureRenderTarget2D* SceneCaptureRT;
+	TObjectPtr<UTextureRenderTarget2D> SceneCaptureRT;
 };
diff --git a/Engine/Source/Editor/Experimental/EditorInteractiveToolsFramework/Private/EditorInteractiveGizmoRegistry.cpp b/Engine/Source/Editor/Experimental/EditorInteractiveToolsFramework/Private/EditorInteractiveGizmoRegistry.cpp
--- a/Engine/Source/Editor/Experimental/EditorInteractiveToolsFramework/Private/EditorInteractiveGizmoRegistry.cpp
+++ b/Engine/Source/Editor/Experimental/EditorInteractiveToolsFramework/Private/EditorInteractiveGizmoRegistry.cpp
@@ -156,7 +156,7 @@ void UEditorInteractiveGizmoRegistry::RegisterEditorGizmoType(EEditorGizmoCatego
 {
 	if (ensure(InGizmoBuilder))
 	{
-		if (UEditorGizmoRegistryCategoryEntry** RegEntry = GizmoCategoryMap.Find(InGizmoCategory))
+		if (UE_TRANSITIONAL_OBJECT_PTR(UEditorGizmoRegistryCategoryEntry)* RegEntry = GizmoCategoryMap.Find(InGizmoCategory))
 		{
 			if (*RegEntry)
 			{
@@ -168,7 +168,7 @@ void UEditorInteractiveGizmoRegistry::RegisterEditorGizmoType(EEditorGizmoCatego
 
 void UEditorInteractiveGizmoRegistry::GetQualifiedEditorGizmoBuilders(EEditorGizmoCategory InGizmoCategory, const FToolBuilderState& InToolBuilderState, TArray<UInteractiveGizmoBuilder*>& InFoundBuilders)
 {
-	if (UEditorGizmoRegistryCategoryEntry** RegEntry = GizmoCategoryMap.Find(InGizmoCategory))
+	if (UE_TRANSITIONAL_OBJECT_PTR(UEditorGizmoRegistryCategoryEntry)* RegEntry = GizmoCategoryMap.Find(InGizmoCategory))
 	{
 		if (*RegEntry)
 		{
@@ -181,7 +181,7 @@ void UEditorInteractiveGizmoRegistry::DeregisterEditorGizmoType(EEditorGizmoCate
 {
 	if (ensure(InGizmoBuilder))
 	{
-		if (UEditorGizmoRegistryCategoryEntry** RegEntry = GizmoCategoryMap.Find(InGizmoCategory))
+		if (UE_TRANSITIONAL_OBJECT_PTR(UEditorGizmoRegistryCategoryEntry)* RegEntry = GizmoCategoryMap.Find(InGizmoCategory))
 		{
 			if (*RegEntry)
 			{
diff --git a/Engine/Source/Editor/Experimental/EditorInteractiveToolsFramework/Public/EditorInteractiveGizmoRegistry.h b/Engine/Source/Editor/Experimental/EditorInteractiveToolsFramework/Public/EditorInteractiveGizmoRegistry.h
--- a/Engine/Source/Editor/Experimental/EditorInteractiveToolsFramework/Public/EditorInteractiveGizmoRegistry.h
+++ b/Engine/Source/Editor/Experimental/EditorInteractiveToolsFramework/Public/EditorInteractiveGizmoRegistry.h
@@ -46,7 +46,7 @@ protected:
 	FString CategoryName;
 
 	UPROPERTY()
-	UClass* BaseGizmoBuilderType;
+	TObjectPtr<UClass> BaseGizmoBuilderType;
 };
 
 UCLASS(Abstract)
@@ -137,5 +137,5 @@ private:
 
 	/** Current set of Gizmo Builders */
 	UPROPERTY()
-	TMap<EEditorGizmoCategory, UEditorGizmoRegistryCategoryEntry*> GizmoCategoryMap;
+	TMap<EEditorGizmoCategory, TObjectPtr<UEditorGizmoRegistryCategoryEntry>> GizmoCategoryMap;
 };
\ No newline at end of file
diff --git a/Engine/Source/Runtime/ClientPilot/Public/ClientPilotBlackboardManager.h b/Engine/Source/Runtime/ClientPilot/Public/ClientPilotBlackboardManager.h
--- a/Engine/Source/Runtime/ClientPilot/Public/ClientPilotBlackboardManager.h
+++ b/Engine/Source/Runtime/ClientPilot/Public/ClientPilotBlackboardManager.h
@@ -15,5 +15,5 @@ protected:
 public:
 	static UClientPilotBlackboardManager * GetInstance();
 	UPROPERTY()
-	UClientPilotBlackboard* PilotBlackboard;
+	TObjectPtr<UClientPilotBlackboard> PilotBlackboard;
 };
\ No newline at end of file
diff --git a/Engine/Source/Runtime/Engine/Public/WorldPartition/HLOD/HLODLayer.h b/Engine/Source/Runtime/Engine/Public/WorldPartition/HLOD/HLODLayer.h
--- a/Engine/Source/Runtime/Engine/Public/WorldPartition/HLOD/HLODLayer.h
+++ b/Engine/Source/Runtime/Engine/Public/WorldPartition/HLOD/HLODLayer.h
@@ -75,7 +75,7 @@ private:
 	TSubclassOf<UHLODBuilder> HLODBuilderClass;
 
 	UPROPERTY(VisibleAnywhere, Export, NoClear, Category=HLOD, meta = (EditInline, NoResetToDefault))
-	UHLODBuilderSettings* HLODBuilderSettings;
+	TObjectPtr<UHLODBuilderSettings> HLODBuilderSettings;
 
 	/** Whether HLOD actors generated for this layer will be always loaded */
 	UPROPERTY(EditAnywhere, Config, Category=HLOD)
diff --git a/Engine/Source/Runtime/RigVM/Private/RigVMCore/RigVM.cpp b/Engine/Source/Runtime/RigVM/Private/RigVMCore/RigVM.cpp
--- a/Engine/Source/Runtime/RigVM/Private/RigVMCore/RigVM.cpp
+++ b/Engine/Source/Runtime/RigVM/Private/RigVMCore/RigVM.cpp
@@ -634,7 +634,7 @@ void URigVM::CopyFrom(URigVM* InVM, bool bDeferCopy, bool bReferenceLiteralMemor
 	
 #else
 
-	auto CopyMemoryStorage = [](URigVMMemoryStorage*& TargetMemory, URigVMMemoryStorage* SourceMemory, UObject* Outer)
+	auto CopyMemoryStorage = [](TObjectPtr<URigVMMemoryStorage>& TargetMemory, URigVMMemoryStorage* SourceMemory, UObject* Outer)
 	{
 		if(SourceMemory != nullptr)
 		{
diff --git a/Engine/Source/Runtime/RigVM/Public/RigVMCore/RigVM.h b/Engine/Source/Runtime/RigVM/Public/RigVMCore/RigVM.h
--- a/Engine/Source/Runtime/RigVM/Public/RigVMCore/RigVM.h
+++ b/Engine/Source/Runtime/RigVM/Public/RigVMCore/RigVM.h
@@ -248,13 +248,13 @@ public:
 #endif
 
 	UPROPERTY(transient, DuplicateTransient)
-	URigVMMemoryStorage* WorkMemoryStorageObject;
+	TObjectPtr<URigVMMemoryStorage> WorkMemoryStorageObject;
 
 	UPROPERTY(transient, DuplicateTransient)
-	URigVMMemoryStorage* LiteralMemoryStorageObject;
+	TObjectPtr<URigVMMemoryStorage> LiteralMemoryStorageObject;
 
 	UPROPERTY(transient, DuplicateTransient)
-	URigVMMemoryStorage* DebugMemoryStorageObject;
+	TObjectPtr<URigVMMemoryStorage> DebugMemoryStorageObject;
 
 	TArray<FRigVMPropertyPathDescription> ExternalPropertyPathDescriptions;
 	TArray<FRigVMPropertyPath> ExternalPropertyPaths;
diff --git a/Engine/Source/Runtime/RuntimeAssetCache/Public/RuntimeAssetCacheBuilders.h b/Engine/Source/Runtime/RuntimeAssetCache/Public/RuntimeAssetCacheBuilders.h
--- a/Engine/Source/Runtime/RuntimeAssetCache/Public/RuntimeAssetCacheBuilders.h
+++ b/Engine/Source/Runtime/RuntimeAssetCache/Public/RuntimeAssetCacheBuilders.h
@@ -186,5 +186,5 @@ public:
 	virtual int64 GetSerializedDataSizeEstimate() override;
 
 	UPROPERTY(BlueprintReadOnly, Category = "Cache")
-	UTexture2D* Texture;
+	TObjectPtr<UTexture2D> Texture;
 };
