Change Editor Gizmo Builders from UObjects to interfaces.

diff --git a/Engine/Source/Editor/Experimental/EditorInteractiveToolsFramework/Private/EditorInteractiveGizmoManager.cpp b/Engine/Source/Editor/Experimental/EditorInteractiveToolsFramework/Private/EditorInteractiveGizmoManager.cpp
--- a/Engine/Source/Editor/Experimental/EditorInteractiveToolsFramework/Private/EditorInteractiveGizmoManager.cpp
+++ b/Engine/Source/Editor/Experimental/EditorInteractiveToolsFramework/Private/EditorInteractiveGizmoManager.cpp
@@ -103,7 +103,7 @@ TArray<UInteractiveGizmo*> UEditorInteractiveGizmoManager::CreateGizmosForCurren
 		{
 			if (UInteractiveGizmo* NewGizmo = FoundBuilder->BuildGizmo(CurrentSceneState))
 			{
-				if (UEditorInteractiveGizmoSelectionBuilder* SelectionBuilder = Cast<UEditorInteractiveGizmoSelectionBuilder>(FoundBuilder))
+				if (IEditorInteractiveGizmoSelectionBuilder* SelectionBuilder = Cast<IEditorInteractiveGizmoSelectionBuilder>(FoundBuilder))
 				{
 					SelectionBuilder->UpdateGizmoForSelection(NewGizmo, CurrentSceneState);
 				}
diff --git a/Engine/Source/Editor/Experimental/EditorInteractiveToolsFramework/Private/EditorInteractiveGizmoRegistry.cpp b/Engine/Source/Editor/Experimental/EditorInteractiveToolsFramework/Private/EditorInteractiveGizmoRegistry.cpp
--- a/Engine/Source/Editor/Experimental/EditorInteractiveToolsFramework/Private/EditorInteractiveGizmoRegistry.cpp
+++ b/Engine/Source/Editor/Experimental/EditorInteractiveToolsFramework/Private/EditorInteractiveGizmoRegistry.cpp
@@ -13,13 +13,6 @@ void UEditorGizmoRegistryCategoryEntry::RegisterGizmoType(UInteractiveGizmoBuild
 {
 	check(InGizmoBuilder);
 
-	if (!InGizmoBuilder->IsA(BaseGizmoBuilderType))
-	{
-		UE_LOG(LogEditorInteractiveGizmoRegistry, Warning,
-			TEXT("%s gizmo builder '%s' of type '%s' does not derive from '%s'! Skipping registration."), *CategoryName, *InGizmoBuilder->GetPathName(), *InGizmoBuilder->GetClass()->GetName(), *BaseGizmoBuilderType->GetName());
-		return;
-	}
-
 	if (GizmoTypes.Contains(InGizmoBuilder))
 	{
 		UE_LOG(LogEditorInteractiveGizmoRegistry, Warning,
@@ -50,14 +43,28 @@ void UEditorGizmoRegistryCategoryEntry::ClearGizmoTypes()
 	GizmoTypes.Reset();
 }
 
-void UEditorGizmoRegistryCategoryEntry_Conditional::RegisterGizmoType(UInteractiveGizmoBuilder* InGizmoBuilder)
+void UEditorGizmoRegistryCategoryEntry_ConditionalSelection::RegisterGizmoType(UInteractiveGizmoBuilder* InGizmoBuilder)
 {
+	if (!InGizmoBuilder->Implements<UEditorInteractiveGizmoConditionalBuilder>())
+	{
+		UE_LOG(LogEditorInteractiveGizmoRegistry, Warning,
+			TEXT("%s gizmo builder '%s' of type '%s' does not implement the IEditorInteractiveGizmoConditionalBuilder interface! Skipping registration."), *CategoryName, *InGizmoBuilder->GetPathName(), *InGizmoBuilder->GetClass()->GetName());
+		return;
+	}
+
+	if (!InGizmoBuilder->Implements<UEditorInteractiveGizmoSelectionBuilder>())
+	{
+		UE_LOG(LogEditorInteractiveGizmoRegistry, Warning,
+			TEXT("%s gizmo builder '%s' of type '%s' does not implement the IEditorInteractiveGizmoSelectionBuilder interface! Skipping registration."), *CategoryName, *InGizmoBuilder->GetPathName(), *InGizmoBuilder->GetClass()->GetName());
+		return;
+	}
+
 	Super::RegisterGizmoType(InGizmoBuilder);
 
 	GizmoTypes.StableSort(
 		[](UInteractiveGizmoBuilder& A, UInteractiveGizmoBuilder& B) {
-			UEditorInteractiveGizmoConditionalBuilder* AA = Cast<UEditorInteractiveGizmoConditionalBuilder>(&A);
-			UEditorInteractiveGizmoConditionalBuilder* BB = Cast<UEditorInteractiveGizmoConditionalBuilder>(&B);
+			IEditorInteractiveGizmoConditionalBuilder* AA = Cast<IEditorInteractiveGizmoConditionalBuilder>(&A);
+			IEditorInteractiveGizmoConditionalBuilder* BB = Cast<IEditorInteractiveGizmoConditionalBuilder>(&B);
 			return (AA)->GetPriority() > (BB)->GetPriority();
 		});
 }
@@ -65,7 +72,6 @@ void UEditorGizmoRegistryCategoryEntry_Conditional::RegisterGizmoType(UInteracti
 UEditorGizmoRegistryCategoryEntry_Primary::UEditorGizmoRegistryCategoryEntry_Primary()
 {
 	CategoryName = TEXT("Primary");
-	BaseGizmoBuilderType = UEditorInteractiveGizmoSelectionBuilder::StaticClass();
 }
 
 void UEditorGizmoRegistryCategoryEntry_Primary::GetQualifiedGizmoBuilders(const FToolBuilderState& InToolBuilderState, TArray<UInteractiveGizmoBuilder*>& FoundBuilders)
@@ -74,7 +80,7 @@ void UEditorGizmoRegistryCategoryEntry_Primary::GetQualifiedGizmoBuilders(const
 
 	if (FoundBuilders.Num() > 0)
 	{
-		if (UEditorInteractiveGizmoConditionalBuilder* Builder = Cast<UEditorInteractiveGizmoConditionalBuilder>(FoundBuilders[0]))
+		if (IEditorInteractiveGizmoConditionalBuilder* Builder = Cast<IEditorInteractiveGizmoConditionalBuilder>(FoundBuilders[0]))
 		{
 			FoundPriority = Builder->GetPriority();
 		}
@@ -82,7 +88,7 @@ void UEditorGizmoRegistryCategoryEntry_Primary::GetQualifiedGizmoBuilders(const
 
 	for (TObjectPtr<UInteractiveGizmoBuilder> GizmoBuilder : GizmoTypes)
 	{
-		if (UEditorInteractiveGizmoConditionalBuilder* Builder = Cast<UEditorInteractiveGizmoConditionalBuilder>(GizmoBuilder))
+		if (IEditorInteractiveGizmoConditionalBuilder* Builder = Cast<IEditorInteractiveGizmoConditionalBuilder>(GizmoBuilder))
 		{
 			if (Builder->GetPriority() < FoundPriority)
 			{
@@ -93,7 +99,7 @@ void UEditorGizmoRegistryCategoryEntry_Primary::GetQualifiedGizmoBuilders(const
 			{
 				// Reset found builders since only one primary builder should be buildable.
 				FoundBuilders.Reset();
-				FoundBuilders.Add(Builder);
+				FoundBuilders.Add(GizmoBuilder);
 				break;
 			}
 		}
@@ -103,7 +109,6 @@ void UEditorGizmoRegistryCategoryEntry_Primary::GetQualifiedGizmoBuilders(const
 UEditorGizmoRegistryCategoryEntry_Accessory::UEditorGizmoRegistryCategoryEntry_Accessory()
 {
 	CategoryName = TEXT("Accessory");
-	BaseGizmoBuilderType = UEditorInteractiveGizmoSelectionBuilder::StaticClass();
 }
 
 void UEditorGizmoRegistryCategoryEntry_Accessory::GetQualifiedGizmoBuilders(const FToolBuilderState& InToolBuilderState, TArray<UInteractiveGizmoBuilder*>& FoundBuilders)
@@ -112,7 +117,7 @@ void UEditorGizmoRegistryCategoryEntry_Accessory::GetQualifiedGizmoBuilders(cons
 
 	if (FoundBuilders.Num() > 0)
 	{
-		if (UEditorInteractiveGizmoConditionalBuilder* Builder = Cast<UEditorInteractiveGizmoConditionalBuilder>(FoundBuilders[0]))
+		if (IEditorInteractiveGizmoConditionalBuilder* Builder = Cast<IEditorInteractiveGizmoConditionalBuilder>(FoundBuilders[0]))
 		{
 			FoundPriority = Builder->GetPriority();
 		}
@@ -120,7 +125,7 @@ void UEditorGizmoRegistryCategoryEntry_Accessory::GetQualifiedGizmoBuilders(cons
 
 	for (TObjectPtr<UInteractiveGizmoBuilder> GizmoBuilder : GizmoTypes)
 	{
-		if (UEditorInteractiveGizmoConditionalBuilder* Builder = Cast<UEditorInteractiveGizmoConditionalBuilder>(GizmoBuilder))
+		if (IEditorInteractiveGizmoConditionalBuilder* Builder = Cast<IEditorInteractiveGizmoConditionalBuilder>(GizmoBuilder))
 		{
 			if (Builder->GetPriority() < FoundPriority)
 			{
@@ -135,7 +140,7 @@ void UEditorGizmoRegistryCategoryEntry_Accessory::GetQualifiedGizmoBuilders(cons
 				{
 					FoundBuilders.Reset();
 				}
-				FoundBuilders.Add(Builder);
+				FoundBuilders.Add(GizmoBuilder);
 			}
 		}
 	}
diff --git a/Engine/Source/Editor/Experimental/EditorInteractiveToolsFramework/Public/EditorGizmos/EditorTransformGizmo.h b/Engine/Source/Editor/Experimental/EditorInteractiveToolsFramework/Public/EditorGizmos/EditorTransformGizmo.h
--- a/Engine/Source/Editor/Experimental/EditorInteractiveToolsFramework/Public/EditorGizmos/EditorTransformGizmo.h
+++ b/Engine/Source/Editor/Experimental/EditorInteractiveToolsFramework/Public/EditorGizmos/EditorTransformGizmo.h
@@ -29,7 +29,7 @@ class FEditorTransformGizmoTransformChange;
 class FRay;
 
 UCLASS()
-class EDITORINTERACTIVETOOLSFRAMEWORK_API UEditorTransformGizmoBuilder : public UEditorInteractiveGizmoSelectionBuilder
+class EDITORINTERACTIVETOOLSFRAMEWORK_API UEditorTransformGizmoBuilder : public UInteractiveGizmoBuilder, public IEditorInteractiveGizmoSelectionBuilder
 {
 	GENERATED_BODY()
 
@@ -38,12 +38,6 @@ public:
 	// UEditorInteractiveGizmoSelectionBuilder interface 
 	virtual UInteractiveGizmo* BuildGizmo(const FToolBuilderState& SceneState) const override;
 	virtual void UpdateGizmoForSelection(UInteractiveGizmo* Gizmo, const FToolBuilderState& SceneState) override;
-
-	// UEditorInteractiveGizmoConditionalBuilder interface 
-	virtual bool SatisfiesCondition(const FToolBuilderState& SceneState) const override
-	{
-		return true;
-	}
 };
 
 
diff --git a/Engine/Source/Editor/Experimental/EditorInteractiveToolsFramework/Public/EditorInteractiveGizmoConditionalBuilder.h b/Engine/Source/Editor/Experimental/EditorInteractiveToolsFramework/Public/EditorInteractiveGizmoConditionalBuilder.h
--- a/Engine/Source/Editor/Experimental/EditorInteractiveToolsFramework/Public/EditorInteractiveGizmoConditionalBuilder.h
+++ b/Engine/Source/Editor/Experimental/EditorInteractiveToolsFramework/Public/EditorInteractiveGizmoConditionalBuilder.h
@@ -54,34 +54,30 @@ struct EDITORINTERACTIVETOOLSFRAMEWORK_API FEditorGizmoTypePriority
 
 /** 
  * Gizmo builders which should be built once a condition is satisfied in the current scene state
- * inherit from UEditorInteractiveGizmoConditionalBuilder. The SatisfiesCondition method
+ * implement the IEditorInteractiveGizmoConditionalBuilder interface. The SatisfiesCondition method
  * should return true whenever the builder is buildable based on scene state, most commonly based
  * on selection. The builder's priority should be specified such that when more than one builder
  * is discovered, the builder with highest priority will be built.
  */
-UCLASS(Transient, Abstract)
-class EDITORINTERACTIVETOOLSFRAMEWORK_API UEditorInteractiveGizmoConditionalBuilder : public UInteractiveGizmoBuilder
+UINTERFACE()
+class EDITORINTERACTIVETOOLSFRAMEWORK_API UEditorInteractiveGizmoConditionalBuilder : public UInterface
+{
+	GENERATED_BODY()
+};
+
+class EDITORINTERACTIVETOOLSFRAMEWORK_API IEditorInteractiveGizmoConditionalBuilder
 {
 	GENERATED_BODY()
 
 public:
 
 	/** Returns the priority for this gizmo type.  */
-	virtual FEditorGizmoTypePriority GetPriority() const 
-	{ 
-		return FEditorGizmoTypePriority::DEFAULT_GIZMO_TYPE_PRIORITY;
-	}
+	virtual FEditorGizmoTypePriority GetPriority() const = 0;
 
 	/** Update the priority for this gizmo type.  */
-	virtual void UpdatePriority(const FEditorGizmoTypePriority& InPriority)
-	{
-		Priority = InPriority;
-	}
+	virtual void SetPriority(const FEditorGizmoTypePriority& InPriority) = 0;
 
 	/** Returns true if this gizmo is valid for creation based on the current state. */
-	virtual bool SatisfiesCondition(const FToolBuilderState& SceneState) const  PURE_VIRTUAL(UEditorInteractiveGizmoConditionalBuilder::SatisfiesCondition, return false; );
-
-protected:
+	virtual bool SatisfiesCondition(const FToolBuilderState& SceneState) const = 0;
 
-	FEditorGizmoTypePriority Priority = FEditorGizmoTypePriority::DEFAULT_GIZMO_TYPE_PRIORITY;
 };
diff --git a/Engine/Source/Editor/Experimental/EditorInteractiveToolsFramework/Public/EditorInteractiveGizmoManager.h b/Engine/Source/Editor/Experimental/EditorInteractiveToolsFramework/Public/EditorInteractiveGizmoManager.h
--- a/Engine/Source/Editor/Experimental/EditorInteractiveToolsFramework/Public/EditorInteractiveGizmoManager.h
+++ b/Engine/Source/Editor/Experimental/EditorInteractiveToolsFramework/Public/EditorInteractiveGizmoManager.h
@@ -55,8 +55,8 @@ public:
 	 * Register a new Editor gizmo type.
 	 * @param InGizmoCategory category in which to register gizmo builder
 	 * @param InGizmonBuilder new Editor gizmo builder
-	 * - Accessory gizmo builders must be inherited from UEditorInteractiveGizmoAccessoryBuilder.
-	 * - Primary gizmo builders must be inherited from UEditorInteractiveGizmoPrimaryBuilder.
+	 * - Accessory and Primary gizmo builders must derive from UInteractiveGizmoBuilder (of from a builder derived from it)
+	 *   and must implement the IEditorInteractiveConditionalGizmoBuilder and IEditorInteractiveSelectionGizmoBuilder interfaces.
 	 */
 	void RegisterEditorGizmoType(EEditorGizmoCategory InGizmoCategory, UInteractiveGizmoBuilder* InGizmoBuilder);
 
diff --git a/Engine/Source/Editor/Experimental/EditorInteractiveToolsFramework/Public/EditorInteractiveGizmoRegistry.h b/Engine/Source/Editor/Experimental/EditorInteractiveToolsFramework/Public/EditorInteractiveGizmoRegistry.h
--- a/Engine/Source/Editor/Experimental/EditorInteractiveToolsFramework/Public/EditorInteractiveGizmoRegistry.h
+++ b/Engine/Source/Editor/Experimental/EditorInteractiveToolsFramework/Public/EditorInteractiveGizmoRegistry.h
@@ -44,13 +44,10 @@ protected:
 
 	UPROPERTY()
 	FString CategoryName;
-
-	UPROPERTY()
-	TObjectPtr<UClass> BaseGizmoBuilderType;
 };
 
 UCLASS(Abstract)
-class UEditorGizmoRegistryCategoryEntry_Conditional : public UEditorGizmoRegistryCategoryEntry
+class UEditorGizmoRegistryCategoryEntry_ConditionalSelection : public UEditorGizmoRegistryCategoryEntry
 {
 	GENERATED_BODY()
 public:
@@ -62,7 +59,7 @@ public:
 
 
 UCLASS()
-class UEditorGizmoRegistryCategoryEntry_Primary : public UEditorGizmoRegistryCategoryEntry_Conditional
+class UEditorGizmoRegistryCategoryEntry_Primary : public UEditorGizmoRegistryCategoryEntry_ConditionalSelection
 {
 	GENERATED_BODY()
 public:
@@ -72,7 +69,7 @@ public:
 };
 
 UCLASS()
-class UEditorGizmoRegistryCategoryEntry_Accessory : public UEditorGizmoRegistryCategoryEntry_Conditional
+class UEditorGizmoRegistryCategoryEntry_Accessory : public UEditorGizmoRegistryCategoryEntry_ConditionalSelection
 {
 	GENERATED_BODY()
 public:
@@ -103,9 +100,9 @@ public:
 	/**
 	 * Register a new Editor gizmo type.
 	 * @param InGizmoCategory category in which to register gizmo builder
-	 * @param InGizmonBuilder new Editor gizmo builder
-	 * - Accessory gizmo builders must be inherited from UEditorInteractiveGizmoAccessoryBuilder.
-	 * - Primary gizmo builders must be inherited from UEditorInteractiveGizmoPrimaryBuilder.
+	 * @param InGizmoBuilder new Editor gizmo builder
+	 * - Accessory and Primary gizmo builders must derive from UInteractiveGizmoBuilder (of from a builder derived from it)
+	 *   and must implement the IEditorInteractiveConditionalGizmoBuilder and IEditorInteractiveSelectionGizmoBuilder interfaces.
 	 */
 	void RegisterEditorGizmoType(EEditorGizmoCategory InGizmoCategory, UInteractiveGizmoBuilder* InGizmoBuilder);
 
diff --git a/Engine/Source/Editor/Experimental/EditorInteractiveToolsFramework/Public/EditorInteractiveGizmoSelectionBuilder.h b/Engine/Source/Editor/Experimental/EditorInteractiveToolsFramework/Public/EditorInteractiveGizmoSelectionBuilder.h
--- a/Engine/Source/Editor/Experimental/EditorInteractiveToolsFramework/Public/EditorInteractiveGizmoSelectionBuilder.h
+++ b/Engine/Source/Editor/Experimental/EditorInteractiveToolsFramework/Public/EditorInteractiveGizmoSelectionBuilder.h
@@ -18,8 +18,13 @@ public:
 	static UTransformProxy* CreateTransformProxyForSelection(const FToolBuilderState& SceneState);
 };
 
-UCLASS(Transient, Abstract)
-class EDITORINTERACTIVETOOLSFRAMEWORK_API UEditorInteractiveGizmoSelectionBuilder : public UEditorInteractiveGizmoConditionalBuilder
+UINTERFACE()
+class EDITORINTERACTIVETOOLSFRAMEWORK_API UEditorInteractiveGizmoSelectionBuilder : public UInterface
+{
+	GENERATED_BODY()
+};
+
+class EDITORINTERACTIVETOOLSFRAMEWORK_API IEditorInteractiveGizmoSelectionBuilder
 {
 	GENERATED_BODY()
 
@@ -31,6 +36,6 @@ public:
 	 * This method is called after a gizmo is automatically built based upon selection and also to update the 
 	 * existing gizmo when the selection changes.
 	 */
-	virtual void UpdateGizmoForSelection(UInteractiveGizmo* Gizmo, const FToolBuilderState& SceneState) PURE_VIRTUAL(UEditorInteractiveGizmoSelectionBuilder::UpdateGizmoForSelection, return;);
+	virtual void UpdateGizmoForSelection(UInteractiveGizmo* Gizmo, const FToolBuilderState& SceneState) = 0;
 
 };
diff --git a/Engine/Source/Editor/Experimental/EditorInteractiveToolsFramework/Public/EditorInteractiveGizmoSubsystem.h b/Engine/Source/Editor/Experimental/EditorInteractiveToolsFramework/Public/EditorInteractiveGizmoSubsystem.h
--- a/Engine/Source/Editor/Experimental/EditorInteractiveToolsFramework/Public/EditorInteractiveGizmoSubsystem.h
+++ b/Engine/Source/Editor/Experimental/EditorInteractiveToolsFramework/Public/EditorInteractiveGizmoSubsystem.h
@@ -68,8 +68,8 @@ public:
 	 * Register a new Editor gizmo type which will be global to the Editor.
 	 * @param InGizmoCategory category in which to register gizmo builder
 	 * @param InGizmonBuilder new Editor gizmo builder
-	 * - Accessory gizmo builders must be inherited from UEditorInteractiveGizmoAccessoryBuilder.
-	 * - Primary gizmo builders must be inherited from UEditorInteractiveGizmoPrimaryBuilder.
+	 * - Accessory and Primary gizmo builders must derive from UInteractiveGizmoBuilder (of from a builder derived from it)
+	 *   and must implement the IEditorInteractiveConditionalGizmoBuilder and IEditorInteractiveSelectionGizmoBuilder interfaces.
 	 */
 	void RegisterGlobalEditorGizmoType(EEditorGizmoCategory InGizmoCategory, UInteractiveGizmoBuilder* InGizmoBuilder);
 
@@ -93,7 +93,7 @@ public:
 	/** 
 	 * Get transform gizmo builder used to build the Level Editor TRS gizmo. 
 	 */
-	UEditorInteractiveGizmoSelectionBuilder* GetTransformGizmoBuilder()
+	UInteractiveGizmoBuilder* GetTransformGizmoBuilder()
 	{
 		return TransformGizmoBuilder;
 	}
@@ -102,7 +102,7 @@ private:
 
 	/** TRS gizmo builder */
 	UPROPERTY()
-	TObjectPtr<UEditorInteractiveGizmoSelectionBuilder> TransformGizmoBuilder;
+	TObjectPtr<UInteractiveGizmoBuilder> TransformGizmoBuilder;
 
 	/** Actual registry */
 	UPROPERTY()
