RemoteExecution: Add task result outcome and details

#rnx

diff --git a/Engine/Source/Editor/Experimental/HordeExecutor/Private/Execution.cpp b/Engine/Source/Editor/Experimental/HordeExecutor/Private/Execution.cpp
--- a/Engine/Source/Editor/Experimental/HordeExecutor/Private/Execution.cpp
+++ b/Engine/Source/Editor/Experimental/HordeExecutor/Private/Execution.cpp
@@ -120,9 +120,9 @@ namespace UE::RemoteExecution
 			for (FGetTaskUpdateResponse& Update : Response.Value.Updates)
 			{
 				UE_LOG(LogHordeExecutor, Verbose, TEXT("WaitForTasksAsync: %s %s"),
-					*TaskStateString(Update.State),
+					*ComputeTaskStateString(Update.State),
 					*FString::FromHexBlob(Update.TaskHash.GetBytes(), sizeof(FIoHash::ByteArray)));
-				if (Update.State == ETaskState::Complete)
+				if (Update.State == EComputeTaskState::Complete)
 				{
 					WaitTask->Complete.Add(Update.TaskHash, MoveTemp(Update));
 				}
diff --git a/Engine/Source/Editor/Experimental/RemoteExecution/Private/RemoteMessages.cpp b/Engine/Source/Editor/Experimental/RemoteExecution/Private/RemoteMessages.cpp
--- a/Engine/Source/Editor/Experimental/RemoteExecution/Private/RemoteMessages.cpp
+++ b/Engine/Source/Editor/Experimental/RemoteExecution/Private/RemoteMessages.cpp
@@ -46,7 +46,8 @@ namespace UE::RemoteExecution
 
 	FGetTaskUpdateResponse::FGetTaskUpdateResponse()
 		: Time(0)
-		, State(ETaskState::Queued)
+		, State(EComputeTaskState::Queued)
+		, Outcome(EComputeTaskOutcome::Success)
 	{
 	}
 
@@ -60,6 +61,8 @@ namespace UE::RemoteExecution
 		}
 		Writer.AddDateTime("t", Time);
 		Writer.AddInteger("s", (int32)State);
+		Writer.AddInteger("o", (int32)Outcome);
+		Writer.AddString("d", Detail);
 		if (ResultHash != FIoHash::Zero)
 		{
 			Writer.AddObjectAttachment("r", ResultHash);
@@ -74,7 +77,9 @@ namespace UE::RemoteExecution
 	{
 		TaskHash = CbObjectView["h"].AsObjectAttachment();
 		Time = CbObjectView["t"].AsDateTime();
-		State = (ETaskState)CbObjectView["s"].AsInt32();
+		State = (EComputeTaskState)CbObjectView["s"].AsInt32();
+		Outcome = (EComputeTaskOutcome)CbObjectView["o"].AsInt32();
+		Detail = FString(CbObjectView["d"].AsString());
 		ResultHash = CbObjectView["r"].AsObjectAttachment();
 		AgentId = FString(CbObjectView["a"].AsString());
 		LeaseId = FString(CbObjectView["l"].AsString());
diff --git a/Engine/Source/Editor/Experimental/RemoteExecution/Public/RemoteMessages.h b/Engine/Source/Editor/Experimental/RemoteExecution/Public/RemoteMessages.h
--- a/Engine/Source/Editor/Experimental/RemoteExecution/Public/RemoteMessages.h
+++ b/Engine/Source/Editor/Experimental/RemoteExecution/Public/RemoteMessages.h
@@ -31,20 +31,46 @@ namespace UE::RemoteExecution
 		GatewayTimeout = 504,
 	};
 
-	enum class ETaskState : int32
+	enum class EComputeTaskState : int32
 	{
 		Queued = 0,
 		Executing = 1,
 		Complete = 2,
 	};
 
-	static FString TaskStateString(const ETaskState TaskState)
+	enum class EComputeTaskOutcome : int32
 	{
-		switch (TaskState)
+		Success = 0,
+		Failed = 1,
+		Cancelled = 2,
+		NoResult = 3,
+		Exipred = 4,
+		BlobNotFound = 5,
+		Exception = 6,
+	};
+
+	static FString ComputeTaskStateString(const EComputeTaskState ComputeTaskState)
+	{
+		switch (ComputeTaskState)
+		{
+		case EComputeTaskState::Queued: return TEXT("Queued");
+		case EComputeTaskState::Executing: return TEXT("Executing");
+		case EComputeTaskState::Complete: return TEXT("Complete");
+		}
+		return TEXT("Unknown");
+	}
+
+	static FString ComputeTaskOutcomeString(const EComputeTaskOutcome ComputeTaskOutcome)
+	{
+		switch (ComputeTaskOutcome)
 		{
-		case ETaskState::Queued: return TEXT("Queued");
-		case ETaskState::Executing: return TEXT("Executing");
-		case ETaskState::Complete: return TEXT("Complete");
+		case EComputeTaskOutcome::Success: return TEXT("Success");
+		case EComputeTaskOutcome::Failed: return TEXT("Failed");
+		case EComputeTaskOutcome::Cancelled: return TEXT("Cancelled");
+		case EComputeTaskOutcome::NoResult: return TEXT("NoResult");
+		case EComputeTaskOutcome::Exipred: return TEXT("Exipred");
+		case EComputeTaskOutcome::BlobNotFound: return TEXT("BlobNotFound");
+		case EComputeTaskOutcome::Exception: return TEXT("Exception");
 		}
 		return TEXT("Unknown");
 	}
@@ -68,7 +94,9 @@ namespace UE::RemoteExecution
 	public:
 		FIoHash TaskHash;
 		FDateTime Time;
-		ETaskState State;
+		EComputeTaskState State;
+		EComputeTaskOutcome Outcome;
+		FString Detail;
 		FIoHash ResultHash;
 		FString AgentId;
 		FString LeaseId;
diff --git a/Engine/Source/Editor/UnrealEd/Private/DerivedDataBuildRemoteExecutor.cpp b/Engine/Source/Editor/UnrealEd/Private/DerivedDataBuildRemoteExecutor.cpp
--- a/Engine/Source/Editor/UnrealEd/Private/DerivedDataBuildRemoteExecutor.cpp
+++ b/Engine/Source/Editor/UnrealEd/Private/DerivedDataBuildRemoteExecutor.cpp
@@ -1160,6 +1160,19 @@ void FRemoteBuildExecutionRequest::OnExecutionCompleted(UE::RemoteExecution::ESt
 		return;
 	}
 
+	if (Result.Updates[0].Outcome != UE::RemoteExecution::EComputeTaskOutcome::Success)
+	{
+		UE_LOG(LogDerivedDataBuildRemoteExecutor, Warning, TEXT("Remote execution system error: Outcome %s: %s"),
+			*UE::RemoteExecution::ComputeTaskOutcomeString(Result.Updates[0].Outcome),
+			*Result.Updates[0].Detail);
+		State.Owner.End(this, [this]
+			{
+				CompletionCallback({ State.BuildAction.GetKey(), {}, {}, EStatus::Error });
+				CompletionEvent->Trigger();
+			});
+		return;
+	}
+
 	if (Result.Updates[0].ResultHash == FIoHash::Zero)
 	{
 		UE_LOG(LogDerivedDataBuildRemoteExecutor, Warning, TEXT("Remote execution system error: Zero ResultHash returned from remote build operation!"));
