HordeExecutor: Remove stale thread code and optimization disable I left in

#rnx

diff --git a/Engine/Source/Editor/Experimental/HordeExecutor/HordeExecutor.Build.cs b/Engine/Source/Editor/Experimental/HordeExecutor/HordeExecutor.Build.cs
--- a/Engine/Source/Editor/Experimental/HordeExecutor/HordeExecutor.Build.cs
+++ b/Engine/Source/Editor/Experimental/HordeExecutor/HordeExecutor.Build.cs
@@ -18,8 +18,6 @@ namespace UnrealBuildTool.Rules
 				"RemoteExecution",
 				}
 			);
-
-			OptimizeCode = CodeOptimization.Never;
 		}
 	}
 }
diff --git a/Engine/Source/Editor/Experimental/HordeExecutor/Private/HordeExecutor.cpp b/Engine/Source/Editor/Experimental/HordeExecutor/Private/HordeExecutor.cpp
--- a/Engine/Source/Editor/Experimental/HordeExecutor/Private/HordeExecutor.cpp
+++ b/Engine/Source/Editor/Experimental/HordeExecutor/Private/HordeExecutor.cpp
@@ -15,17 +15,11 @@ namespace UE::RemoteExecution
 		Shutdown();
 
 		ContentAddressableStorage.Reset(new FContentAddressableStorage(Settings.ContentAddressableStorageTarget, Settings.ContentAddressableStorageHeaders));
-
 		Execution.Reset(new FExecution(Settings.ExecutionTarget, Settings.ExecutionHeaders));
 	}
 
 	void FHordeExecutor::Shutdown()
 	{
-		if (Thread.IsValid())
-		{
-			Thread->Kill();
-		}
-		Thread.Reset();
 		ContentAddressableStorage.Reset();
 		Execution.Reset();
 	}
diff --git a/Engine/Source/Editor/Experimental/HordeExecutor/Private/HordeExecutor.h b/Engine/Source/Editor/Experimental/HordeExecutor/Private/HordeExecutor.h
--- a/Engine/Source/Editor/Experimental/HordeExecutor/Private/HordeExecutor.h
+++ b/Engine/Source/Editor/Experimental/HordeExecutor/Private/HordeExecutor.h
@@ -7,7 +7,6 @@
 #include "IRemoteExecutor.h"
 #include "IContentAddressableStorage.h"
 #include "IExecution.h"
-#include "HAL/RunnableThread.h"
 
 
 namespace UE::RemoteExecution
@@ -26,7 +25,6 @@ namespace UE::RemoteExecution
 	private:
 		TUniquePtr<IContentAddressableStorage> ContentAddressableStorage;
 		TUniquePtr<IExecution> Execution;
-		TUniquePtr<FRunnableThread> Thread;
 
 	public:
 		void Initialize(const FSettings& Settings);
