Update LevelEditor TRS gizmo to handle Editor mode and coord system changes.

diff --git a/Engine/Source/Editor/Experimental/EditorInteractiveToolsFramework/Private/EditorGizmos/EditorTransformGizmo.cpp b/Engine/Source/Editor/Experimental/EditorInteractiveToolsFramework/Private/EditorGizmos/EditorTransformGizmo.cpp
deleted file mode 100644
diff --git a/Engine/Source/Editor/Experimental/EditorInteractiveToolsFramework/Private/EditorGizmos/EditorTransformGizmoBuilder.cpp b/Engine/Source/Editor/Experimental/EditorInteractiveToolsFramework/Private/EditorGizmos/EditorTransformGizmoBuilder.cpp
new file mode 100644
--- /dev/null
+++ b/Engine/Source/Editor/Experimental/EditorInteractiveToolsFramework/Private/EditorGizmos/EditorTransformGizmoBuilder.cpp
@@ -0,0 +1,24 @@
+// Copyright Epic Games, Inc. All Rights Reserved.
+
+#include "EditorGizmos/EditorTransformGizmoBuilder.h"
+#include "EditorGizmos/EditorTransformGizmoSource.h"
+#include "EditorGizmos/EditorTransformProxy.h"
+#include "EditorGizmos/TransformGizmo.h"
+
+UInteractiveGizmo* UEditorTransformGizmoBuilder::BuildGizmo(const FToolBuilderState& SceneState) const
+{
+	UTransformGizmo* TransformGizmo = NewObject<UTransformGizmo>(SceneState.GizmoManager);
+	TransformGizmo->Setup();
+	TransformGizmo->TransformSource = UEditorTransformGizmoSource::Construct(TransformGizmo);
+	return TransformGizmo;
+}
+
+void UEditorTransformGizmoBuilder::UpdateGizmoForSelection(UInteractiveGizmo* Gizmo, const FToolBuilderState& SceneState)
+{
+	if (UTransformGizmo* TransformGizmo = Cast<UTransformGizmo>(Gizmo))
+	{
+		UEditorTransformProxy* TransformProxy = NewObject<UEditorTransformProxy>();
+		TransformGizmo->SetActiveTarget(TransformProxy);
+		TransformGizmo->SetVisibility(true);
+	}
+}
diff --git a/Engine/Source/Editor/Experimental/EditorInteractiveToolsFramework/Private/EditorGizmos/EditorTransformGizmoSource.cpp b/Engine/Source/Editor/Experimental/EditorInteractiveToolsFramework/Private/EditorGizmos/EditorTransformGizmoSource.cpp
new file mode 100644
--- /dev/null
+++ b/Engine/Source/Editor/Experimental/EditorInteractiveToolsFramework/Private/EditorGizmos/EditorTransformGizmoSource.cpp
@@ -0,0 +1,78 @@
+// Copyright Epic Games, Inc. All Rights Reserved.
+
+#include "EditorGizmos/EditorTransformGizmoSource.h"
+#include "EditorModeManager.h"
+#include "EditorViewportClient.h"
+
+EGizmoTransformMode UEditorTransformGizmoSource::GetGizmoMode() const
+{
+	if (FEditorViewportClient* ViewportClient = GetViewportClient())
+	{
+		UE::Widget::EWidgetMode WidgetMode = ViewportClient->GetWidgetMode();
+		return FTransformGizmoUtil::GetGizmoMode(WidgetMode);
+	}
+	return EGizmoTransformMode::None;
+}
+
+EAxisList::Type UEditorTransformGizmoSource::GetGizmoAxisToDraw(EGizmoTransformMode InGizmoMode) const
+{ 
+	if (FEditorViewportClient* ViewportClient = GetViewportClient())
+	{
+		UE::Widget::EWidgetMode WidgetMode = ViewportClient->GetWidgetMode();
+		return GetModeTools().GetWidgetAxisToDraw(WidgetMode);
+	}
+	return EAxisList::None;
+}
+
+EToolContextCoordinateSystem UEditorTransformGizmoSource::GetGizmoCoordSystemSpace() const
+{
+	FEditorViewportClient* ViewportClient = GetViewportClient();
+	if (ViewportClient && ViewportClient->GetWidgetCoordSystemSpace() == ECoordSystem::COORD_Local)
+	{
+		return EToolContextCoordinateSystem::Local;
+	}
+	else
+	{
+		return EToolContextCoordinateSystem::World;
+	}
+}
+
+float UEditorTransformGizmoSource::GetGizmoScale() const
+{
+	return GetModeTools().GetWidgetScale();
+}
+
+
+bool UEditorTransformGizmoSource::GetVisible() const
+{
+	if (FEditorViewportClient* ViewportClient = GetViewportClient()) 
+	{
+		if (GetModeTools().GetShowWidget() && GetModeTools().UsesTransformWidget())
+		{
+			UE::Widget::EWidgetMode WidgetMode = ViewportClient->GetWidgetMode();
+			bool bUseLegacyWidget = (WidgetMode == UE::Widget::WM_TranslateRotateZ || WidgetMode == UE::Widget::WM_2D);
+			if (!bUseLegacyWidget)
+			{
+				static IConsoleVariable* const UseLegacyWidgetCVar = IConsoleManager::Get().FindConsoleVariable(TEXT("Gizmos.UseLegacyWidget"));
+				if (ensure(UseLegacyWidgetCVar))
+				{
+					bUseLegacyWidget = UseLegacyWidgetCVar->GetInt() > 0;
+				}
+			}
+
+			return !bUseLegacyWidget;
+		}
+	}
+
+	return false;
+}
+
+FEditorModeTools& UEditorTransformGizmoSource::GetModeTools() const
+{
+	return GLevelEditorModeTools();
+}
+
+FEditorViewportClient* UEditorTransformGizmoSource::GetViewportClient() const
+{
+	return GetModeTools().GetFocusedViewportClient();
+}
\ No newline at end of file
diff --git a/Engine/Source/Editor/Experimental/EditorInteractiveToolsFramework/Private/EditorGizmos/EditorTransformProxy.cpp b/Engine/Source/Editor/Experimental/EditorInteractiveToolsFramework/Private/EditorGizmos/EditorTransformProxy.cpp
new file mode 100644
--- /dev/null
+++ b/Engine/Source/Editor/Experimental/EditorInteractiveToolsFramework/Private/EditorGizmos/EditorTransformProxy.cpp
@@ -0,0 +1,28 @@
+// Copyright Epic Games, Inc. All Rights Reserved.
+
+#include "EditorGizmos/EditorTransformProxy.h"
+#include "EditorViewportClient.h"
+#include "EditorModeManager.h"
+
+
+#define LOCTEXT_NAMESPACE "UEditorTransformProxy"
+
+FTransform UEditorTransformProxy::GetTransform() const
+{
+	if (FEditorViewportClient* ViewportClient = GLevelEditorModeTools().GetFocusedViewportClient())
+	{
+		FVector Location = ViewportClient->GetWidgetLocation();
+		FMatrix RotMatrix = ViewportClient->GetWidgetCoordSystem();
+		return FTransform(FQuat(RotMatrix), Location, FVector::OneVector);
+	}
+	
+	return FTransform::Identity;
+}
+
+void UEditorTransformProxy::SetTransform(const FTransform& Transform)
+{
+	// @todo: use ViewportClient->HandleWidgetDelta();
+	// @todo: handle Pivots, need to look at how this is currently handled for component visualizers, for example.
+}
+
+#undef LOCTEXT_NAMESPACE
diff --git a/Engine/Source/Editor/Experimental/EditorInteractiveToolsFramework/Private/EditorGizmos/GizmoArrowObject.cpp b/Engine/Source/Editor/Experimental/EditorInteractiveToolsFramework/Private/EditorGizmos/GizmoArrowObject.cpp
--- a/Engine/Source/Editor/Experimental/EditorInteractiveToolsFramework/Private/EditorGizmos/GizmoArrowObject.cpp
+++ b/Engine/Source/Editor/Experimental/EditorInteractiveToolsFramework/Private/EditorGizmos/GizmoArrowObject.cpp
@@ -124,11 +124,11 @@ void UGizmoArrowObject::SetVisibility(bool bVisibleIn)
 	}
 	if (ConeObject)
 	{
-		ConeObject->SetVisibility(bVisibleIn);
+		ConeObject->SetVisibility(bVisibleIn && bHasConeHead);
 	}
 	if (BoxObject)
 	{
-		BoxObject->SetVisibility(bVisibleIn);
+		BoxObject->SetVisibility(bVisibleIn && !bHasConeHead);
 	}
 }
 
diff --git a/Engine/Source/Editor/Experimental/EditorInteractiveToolsFramework/Private/EditorGizmos/TransformGizmo.cpp b/Engine/Source/Editor/Experimental/EditorInteractiveToolsFramework/Private/EditorGizmos/TransformGizmo.cpp
new file mode 100644
--- /dev/null
+++ b/Engine/Source/Editor/Experimental/EditorInteractiveToolsFramework/Private/EditorGizmos/TransformGizmo.cpp
@@ -0,0 +1,829 @@
+// Copyright Epic Games, Inc. All Rights Reserved.
+
+#include "EditorGizmos/TransformGizmo.h"
+#include "BaseGizmos/AxisPositionGizmo.h"
+#include "BaseGizmos/TransformSources.h"
+#include "EditorGizmos/EditorAxisSources.h"
+#include "EditorGizmos/EditorTransformGizmoSource.h"
+#include "EditorGizmos/EditorTransformProxy.h"
+#include "EditorGizmos/EditorParameterToTransformAdapters.h"
+#include "EditorGizmos/GizmoArrowObject.h"
+#include "EditorGizmos/GizmoBoxObject.h"
+#include "EditorGizmos/GizmoConeObject.h"
+#include "EditorGizmos/GizmoCylinderObject.h"
+#include "EditorGizmos/GizmoGroupObject.h"
+#include "Elements/Interfaces/TypedElementObjectInterface.h"
+#include "Elements/Interfaces/TypedElementWorldInterface.h"
+#include "Engine/CollisionProfile.h"
+#include "Engine/Selection.h"
+#include "Engine/World.h"
+#include "EditorModeTools.h"
+#include "UnrealEdGlobals.h"
+#include "UnrealEngine.h"
+
+#define LOCTEXT_NAMESPACE "UTransformGizmo"
+
+constexpr float UTransformGizmo::AXIS_LENGTH;
+constexpr float UTransformGizmo::AXIS_RADIUS;
+constexpr float UTransformGizmo::AXIS_CONE_ANGLE;
+constexpr float UTransformGizmo::AXIS_CONE_HEIGHT;
+constexpr float UTransformGizmo::AXIS_CONE_HEAD_OFFSET;
+constexpr float UTransformGizmo::AXIS_CUBE_SIZE;
+constexpr float UTransformGizmo::AXIS_CUBE_HEAD_OFFSET;
+constexpr float UTransformGizmo::TRANSLATE_ROTATE_AXIS_CIRCLE_RADIUS;
+constexpr float UTransformGizmo::TWOD_AXIS_CIRCLE_RADIUS;
+constexpr float UTransformGizmo::INNER_AXIS_CIRCLE_RADIUS;
+constexpr float UTransformGizmo::OUTER_AXIS_CIRCLE_RADIUS;
+constexpr float UTransformGizmo::ROTATION_TEXT_RADIUS;
+constexpr int32 UTransformGizmo::AXIS_CIRCLE_SIDES;
+constexpr float UTransformGizmo::ARCALL_RELATIVE_INNER_SIZE;
+constexpr float UTransformGizmo::AXIS_LENGTH_SCALE;
+constexpr float UTransformGizmo::AXIS_LENGTH_SCALE_OFFSET;
+
+constexpr FLinearColor UTransformGizmo::AxisColorX;
+constexpr FLinearColor UTransformGizmo::AxisColorY;
+constexpr FLinearColor UTransformGizmo::AxisColorZ;
+constexpr FLinearColor UTransformGizmo::ScreenAxisColor;
+constexpr FColor UTransformGizmo::PlaneColorXY;
+constexpr FColor UTransformGizmo::ArcBallColor;
+constexpr FColor UTransformGizmo::ScreenSpaceColor;
+constexpr FColor UTransformGizmo::CurrentColor;
+
+void UTransformGizmo::SetDisallowNegativeScaling(bool bDisallow)
+{
+	if (bDisallowNegativeScaling != bDisallow)
+	{
+		bDisallowNegativeScaling = bDisallow;
+		for (UInteractiveGizmo* SubGizmo : this->ActiveGizmos)
+		{
+			if (UAxisPositionGizmo* CastGizmo = Cast<UAxisPositionGizmo>(SubGizmo))
+			{
+				if (UGizmoAxisScaleParameterSource* ParamSource = Cast<UGizmoAxisScaleParameterSource>(CastGizmo->ParameterSource.GetObject()))
+				{
+					ParamSource->bClampToZero = bDisallow;
+				}
+			}
+			/* @todo
+			if (UPlanePositionGizmo* CastGizmo = Cast<UEditorPlanePositionGizmo>(SubGizmo))
+			{
+				if (UGizmoPlaneScaleParameterSource* ParamSource = Cast<UGizmoPlaneScaleParameterSource>(CastGizmo->ParameterSource.GetObject()))
+				{
+					ParamSource->bClampToZero = bDisallow;
+				}
+			}
+			*/
+		}
+	}
+}
+
+void UTransformGizmo::Setup()
+{
+	UInteractiveGizmo::Setup();
+
+	UMaterial* AxisMaterialBase = GEngine->ArrowMaterial;
+
+	AxisMaterialX = UMaterialInstanceDynamic::Create(AxisMaterialBase, NULL);
+	AxisMaterialX->SetVectorParameterValue("GizmoColor", AxisColorX);
+
+	AxisMaterialY = UMaterialInstanceDynamic::Create(AxisMaterialBase, NULL);
+	AxisMaterialY->SetVectorParameterValue("GizmoColor", AxisColorY);
+
+	AxisMaterialZ = UMaterialInstanceDynamic::Create(AxisMaterialBase, NULL);
+	AxisMaterialZ->SetVectorParameterValue("GizmoColor", AxisColorZ);
+
+	CurrentAxisMaterial = UMaterialInstanceDynamic::Create(AxisMaterialBase, NULL);
+	CurrentAxisMaterial->SetVectorParameterValue("GizmoColor", CurrentColor);
+
+	OpaquePlaneMaterialXY = UMaterialInstanceDynamic::Create(AxisMaterialBase, NULL);
+	OpaquePlaneMaterialXY->SetVectorParameterValue("GizmoColor", FLinearColor::White);
+
+	TransparentPlaneMaterialXY = (UMaterial*)StaticLoadObject(
+		UMaterial::StaticClass(), NULL,
+		TEXT("/Engine/EditorMaterials/WidgetVertexColorMaterial.WidgetVertexColorMaterial"), NULL, LOAD_None, NULL);
+
+	GridMaterial = (UMaterial*)StaticLoadObject(
+		UMaterial::StaticClass(), NULL,
+		TEXT("/Engine/EditorMaterials/WidgetGridVertexColorMaterial_Ma.WidgetGridVertexColorMaterial_Ma"), NULL,
+		LOAD_None, NULL);
+	if (!GridMaterial)
+	{
+		GridMaterial = TransparentPlaneMaterialXY;
+	}
+}
+
+
+void UTransformGizmo::Shutdown()
+{
+	ClearActiveTarget();
+}
+
+void UTransformGizmo::Render(IToolsContextRenderAPI* RenderAPI)
+{
+	if (bVisible && GizmoGroupObject)
+	{
+		GizmoGroupObject->Render(RenderAPI);
+	}
+}
+
+void UTransformGizmo::UpdateMode()
+{
+	if (TransformSource && TransformSource->GetVisible())
+	{
+		EGizmoTransformMode NewMode = TransformSource->GetGizmoMode();
+		EAxisList::Type NewAxisToDraw = TransformSource->GetGizmoAxisToDraw(NewMode);
+
+		if (NewMode != CurrentMode)
+		{
+			ActiveObjects.Empty();
+			EnableMode(CurrentMode, EAxisList::None);
+			EnableMode(NewMode, NewAxisToDraw);
+
+			CurrentMode = NewMode;
+			CurrentAxisToDraw = NewAxisToDraw;
+		}
+		else if (NewAxisToDraw != CurrentAxisToDraw)
+		{
+			ActiveObjects.Empty();
+			EnableMode(CurrentMode, NewAxisToDraw);
+			CurrentAxisToDraw = NewAxisToDraw;
+		}
+	}
+	else
+	{
+		ActiveObjects.Empty();
+		EnableMode(CurrentMode, EAxisList::None);
+		CurrentMode = EGizmoTransformMode::None;
+	}
+}
+
+void UTransformGizmo::UpdateCoordSystem()
+{
+	// Note: the following will change with upcoming changes to gizmo objects
+	if (ActiveTarget && ActiveObjects.Num() > 0)
+	{
+		EToolContextCoordinateSystem Space = EToolContextCoordinateSystem::World;
+		float Scale = 1.0f;
+
+		if (TransformSource)
+		{
+			Space = TransformSource->GetGizmoCoordSystemSpace();
+			Scale = TransformSource->GetGizmoScale();
+		}
+
+		FTransform LocalToWorldTransform = ActiveTarget->GetTransform();
+		if (Space == EToolContextCoordinateSystem::World)
+		{
+			LocalToWorldTransform.SetRotation(FQuat::Identity);
+		}
+
+		for (UGizmoBaseObject* Object : ActiveObjects)
+		{
+			Object->SetWorldLocalState(Space == EToolContextCoordinateSystem::World);
+			Object->SetLocalToWorldTransform(LocalToWorldTransform);
+			Object->SetGizmoScale(Scale);
+		}
+	}
+}
+
+void UTransformGizmo::EnableMode(EGizmoTransformMode InMode, EAxisList::Type InAxisListToDraw)
+{
+	if (InMode == EGizmoTransformMode::Translate)
+	{
+		EnableTranslate(InAxisListToDraw);
+	}
+	else if (InMode == EGizmoTransformMode::Rotate)
+	{
+		EnableRotate(InAxisListToDraw);
+	}
+	else if (InMode == EGizmoTransformMode::Scale)
+	{
+		EnableScale(InAxisListToDraw);
+	}
+}
+
+void UTransformGizmo::EnableObject(UGizmoBaseObject* InGizmoObject, EAxisList::Type InGizmoAxis, EAxisList::Type InAxisListToDraw)
+{
+	if (static_cast<uint8>(InAxisListToDraw) & static_cast<uint8>(InGizmoAxis))
+	{
+		InGizmoObject->SetVisibility(true);
+		ActiveObjects.Add(InGizmoObject);
+	}
+	else
+	{
+		InGizmoObject->SetVisibility(false);
+	}
+}
+
+void UTransformGizmo::EnableTranslate(EAxisList::Type InAxisListToDraw)
+{
+	if (ensure(AxisXObject && AxisYObject && AxisZObject))
+	{
+		EnableObject(AxisXObject, EAxisList::X, InAxisListToDraw);
+		EnableObject(AxisYObject, EAxisList::Y, InAxisListToDraw);
+		EnableObject(AxisZObject, EAxisList::Z, InAxisListToDraw);
+	}
+}
+
+void UTransformGizmo::EnableRotate(EAxisList::Type InAxisListToDraw)
+{
+	// @todo
+}
+
+void UTransformGizmo::EnableScale(EAxisList::Type InAxisListToDraw)
+{
+	if (ensure(ScaleAxisXObject && ScaleAxisYObject && ScaleAxisZObject))
+	{
+		EnableObject(ScaleAxisXObject, EAxisList::X, InAxisListToDraw);
+		EnableObject(ScaleAxisYObject, EAxisList::Y, InAxisListToDraw);
+		EnableObject(ScaleAxisZObject, EAxisList::Z, InAxisListToDraw);
+	}
+}
+
+void UTransformGizmo::UpdateCameraAxisSource()
+{
+	FViewCameraState CameraState;
+	GetGizmoManager()->GetContextQueriesAPI()->GetCurrentViewState(CameraState);
+	if (CameraAxisSource != nullptr)
+	{
+		CameraAxisSource->Origin = ActiveTarget ? ActiveTarget->GetTransform().GetLocation() : FVector::ZeroVector;
+		CameraAxisSource->Direction = -CameraState.Forward();
+		CameraAxisSource->TangentX = CameraState.Right();
+		CameraAxisSource->TangentY = CameraState.Up();
+	}
+}
+
+
+void UTransformGizmo::Tick(float DeltaTime)
+{
+	UpdateMode();
+
+	UpdateCoordSystem();
+
+	UpdateCameraAxisSource();
+}
+
+void UTransformGizmo::SetActiveTarget(UTransformProxy* Target, IToolContextTransactionProvider* TransactionProvider)
+{
+	if (ActiveTarget != nullptr)
+	{
+		ClearActiveTarget();
+	}
+
+	ActiveTarget = Target;
+
+	if (!ActiveTarget)
+	{
+		return;
+	}
+
+	// create group object to which all active objects will be added
+	GizmoGroupObject = NewObject<UGizmoGroupObject>();
+
+	// root component provides local X/Y/Z axis, identified by AxisIndex
+	AxisXSource = UGizmoEditorAxisSource::Construct(0, true, this);
+	AxisYSource = UGizmoEditorAxisSource::Construct(1, true, this);
+	AxisZSource = UGizmoEditorAxisSource::Construct(2, true, this);
+
+	auto MakeArrowObjectFunc = [](FVector Axis, UMaterialInterface* Material, UMaterialInterface* CurrentMaterial)
+	{
+		UGizmoArrowObject* ArrowObject = NewObject<UGizmoArrowObject>();
+		ArrowObject->CylinderObject->Direction = Axis;
+		ArrowObject->ConeObject->Direction = -Axis;
+		ArrowObject->SetMaterial(Material);
+		ArrowObject->SetCurrentMaterial(CurrentMaterial);
+		return ArrowObject;
+	};
+
+	AxisXObject = MakeArrowObjectFunc(FVector(1.0f, 0.0f, 0.0f), AxisMaterialX, CurrentAxisMaterial);
+	AxisYObject = MakeArrowObjectFunc(FVector(0.0f, 1.0f, 0.0f), AxisMaterialY, CurrentAxisMaterial);
+	AxisZObject = MakeArrowObjectFunc(FVector(0.0f, 0.0f, 1.0f), AxisMaterialZ, CurrentAxisMaterial);
+
+	UGizmoScaledAndUnscaledTransformSources* AxisTransformSource = UGizmoScaledAndUnscaledTransformSources::Construct(
+		UGizmoTransformProxyTransformSource::Construct(ActiveTarget, this),
+		UGizmoObjectWorldTransformSource::Construct(GizmoGroupObject, this));
+	// This state target emits an explicit FChange that moves the GizmoActor root component during undo/redo.
+	// It also opens/closes the Transaction that saves/restores the target object locations.
+	if (TransactionProvider == nullptr)
+	{
+		TransactionProvider = GetGizmoManager();
+	}
+	StateTarget = UGizmoObjectTransformChangeStateTarget::Construct(GizmoGroupObject,
+		LOCTEXT("UTransformGizmoTransaction", "Transform"), TransactionProvider, this);
+	StateTarget->DependentChangeSources.Add(MakeUnique<FTransformProxyChangeSource>(Target));
+
+	CameraAxisSource = NewObject<UGizmoConstantFrameAxisSource>(this);
+
+	{
+		AddAxisTranslationGizmo(AxisXObject, AxisXSource, AxisTransformSource, StateTarget, EAxisList::X, AxisColorX);
+		GizmoGroupObject->Add(AxisXObject);
+	}
+
+	{
+		AddAxisTranslationGizmo(AxisYObject, AxisYSource, AxisTransformSource, StateTarget, EAxisList::Y, AxisColorY);
+		GizmoGroupObject->Add(AxisYObject);
+	}
+
+	{
+		AddAxisTranslationGizmo(AxisZObject,  AxisZSource, AxisTransformSource, StateTarget, EAxisList::Z, AxisColorZ);
+		GizmoGroupObject->Add(AxisZObject);
+	}
+
+/*
+	// @todo: add plane translation
+	if (GizmoActor->TranslateYZ != nullptr)
+	{
+		AddPlaneTranslationGizmo(GizmoActor->TranslateYZ, GizmoComponent, AxisXSource, TransformSource, StateTarget);
+		ActiveComponents.Add(GizmoActor->TranslateYZ);
+	}
+	if (GizmoActor->TranslateXZ != nullptr)
+	{
+		AddPlaneTranslationGizmo(GizmoActor->TranslateXZ, GizmoComponent, AxisYSource, TransformSource, StateTarget);
+		ActiveComponents.Add(GizmoActor->TranslateXZ);
+	}
+	if (GizmoActor->TranslateXY != nullptr)
+	{
+		AddPlaneTranslationGizmo(GizmoActor->TranslateXY, GizmoComponent, AxisZSource, TransformSource, StateTarget);
+		ActiveComponents.Add(GizmoActor->TranslateXY);
+	}
+
+	// @todo: finish rotation implementation
+	if (GizmoActor->RotateX != nullptr)
+	{
+		AddAxisRotationGizmo(GizmoActor->RotateX, GizmoComponent, AxisXSource, SpaceSource, TransformSource, StateTarget, EAxisList::X, AxisColorX);
+		ActiveComponents.Add(GizmoActor->RotateX);
+	}
+	if (GizmoActor->RotateY != nullptr)
+	{
+		AddAxisRotationGizmo(GizmoActor->RotateY, GizmoComponent, AxisYSource, SpaceSource, TransformSource, StateTarget, EAxisList::Y, AxisColorY);
+		ActiveComponents.Add(GizmoActor->RotateY);
+	}
+	if (GizmoActor->RotateZ != nullptr)
+	{
+		AddAxisRotationGizmo(GizmoActor->RotateZ, GizmoComponent, AxisZSource, SpaceSource, TransformSource, StateTarget, EAxisList::Z, AxisColorZ);
+		ActiveComponents.Add(GizmoActor->RotateZ);
+	}
+*/
+
+	// Create objects for scale gizmo
+	ScaleAxisXObject = MakeArrowObjectFunc(FVector(1.0f, 0.0f, 0.0f), AxisMaterialX, CurrentAxisMaterial); 
+	ScaleAxisYObject = MakeArrowObjectFunc(FVector(0.0f, 1.0f, 0.0f), AxisMaterialY, CurrentAxisMaterial); 
+	ScaleAxisZObject = MakeArrowObjectFunc(FVector(0.0f, 0.0f, 1.0f), AxisMaterialZ, CurrentAxisMaterial); 
+
+	// only need these if scaling enabled. Essentially these are just the unit axes, regardless
+	// of what 3D axis is in use, we will tell the ParameterSource-to-3D-Scale mapper to
+	// use the coordinate axes
+	UnitAxisXSource = UGizmoEditorAxisSource::Construct(0, false, this);
+	UnitAxisYSource = UGizmoEditorAxisSource::Construct(1, false, this);
+	UnitAxisZSource = UGizmoEditorAxisSource::Construct(2, false, this);
+
+/*
+	// @todo: add uniform scale handle
+	if (GizmoActor->UniformScale != nullptr)
+	{
+		AddUniformScaleGizmo(GizmoActor->UniformScale, GizmoComponent, CameraAxisSource, CameraAxisSource, TransformSource, StateTarget);
+		ActiveComponents.Add(GizmoActor->UniformScale);
+	}
+*/
+
+	{
+		AddAxisScaleGizmo(ScaleAxisXObject, AxisXSource, UnitAxisXSource, AxisTransformSource, StateTarget, EAxisList::X, AxisColorX);
+		GizmoGroupObject->Add(ScaleAxisXObject);
+	}
+
+	{
+		AddAxisScaleGizmo(ScaleAxisYObject, AxisYSource, UnitAxisYSource, AxisTransformSource, StateTarget, EAxisList::Y, AxisColorY);
+		GizmoGroupObject->Add(ScaleAxisYObject);
+	}
+
+	{
+		AddAxisScaleGizmo(ScaleAxisZObject, AxisZSource, UnitAxisZSource, AxisTransformSource, StateTarget, EAxisList::Z, AxisColorZ);
+		GizmoGroupObject->Add(ScaleAxisZObject);
+	}
+
+/*
+	// @todo: add plane scale
+	if (GizmoActor->PlaneScaleYZ != nullptr)
+	{
+		AddPlaneScaleGizmo(GizmoActor->PlaneScaleYZ, GizmoComponent, AxisXSource, UnitAxisXSource, TransformSource, StateTarget);
+		ActiveComponents.Add(GizmoActor->PlaneScaleYZ);
+		NonuniformScaleComponents.Add(GizmoActor->PlaneScaleYZ);
+	}
+	if (GizmoActor->PlaneScaleXZ != nullptr)
+	{
+		AddPlaneScaleGizmo(GizmoActor->PlaneScaleXZ, GizmoComponent, AxisYSource, UnitAxisYSource, TransformSource, StateTarget);
+		ActiveComponents.Add(GizmoActor->PlaneScaleXZ);
+		NonuniformScaleComponents.Add(GizmoActor->PlaneScaleXZ);
+	}
+	if (GizmoActor->PlaneScaleXY != nullptr)
+	{
+		AddPlaneScaleGizmo(GizmoActor->PlaneScaleXY, GizmoComponent, AxisZSource, UnitAxisZSource, TransformSource, StateTarget);
+		ActiveComponents.Add(GizmoActor->PlaneScaleXY);
+		NonuniformScaleComponents.Add(GizmoActor->PlaneScaleXY);
+	}
+*/ 
+
+	GizmoGroupObject->SetVisibility(false);
+	CurrentMode = EGizmoTransformMode::None;
+}
+
+
+void UTransformGizmo::ReinitializeGizmoTransform(const FTransform& NewTransform)
+{
+	// @todo update gizmo objects here?
+
+	// The underlying proxy has an existing way to reinitialize its transform without callbacks.
+	TGuardValue<bool>(ActiveTarget->bSetPivotMode, true);
+	ActiveTarget->SetTransform(NewTransform);
+}
+
+
+void UTransformGizmo::SetNewGizmoTransform(const FTransform& NewTransform)
+{
+	// @todo update gizmo objects here?
+
+	check(ActiveTarget != nullptr);
+
+	StateTarget->BeginUpdate();
+
+	ActiveTarget->SetTransform(NewTransform);
+
+	StateTarget->EndUpdate();
+}
+
+
+// @todo: This should either be named to "SetScale" or removed, since it can be done with ReinitializeGizmoTransform
+void UTransformGizmo::SetNewChildScale(const FVector& NewChildScale)
+{
+	FTransform NewTransform = ActiveTarget->GetTransform();
+	NewTransform.SetScale3D(NewChildScale);
+
+	TGuardValue<bool>(ActiveTarget->bSetPivotMode, true);
+	ActiveTarget->SetTransform(NewTransform);
+}
+
+
+void UTransformGizmo::SetVisibility(bool bVisibleIn)
+{
+	bVisible = bVisibleIn;
+}
+
+UInteractiveGizmo* UTransformGizmo::AddAxisTranslationGizmo(
+	UGizmoArrowObject * InArrowObject,
+	IGizmoAxisSource * InAxisSource,
+	IGizmoTransformSource * InTransformSource,
+	IGizmoStateTarget * InStateTarget,
+	EAxisList::Type InAxisType,
+	const FLinearColor InAxisColor)
+{
+	check(InArrowObject);
+	check(InArrowObject->CylinderObject);
+	check(InArrowObject->ConeObject);
+
+	// create axis-position gizmo, axis-position parameter will drive translation
+	UAxisPositionGizmo* TranslateGizmo = Cast<UAxisPositionGizmo>(GetGizmoManager()->CreateGizmo(
+		UInteractiveGizmoManager::DefaultAxisPositionBuilderIdentifier, FString(), this));
+	check(TranslateGizmo);
+
+	InArrowObject->CylinderObject->Length = AXIS_LENGTH;
+	InArrowObject->CylinderObject->Radius = AXIS_RADIUS;
+	InArrowObject->bHasConeHead = true;
+
+	InArrowObject->ConeObject->Angle = FMath::DegreesToRadians(AXIS_CONE_ANGLE);
+	InArrowObject->ConeObject->Height = AXIS_CONE_HEIGHT;
+	InArrowObject->ConeObject->Offset = -(AXIS_LENGTH + AXIS_CONE_HEAD_OFFSET);
+
+	// axis source provides the translation axis
+	TranslateGizmo->AxisSource = Cast<UObject>(InAxisSource);
+
+	// parameter source maps axis-parameter-change to translation of TransformSource's transform
+	UGizmoEditorAxisTranslationParameterSource* ParamSource = UGizmoEditorAxisTranslationParameterSource::Construct(InAxisSource, InTransformSource, this);
+	ParamSource->AxisTranslationParameterSource->PositionConstraintFunction = [this](const FVector& Pos, FVector& Snapped) { return PositionSnapFunction(Pos, Snapped); };
+	TranslateGizmo->ParameterSource = ParamSource;
+
+	UGizmoObjectHitTarget* HitTarget = UGizmoObjectHitTarget::Construct(InArrowObject, this);
+
+	TranslateGizmo->HitTarget = HitTarget;
+	TranslateGizmo->StateTarget = Cast<UObject>(InStateTarget);
+
+	TranslateGizmo->ShouldUseCustomDestinationFunc = [this]() { return ShouldAlignDestination(); };
+	TranslateGizmo->CustomDestinationFunc =
+		[this](const UAxisPositionGizmo::FCustomDestinationParams& Params, FVector& OutputPoint) {
+		return DestinationAlignmentRayCaster(*Params.WorldRay, OutputPoint);
+	};
+
+	ActiveGizmos.Add(TranslateGizmo);
+
+	return TranslateGizmo;
+}
+
+UInteractiveGizmo* UTransformGizmo::AddPlaneTranslationGizmo(
+	IGizmoAxisSource* InAxisSource,
+	IGizmoTransformSource* InTransformSource,
+	IGizmoStateTarget* InStateTarget)
+{
+	/* @todo
+	// create axis-position gizmo, axis-position parameter will drive translation
+	UEditorPlanePositionGizmo* TranslateGizmo = Cast<UEditorPlanePositionGizmo>(GetGizmoManager()->CreateGizmo(
+		UEditorInteractiveGizmoManager::DefaultEditorPlanePositionBuilderIdentifier, FString(), this));
+	check(TranslateGizmo);
+
+	// axis source provides the translation axis
+	TranslateGizmo->AxisSource = Cast<UObject>(InAxisSource);
+
+	// parameter source maps axis-parameter-change to translation of TransformSource's transform
+	UGizmoPlaneTranslationParameterSource* ParamSource = UGizmoPlaneTranslationParameterSource::Construct(InAxisSource, InTransformSource, this);
+	ParamSource->PositionConstraintFunction = [this](const FVector& Pos, FVector& Snapped) { return PositionSnapFunction(Pos, Snapped); };
+	TranslateGizmo->ParameterSource = ParamSource;
+
+	// sub-component provides hit target
+	UGizmoComponentHitTarget* HitTarget = UGizmoComponentHitTarget::Construct(InAxisComponent, this);
+	if (this->UpdateHoverFunctionComponent)
+	{
+		HitTarget->UpdateHoverFunction = [InAxisComponent, this](bool bHovering) { this->UpdateHoverFunctionComponent(InAxisComponent, bHovering); };
+	}
+	TranslateGizmo->HitTarget = HitTarget;
+
+	TranslateGizmo->StateTarget = Cast<UObject>(InStateTarget);
+
+	TranslateGizmo->ShouldUseCustomDestinationFunc = [this]() { return ShouldAlignDestination(); };
+	TranslateGizmo->CustomDestinationFunc =
+		[this](const UEditorPlanePositionGizmo::FCustomDestinationParams& Params, FVector& OutputPoint) {
+		return DestinationAlignmentRayCaster(*Params.WorldRay, OutputPoint);
+	};
+
+	ActiveGizmos.Add(TranslateGizmo);
+	return TranslateGizmo;
+	*/
+	return nullptr;
+}
+
+UInteractiveGizmo* UTransformGizmo::AddAxisRotationGizmo(
+	IGizmoAxisSource* InAxisSource,
+	IGizmoTransformSource* InTransformSource,
+	IGizmoStateTarget* InStateTarget,
+	EAxisList::Type InAxisType,
+	const FLinearColor InAxisColor)
+{
+	/* @todo
+	// create axis-angle gizmo, angle will drive axis-rotation
+	UEditorAxisAngleGizmo* RotateGizmo = Cast<UEditorAxisAngleGizmo>(GetGizmoManager()->CreateGizmo(
+		UEditorInteractiveGizmoManager::DefaultEditorAxisAngleBuilderIdentifier, FString(), this));
+	check(RotateGizmo);
+
+	// axis source provides the translation axis
+	RotateGizmo->Initialize(InAxisType, InAxisColor, CurrentColor);
+
+	// axis source provides the rotation axis
+	RotateGizmo->AxisSource = Cast<UObject>(InAxisSource);
+
+	// space source provides coordinate space info specific to the gizmo
+	RotateGizmo->GizmoSpaceSource = Cast<UObject>(SpaceSource);
+
+	// parameter source maps angle-parameter-change to rotation of TransformSource's transform
+	UGizmoAxisRotationParameterSource* AngleSource = UGizmoAxisRotationParameterSource::Construct(InAxisSource, InTransformSource, this);
+	AngleSource->RotationConstraintFunction = [this](const FQuat& DeltaRotation){ return RotationSnapFunction(DeltaRotation); };
+	RotateGizmo->AngleSource = AngleSource;
+
+	// sub-component provides hit target
+	UGizmoComponentHitTarget* HitTarget = UGizmoComponentHitTarget::Construct(InAxisComponent, this);
+	if (this->UpdateHoverFunctionComponent)
+	{
+		HitTarget->UpdateHoverFunction = [InAxisComponent, this](bool bHovering) { this->UpdateHoverFunctionComponent(InAxisComponent, bHovering); };
+	}
+	RotateGizmo->HitTarget = HitTarget;
+
+	RotateGizmo->StateTarget = Cast<UObject>(StateTarget);
+
+	RotateGizmo->ShouldUseCustomDestinationFunc = [this]() { return ShouldAlignDestination(); };
+	RotateGizmo->CustomDestinationFunc =
+		[this](const UEditorAxisAngleGizmo::FCustomDestinationParams& Params, FVector& OutputPoint) {
+		return DestinationAlignmentRayCaster(*Params.WorldRay, OutputPoint);
+	};
+
+	ActiveGizmos.Add(RotateGizmo);
+
+	return RotateGizmo;
+	*/
+	return nullptr;
+}
+
+
+UInteractiveGizmo* UTransformGizmo::AddAxisScaleGizmo(
+	UGizmoArrowObject* InArrowObject,
+	IGizmoAxisSource* InGizmoAxisSource, IGizmoAxisSource* InParameterAxisSource,
+	IGizmoTransformSource* InTransformSource,
+	IGizmoStateTarget* InStateTarget,
+	EAxisList::Type InAxisType,
+	const FLinearColor InAxisColor)
+{
+	check(InArrowObject);
+	check(InArrowObject->CylinderObject);
+	check(InArrowObject->BoxObject);
+
+	InArrowObject->CylinderObject->Length = AXIS_LENGTH_SCALE;
+	InArrowObject->CylinderObject->Radius = AXIS_RADIUS;
+	InArrowObject->CylinderObject->Offset = AXIS_LENGTH_SCALE_OFFSET;
+	InArrowObject->bHasConeHead = false;
+
+	InArrowObject->BoxObject->Dimensions = FVector(AXIS_CUBE_SIZE, AXIS_CUBE_SIZE, AXIS_CUBE_SIZE);
+	InArrowObject->BoxObject->Offset = AXIS_LENGTH_SCALE + AXIS_LENGTH_SCALE_OFFSET + AXIS_CUBE_HEAD_OFFSET;
+	if (InAxisType == EAxisList::X)
+	{
+		InArrowObject->BoxObject->UpDirection = FVector(1, 0, 0);
+		InArrowObject->BoxObject->SideDirection = FVector(0, 1, 0);
+	}
+	else if (InAxisType == EAxisList::Y)
+	{
+		InArrowObject->BoxObject->UpDirection = FVector(0, 1, 0);
+		InArrowObject->BoxObject->SideDirection = FVector(0, 0, 1);
+	}
+	else
+	{
+		InArrowObject->BoxObject->UpDirection = FVector(0, 0, 1);
+		InArrowObject->BoxObject->SideDirection = FVector(1, 0, 0);
+	}
+
+	// create axis-position gizmo, axis-position parameter will drive scale
+	UAxisPositionGizmo* ScaleGizmo = Cast<UAxisPositionGizmo>(GetGizmoManager()->CreateGizmo(
+		UInteractiveGizmoManager::DefaultAxisPositionBuilderIdentifier, FString(), this));
+	ScaleGizmo->bEnableSignedAxis = true;
+	check(ScaleGizmo);
+
+	// axis source provides the translation axis
+	ScaleGizmo->AxisSource = Cast<UObject>(InGizmoAxisSource);
+
+	// parameter source maps axis-parameter-change to translation of TransformSource's transform
+	UGizmoAxisScaleParameterSource* ParamSource = UGizmoAxisScaleParameterSource::Construct(InParameterAxisSource, InTransformSource, this);
+	ParamSource->bClampToZero = bDisallowNegativeScaling;
+	ScaleGizmo->ParameterSource = ParamSource;
+
+	UGizmoObjectHitTarget* HitTarget = UGizmoObjectHitTarget::Construct(InArrowObject, this);
+
+	ScaleGizmo->HitTarget = HitTarget;
+	ScaleGizmo->StateTarget = Cast<UObject>(StateTarget);
+
+	ActiveGizmos.Add(ScaleGizmo);
+	return ScaleGizmo;
+}
+
+
+UInteractiveGizmo* UTransformGizmo::AddPlaneScaleGizmo(
+	IGizmoAxisSource* InGizmoAxisSource, IGizmoAxisSource* InParameterAxisSource,
+	IGizmoTransformSource* InTransformSource,
+	IGizmoStateTarget* InStateTarget)
+{
+	/* @todo
+	// create axis-position gizmo, axis-position parameter will drive scale
+	UEditorPlanePositionGizmo* ScaleGizmo = Cast<UEditorPlanePositionGizmo>(GetGizmoManager()->CreateGizmo(
+		UEditorInteractiveGizmoManager::DefaultEditorPlanePositionBuilderIdentifier, FString(), this));
+	ScaleGizmo->bEnableSignedAxis = true;
+	check(ScaleGizmo);
+
+	// axis source provides the translation axis
+	ScaleGizmo->AxisSource = Cast<UObject>(InGizmoAxisSource);
+
+	// parameter source maps axis-parameter-change to translation of TransformSource's transform
+	UGizmoPlaneScaleParameterSource* ParamSource = UGizmoPlaneScaleParameterSource::Construct(InParameterAxisSource, InTransformSource, this);
+	ParamSource->bClampToZero = bDisallowNegativeScaling;
+	ParamSource->bUseEqualScaling = true;
+	ScaleGizmo->ParameterSource = ParamSource;
+
+	// sub-component provides hit target
+	UGizmoComponentHitTarget* HitTarget = UGizmoComponentHitTarget::Construct(InAxisComponent, this);
+	if (this->UpdateHoverFunctionComponent)
+	{
+		HitTarget->UpdateHoverFunction = [InAxisComponent, this](bool bHovering) { this->UpdateHoverFunctionComponent(InAxisComponent, bHovering); };
+	}
+	ScaleGizmo->HitTarget = HitTarget;
+
+	ScaleGizmo->StateTarget = Cast<UObject>(StateTarget);
+
+	ActiveGizmos.Add(ScaleGizmo);
+	return ScaleGizmo;
+	*/
+	return nullptr;
+}
+
+
+UInteractiveGizmo* UTransformGizmo::AddUniformScaleGizmo(
+	IGizmoAxisSource* InGizmoAxisSource, IGizmoAxisSource* InParameterAxisSource,
+	IGizmoTransformSource* InTransformSource,
+	IGizmoStateTarget* InStateTarget)
+{
+	/* @todo
+	// create plane-position gizmo, plane-position parameter will drive scale
+	UEditorPlanePositionGizmo* ScaleGizmo = Cast<UEditorPlanePositionGizmo>(GetGizmoManager()->CreateGizmo(
+		UEditorInteractiveGizmoManager::DefaultEditorPlanePositionBuilderIdentifier, FString(), this));
+	check(ScaleGizmo);
+
+	// axis source provides the translation plane
+	ScaleGizmo->AxisSource = Cast<UObject>(InGizmoAxisSource);
+
+	// parameter source maps axis-parameter-change to translation of TransformSource's transform
+	UGizmoUniformScaleParameterSource* ParamSource = UGizmoUniformScaleParameterSource::Construct(InParameterAxisSource, InTransformSource, this);
+	//ParamSource->PositionConstraintFunction = [this](const FVector& Pos, FVector& Snapped) { return PositionSnapFunction(Pos, Snapped); };
+	ScaleGizmo->ParameterSource = ParamSource;
+
+	// sub-component provides hit target
+	UGizmoComponentHitTarget* HitTarget = UGizmoComponentHitTarget::Construct(InScaleComponent, this);
+	if (this->UpdateHoverFunctionComponent)
+	{
+		HitTarget->UpdateHoverFunction = [InScaleComponent, this](bool bHovering) { this->UpdateHoverFunctionComponent(InScaleComponent, bHovering); };
+	}
+	ScaleGizmo->HitTarget = HitTarget;
+
+	ScaleGizmo->StateTarget = Cast<UObject>(InStateTarget);
+
+	ActiveGizmos.Add(ScaleGizmo);
+	return ScaleGizmo;
+	*/
+	return nullptr;
+}
+
+
+void UTransformGizmo::ClearActiveTarget()
+{
+	for (UInteractiveGizmo* Gizmo : ActiveGizmos)
+	{
+		GetGizmoManager()->DestroyGizmo(Gizmo);
+	}
+	ActiveGizmos.SetNum(0);
+	ActiveObjects.SetNum(0);
+
+	CameraAxisSource = nullptr;
+	GizmoGroupObject = nullptr;
+	AxisXSource = nullptr;
+	AxisYSource = nullptr;
+	AxisZSource = nullptr;
+	AxisXObject = nullptr;
+	AxisYObject = nullptr;
+	AxisZObject = nullptr;
+	UnitAxisXSource = nullptr;
+	UnitAxisYSource = nullptr;
+	UnitAxisZSource = nullptr;
+	StateTarget = nullptr;
+
+	ActiveTarget = nullptr;
+}
+
+
+bool UTransformGizmo::PositionSnapFunction(const FVector& WorldPosition, FVector& SnappedPositionOut) const
+{
+	SnappedPositionOut = WorldPosition;
+
+	// only snap if we want snapping obvs
+	if (bSnapToWorldGrid == false)
+	{
+		return false;
+	}
+
+	// only snap to world grid when using world axes
+	if (GetGizmoManager()->GetContextQueriesAPI()->GetCurrentCoordinateSystem() != EToolContextCoordinateSystem::World)
+	{
+		return false;
+	}
+
+	FSceneSnapQueryRequest Request;
+	Request.RequestType = ESceneSnapQueryType::Position;
+	Request.TargetTypes = ESceneSnapQueryTargetType::Grid;
+	Request.Position = WorldPosition;
+	if ( bGridSizeIsExplicit )
+	{
+		Request.GridSize = ExplicitGridSize;
+	}
+	TArray<FSceneSnapQueryResult> Results;
+	if (GetGizmoManager()->GetContextQueriesAPI()->ExecuteSceneSnapQuery(Request, Results))
+	{
+		SnappedPositionOut = Results[0].Position;
+		return true;
+	};
+
+	return false;
+}
+
+
+FQuat UTransformGizmo::RotationSnapFunction(const FQuat& DeltaRotation) const
+{
+	FQuat SnappedDeltaRotation = DeltaRotation;
+
+	// only snap if we want snapping 
+	if (bSnapToWorldRotGrid)
+	{
+		FSceneSnapQueryRequest Request;
+		Request.RequestType   = ESceneSnapQueryType::Rotation;
+		Request.TargetTypes   = ESceneSnapQueryTargetType::Grid;
+		Request.DeltaRotation = DeltaRotation;
+		if ( bRotationGridSizeIsExplicit )
+		{
+			Request.RotGridSize = ExplicitRotationGridSize;
+		}
+		TArray<FSceneSnapQueryResult> Results;
+		if (GetGizmoManager()->GetContextQueriesAPI()->ExecuteSceneSnapQuery(Request, Results))
+		{
+			SnappedDeltaRotation = Results[0].DeltaRotation;
+		};
+	}
+	return SnappedDeltaRotation;
+}
+
+#undef LOCTEXT_NAMESPACE
diff --git a/Engine/Source/Editor/Experimental/EditorInteractiveToolsFramework/Private/EditorGizmos/TransformGizmoInterfaces.cpp b/Engine/Source/Editor/Experimental/EditorInteractiveToolsFramework/Private/EditorGizmos/TransformGizmoInterfaces.cpp
new file mode 100644
--- /dev/null
+++ b/Engine/Source/Editor/Experimental/EditorInteractiveToolsFramework/Private/EditorGizmos/TransformGizmoInterfaces.cpp
@@ -0,0 +1,27 @@
+// Copyright Epic Games, Inc. All Rights Reserved.
+
+#include "EditorGizmos/TransformGizmoInterfaces.h"
+#include "EditorModeManager.h"
+#include "EditorViewportClient.h"
+
+EGizmoTransformMode FTransformGizmoUtil::GetGizmoMode(UE::Widget::EWidgetMode InWidgetMode)
+{
+	switch (InWidgetMode)
+	{
+		case UE::Widget::EWidgetMode::WM_Translate: return EGizmoTransformMode::Translate;
+		case UE::Widget::EWidgetMode::WM_Rotate: return EGizmoTransformMode::Rotate;
+		case UE::Widget::EWidgetMode::WM_Scale: return EGizmoTransformMode::Scale;
+	}
+	return EGizmoTransformMode::None;
+}
+
+UE::Widget::EWidgetMode FTransformGizmoUtil::GetWidgetMode(EGizmoTransformMode InGizmoMode)
+{
+	switch (InGizmoMode)
+	{
+		case EGizmoTransformMode::Translate: return UE::Widget::EWidgetMode::WM_Translate;
+		case EGizmoTransformMode::Rotate: return UE::Widget::EWidgetMode::WM_Rotate;
+		case EGizmoTransformMode::Scale: return UE::Widget::EWidgetMode::WM_Scale;
+	}
+	return UE::Widget::EWidgetMode::WM_None;
+}
diff --git a/Engine/Source/Editor/Experimental/EditorInteractiveToolsFramework/Private/EditorInteractiveGizmoSubsystem.cpp b/Engine/Source/Editor/Experimental/EditorInteractiveToolsFramework/Private/EditorInteractiveGizmoSubsystem.cpp
--- a/Engine/Source/Editor/Experimental/EditorInteractiveToolsFramework/Private/EditorInteractiveGizmoSubsystem.cpp
+++ b/Engine/Source/Editor/Experimental/EditorInteractiveToolsFramework/Private/EditorInteractiveGizmoSubsystem.cpp
@@ -1,7 +1,7 @@
 // Copyright Epic Games, Inc. All Rights Reserved.
 
 #include "EditorInteractiveGizmoSubsystem.h"
-#include "EditorGizmos/EditorTransformGizmo.h"
+#include "EditorGizmos/EditorTransformGizmoBuilder.h"
 
 #define LOCTEXT_NAMESPACE "UEditorInteractiveGizmoSubsystem"
 
diff --git a/Engine/Source/Editor/Experimental/EditorInteractiveToolsFramework/Public/EditorGizmos/EditorTransformGizmo.h b/Engine/Source/Editor/Experimental/EditorInteractiveToolsFramework/Public/EditorGizmos/EditorTransformGizmo.h
deleted file mode 100644
diff --git a/Engine/Source/Editor/Experimental/EditorInteractiveToolsFramework/Public/EditorGizmos/EditorTransformGizmoBuilder.h b/Engine/Source/Editor/Experimental/EditorInteractiveToolsFramework/Public/EditorGizmos/EditorTransformGizmoBuilder.h
new file mode 100644
--- /dev/null
+++ b/Engine/Source/Editor/Experimental/EditorInteractiveToolsFramework/Public/EditorGizmos/EditorTransformGizmoBuilder.h
@@ -0,0 +1,24 @@
+// Copyright Epic Games, Inc. All Rights Reserved.
+
+#pragma once
+
+#include "CoreMinimal.h"
+#include "BaseGizmos/TransformProxy.h"
+#include "EditorInteractiveGizmoSelectionBuilder.h"
+#include "EditorTransformGizmoBuilder.generated.h"
+
+class UInteractiveGizmo;
+
+UCLASS()
+class EDITORINTERACTIVETOOLSFRAMEWORK_API UEditorTransformGizmoBuilder : public UInteractiveGizmoBuilder, public IEditorInteractiveGizmoSelectionBuilder
+{
+	GENERATED_BODY()
+
+public:
+
+	// UEditorInteractiveGizmoSelectionBuilder interface 
+	virtual UInteractiveGizmo* BuildGizmo(const FToolBuilderState& SceneState) const override;
+	virtual void UpdateGizmoForSelection(UInteractiveGizmo* Gizmo, const FToolBuilderState& SceneState) override;
+};
+
+
diff --git a/Engine/Source/Editor/Experimental/EditorInteractiveToolsFramework/Public/EditorGizmos/EditorTransformGizmoSource.h b/Engine/Source/Editor/Experimental/EditorInteractiveToolsFramework/Public/EditorGizmos/EditorTransformGizmoSource.h
new file mode 100644
--- /dev/null
+++ b/Engine/Source/Editor/Experimental/EditorInteractiveToolsFramework/Public/EditorGizmos/EditorTransformGizmoSource.h
@@ -0,0 +1,61 @@
+// Copyright Epic Games, Inc. All Rights Reserved.
+
+#pragma once
+
+#include "CoreMinimal.h"
+#include "EditorGizmos/TransformGizmoInterfaces.h"
+#include "EditorTransformGizmoSource.generated.h"
+
+class FEditorViewportClient;
+class FSceneView;
+class FEditorModeTools;
+
+/**
+ * UEditorTransformGizmoSource is an ITransformGizmoSource implementation that provides
+ * current state information used to configure the Editor transform gizmo.
+ */
+UCLASS()
+class EDITORINTERACTIVETOOLSFRAMEWORK_API UEditorTransformGizmoSource : public UObject, public ITransformGizmoSource
+{
+	GENERATED_BODY()
+public:
+	
+	/**
+	 * @return The current display mode for the Editor transform gizmo
+	 */
+	virtual EGizmoTransformMode GetGizmoMode() const;
+
+	/**
+	 * @return The current axes to draw for the specified mode
+	 */
+	virtual EAxisList::Type GetGizmoAxisToDraw(EGizmoTransformMode InWidgetMode) const;
+
+	/**
+	 * @return The coordinate system space (world or local) to display the widget in
+	 */
+	virtual EToolContextCoordinateSystem GetGizmoCoordSystemSpace() const;
+
+	/**
+	 * Returns a scale factor for the gizmo
+	 */
+	virtual float GetGizmoScale() const;
+
+	/**
+	 * Whether the gizmo is visible
+	 */
+	virtual bool GetVisible() const;
+
+public:
+	static UEditorTransformGizmoSource* Construct(
+		UObject* Outer = (UObject*)GetTransientPackage())
+	{
+		return NewObject<UEditorTransformGizmoSource>(Outer);
+	}
+
+protected:
+
+	FEditorModeTools& GetModeTools() const;
+
+	FEditorViewportClient* GetViewportClient() const;
+};
+
diff --git a/Engine/Source/Editor/Experimental/EditorInteractiveToolsFramework/Public/EditorGizmos/EditorTransformProxy.h b/Engine/Source/Editor/Experimental/EditorInteractiveToolsFramework/Public/EditorGizmos/EditorTransformProxy.h
new file mode 100644
--- /dev/null
+++ b/Engine/Source/Editor/Experimental/EditorInteractiveToolsFramework/Public/EditorGizmos/EditorTransformProxy.h
@@ -0,0 +1,36 @@
+// Copyright Epic Games, Inc. All Rights Reserved.
+
+#pragma once
+
+#include "CoreMinimal.h"
+#include "UObject/ObjectMacros.h"
+#include "BaseGizmos/TransformProxy.h"
+#include "EditorTransformProxy.generated.h"
+
+/**
+ * UEditorTransformProxy is a derivation of UTransformProxy that
+ * returns the transform that defines the current space of the default
+ * Editor transform gizmo for a given mode manager / viewport.
+ * 
+ * @todo Currently this defaults internally to GLevelEditorModeManager()
+ * but eventually it should be possible to set and use a different mode
+ * manager.
+ */
+UCLASS(Transient)
+class EDITORINTERACTIVETOOLSFRAMEWORK_API UEditorTransformProxy : public UTransformProxy
+{
+	GENERATED_BODY()
+public:
+
+	/**
+	 * @return the shared transform for all the sub-objects
+	 */
+	virtual FTransform GetTransform() const override;
+
+	/**
+	 * Update the main transform and then update the sub-objects based on their relative transformations
+	 */
+	virtual void SetTransform(const FTransform& Transform) override;
+
+};
+
diff --git a/Engine/Source/Editor/Experimental/EditorInteractiveToolsFramework/Public/EditorGizmos/TransformGizmo.h b/Engine/Source/Editor/Experimental/EditorInteractiveToolsFramework/Public/EditorGizmos/TransformGizmo.h
new file mode 100644
--- /dev/null
+++ b/Engine/Source/Editor/Experimental/EditorInteractiveToolsFramework/Public/EditorGizmos/TransformGizmo.h
@@ -0,0 +1,344 @@
+// Copyright Epic Games, Inc. All Rights Reserved.
+
+#pragma once
+
+#include "CoreMinimal.h"
+#include "BaseGizmos/TransformProxy.h"
+#include "EditorGizmos/EditorAxisSources.h"
+#include "EditorGizmos/GizmoArrowObject.h"
+#include "EditorGizmos/GizmoBaseObject.h"
+#include "EditorGizmos/GizmoGroupObject.h"
+#include "EditorGizmos/GizmoObjectHitTargets.h"
+#include "EditorGizmos/GizmoObjectStateTargets.h"
+#include "EditorGizmos/GizmoObjectTransformSources.h"
+#include "EditorGizmos/TransformGizmoInterfaces.h"
+#include "InteractiveGizmo.h"
+#include "InteractiveToolObjects.h"
+#include "InteractiveToolChange.h"
+#include "Materials/MaterialInstanceDynamic.h"
+#include "Math/Axis.h"
+#include "TransformGizmo.generated.h"
+
+class UInteractiveGizmoManager;
+class IGizmoAxisSource;
+class IGizmoTransformSource;
+class IGizmoStateTarget;
+class UGizmoConstantFrameAxisSource;
+class UGizmoTransformChangeStateTarget;
+class FEditorTransformGizmoTransformChange;
+
+
+/**
+ * UTransformGizmo provides standard Transformation Gizmo interactions,
+ * applied to a UTransformProxy target object. By default the Gizmo will be
+ * a standard XYZ translate/rotate Gizmo (axis and plane translation).
+ */
+UCLASS()
+class EDITORINTERACTIVETOOLSFRAMEWORK_API UTransformGizmo : public UInteractiveGizmo
+{
+	GENERATED_BODY()
+
+public:
+	constexpr static float AXIS_LENGTH = 35.0f;
+	constexpr static float AXIS_RADIUS = 1.2f;
+	constexpr static float AXIS_CONE_ANGLE = 15.7f;
+	constexpr static float AXIS_CONE_HEIGHT = 13;
+	constexpr static float AXIS_CONE_HEAD_OFFSET = 12;
+	constexpr static float AXIS_CUBE_SIZE = 4;
+	constexpr static float AXIS_CUBE_HEAD_OFFSET= 3;
+	constexpr static float TRANSLATE_ROTATE_AXIS_CIRCLE_RADIUS = 20.0f;
+	constexpr static float TWOD_AXIS_CIRCLE_RADIUS = 10.0f;
+	constexpr static float INNER_AXIS_CIRCLE_RADIUS = 48.0f;
+	constexpr static float OUTER_AXIS_CIRCLE_RADIUS = 56.0f;
+	constexpr static float ROTATION_TEXT_RADIUS = 75.0f;
+	constexpr static int32 AXIS_CIRCLE_SIDES = 24;
+	constexpr static float ARCALL_RELATIVE_INNER_SIZE = 0.75f;
+	constexpr static float AXIS_LENGTH_SCALE = 25.0f;
+	constexpr static float AXIS_LENGTH_SCALE_OFFSET = 5.0f;
+
+	constexpr static FLinearColor AxisColorX = FLinearColor(0.594f, 0.0197f, 0.0f);
+	constexpr static FLinearColor AxisColorY = FLinearColor(0.1349f, 0.3959f, 0.0f);
+	constexpr static FLinearColor AxisColorZ = FLinearColor(0.0251f, 0.207f, 0.85f);
+	constexpr static FLinearColor ScreenAxisColor = FLinearColor(0.76, 0.72, 0.14f);
+	constexpr static FColor PlaneColorXY = FColor(255, 255, 0); // FColor::Yellow
+	constexpr static FColor ArcBallColor = FColor(128, 128, 128, 6);
+	constexpr static FColor ScreenSpaceColor = FColor(196, 196, 196);
+	constexpr static FColor CurrentColor = FColor(255, 255, 0); // FColor::Yellow
+
+public:
+
+	/**
+	 * By default, the nonuniform scale components can scale negatively. However, they can be made to clamp
+	 * to zero instead by passing true here. This is useful for using the gizmo to flatten geometry.
+	 *
+	 * TODO: Should this affect uniform scaling too?
+	 */
+	virtual void SetDisallowNegativeScaling(bool bDisallow);
+
+	// UInteractiveGizmo overrides
+	virtual void Setup() override;
+	virtual void Shutdown() override;
+	virtual void Render(IToolsContextRenderAPI* RenderAPI);
+	virtual void Tick(float DeltaTime) override;
+
+
+	/**
+	 * Set the active target object for the Gizmo
+	 * @param Target active target
+	 * @param TransactionProvider optional IToolContextTransactionProvider implementation to use - by default uses GizmoManager
+	 */
+	virtual void SetActiveTarget(UTransformProxy* Target, IToolContextTransactionProvider* TransactionProvider = nullptr);
+
+	/**
+	 * Clear the active target object for the Gizmo
+	 */
+	virtual void ClearActiveTarget();
+
+	/** The active target object for the Gizmo */
+	UPROPERTY()
+	TObjectPtr<UTransformProxy> ActiveTarget;
+
+	/**
+	 * Repositions the gizmo without issuing undo/redo changes, triggering callbacks, 
+	 * or moving any components. Useful for resetting the gizmo to a new location without
+	 * it being viewed as a gizmo manipulation.
+	 */
+	void ReinitializeGizmoTransform(const FTransform& NewTransform);
+
+	/**
+	 * Set a new position for the Gizmo. This is done via the same mechanisms as the sub-gizmos,
+	 * so it generates the same Change/Modify() events, and hence works with Undo/Redo
+	 */
+	virtual void SetNewGizmoTransform(const FTransform& NewTransform);
+
+	/**
+	 * Explicitly set the child scale. Mainly useful to "reset" the child scale to (1,1,1) when re-using Gizmo across multiple transform actions.
+	 * @warning does not generate change/modify events!!
+	 */
+	virtual void SetNewChildScale(const FVector& NewChildScale);
+
+	/**
+	 * Set visibility for this Gizmo
+	 */
+	virtual void SetVisibility(bool bVisible);
+
+	/**
+	 * Whether gizmo is visible.
+	 */
+	UPROPERTY()
+	bool bVisible = false;
+
+	/**
+	 * If true, then when using world frame, Axis and Plane translation snap to the world grid via the ContextQueriesAPI (in PositionSnapFunction)
+	 */
+	UPROPERTY()
+	bool bSnapToWorldGrid = false;
+
+	/**
+	 * Optional grid size which overrides the Context Grid
+	 */
+	UPROPERTY()
+	bool bGridSizeIsExplicit = false;
+	UPROPERTY()
+	FVector ExplicitGridSize;
+
+	/**
+	 * Optional grid size which overrides the Context Rotation Grid
+	 */
+	UPROPERTY()
+	bool bRotationGridSizeIsExplicit = false;
+	UPROPERTY()
+	FRotator ExplicitRotationGridSize;
+
+	/**
+	 * If true, then when using world frame, Axis and Plane translation snap to the world grid via the ContextQueriesAPI (in RotationSnapFunction)
+	 */
+	UPROPERTY()
+	bool bSnapToWorldRotGrid = false;
+
+	//
+	// Transform Source
+	//
+	UPROPERTY()
+	TScriptInterface<ITransformGizmoSource> TransformSource;
+
+protected:
+
+	/** List of current-active gizmo objects */
+	UPROPERTY()
+	TArray<TObjectPtr<UGizmoBaseObject>> ActiveObjects;
+
+	/** list of currently-active child gizmos */
+	UPROPERTY()
+	TArray<TObjectPtr<UInteractiveGizmo>> ActiveGizmos;
+
+	//
+	// Axis Sources
+	//
+	/** Axis that points towards camera, X/Y plane tangents aligned to right/up. Shared across Gizmos, and created internally during SetActiveTarget() */
+	UPROPERTY()
+	TObjectPtr<UGizmoConstantFrameAxisSource> CameraAxisSource;
+
+	// internal function that updates CameraAxisSource by getting current view state from GizmoManager
+	void UpdateCameraAxisSource();
+
+	/** Gizmo group object contains all gizmos created internally during SetActiveTarget() */
+	UPROPERTY()
+	TObjectPtr<UGizmoGroupObject> GizmoGroupObject;
+
+	// @todo - these properties can probably be removed since they are contained within the group now.
+	/** X-axis source is shared across Gizmos, and created internally during SetActiveTarget() */
+	UPROPERTY()
+	TObjectPtr<UGizmoEditorAxisSource> AxisXSource;
+
+	/** Y-axis source is shared across Gizmos, and created internally during SetActiveTarget() */
+	UPROPERTY()
+	TObjectPtr<UGizmoEditorAxisSource> AxisYSource;
+
+	/** Z-axis source is shared across Gizmos, and created internally during SetActiveTarget() */
+	UPROPERTY()
+	TObjectPtr<UGizmoEditorAxisSource> AxisZSource;
+
+	/** X-axis arrow object is shared across Gizmos, and created internally during SetActiveTarget() */
+	UPROPERTY()
+	TObjectPtr<UGizmoArrowObject> AxisXObject;
+
+	/** Y-axis arrow object is shared across Gizmos, and created internally during SetActiveTarget() */
+	UPROPERTY()
+	TObjectPtr<UGizmoArrowObject> AxisYObject;
+
+	/** Z-axis arrow object is shared across Gizmos, and created internally during SetActiveTarget() */
+	UPROPERTY()
+	TObjectPtr<UGizmoArrowObject> AxisZObject;
+	//
+	// Scaling support. 
+	// UE Components only support scaling in local coordinates, so we have to create separate sources for that.
+	//
+
+	/** Local X-axis source (ie 1,0,0) is shared across Scale Gizmos, and created internally during SetActiveTarget() */
+	UPROPERTY()
+	TObjectPtr<UGizmoEditorAxisSource> UnitAxisXSource;
+
+	/** Y-axis source (ie 0,1,0) is shared across Scale Gizmos, and created internally during SetActiveTarget() */
+	UPROPERTY()
+	TObjectPtr<UGizmoEditorAxisSource> UnitAxisYSource;
+
+	/** Z-axis source (ie 0,0,1) is shared across Scale Gizmos, and created internally during SetActiveTarget() */
+	UPROPERTY()
+	TObjectPtr<UGizmoEditorAxisSource> UnitAxisZSource;
+
+	/** X-axis arrow object is shared across Gizmos, and created internally during SetActiveTarget() */
+	UPROPERTY()
+	TObjectPtr<UGizmoArrowObject> ScaleAxisXObject;
+
+	/** Y-axis arrow object is shared across Gizmos, and created internally during SetActiveTarget() */
+	UPROPERTY()
+	TObjectPtr<UGizmoArrowObject> ScaleAxisYObject;
+
+	/** Z-axis arrow object is shared across Gizmos, and created internally during SetActiveTarget() */
+	UPROPERTY()
+	TObjectPtr<UGizmoArrowObject> ScaleAxisZObject;
+
+	/** 
+	 * State target is shared across gizmos, and created internally during SetActiveTarget(). 
+	 * Several FChange providers are registered with this StateTarget, including the UTransformGizmo
+	 * itself (IToolCommandChangeSource implementation above is called)
+	 */
+	UPROPERTY()
+	TObjectPtr<UGizmoObjectTransformChangeStateTarget> StateTarget;
+
+	/**
+	 * These are used to let the translation subgizmos use raycasts into the scene to align the gizmo with scene geometry.
+	 * See comment for SetWorldAlignmentFunctions().
+	 */
+	TUniqueFunction<bool()> ShouldAlignDestination = []() { return false; };
+	TUniqueFunction<bool(const FRay&, FVector&)> DestinationAlignmentRayCaster = [](const FRay&, FVector&) {return false; };
+
+	bool bDisallowNegativeScaling = false;
+protected:
+
+	/** Update current gizmo mode based on transform source */
+	void UpdateMode();
+
+	/** Update current coord system based on the active target and transform source */
+	void UpdateCoordSystem();
+
+	/** Enable the given mode with the specified axes, EAxisList::Type::None will hide objects associated with mode */
+	void EnableMode(EGizmoTransformMode InGizmoMode, EAxisList::Type InAxisListToDraw);
+
+	/** Enable object if the specified InGizmoAxis is enabled in InAxisListToDraw*/
+	void EnableObject(UGizmoBaseObject* InGizmoObject, EAxisList::Type InGizmoAxis, EAxisList::Type InAxisListToDraw);
+
+	/** Enable translate using specified axis list */
+	void EnableTranslate(EAxisList::Type InAxisListToDraw);
+
+	/** Enable rotate using specified axis list */
+	void EnableRotate(EAxisList::Type InAxisListToDraw);
+
+	/** Enable scale using specified axis list */
+	void EnableScale(EAxisList::Type InAxisListToDraw);
+
+	UPROPERTY()
+	EGizmoTransformMode CurrentMode = EGizmoTransformMode::None;
+
+	UPROPERTY()
+	TEnumAsByte<EAxisList::Type> CurrentAxisToDraw = EAxisList::None;
+
+	/** @return a new instance of the standard axis-translation Gizmo */
+	virtual UInteractiveGizmo* AddAxisTranslationGizmo(
+		UGizmoArrowObject* InArrowObject,
+		IGizmoAxisSource* InAxisSource,
+		IGizmoTransformSource* InTransformSource,
+		IGizmoStateTarget* InStateTarget,
+		EAxisList::Type InAxisType,
+		const FLinearColor InAxisColor);
+
+	/** @return a new instance of the standard plane-translation Gizmo */
+	virtual UInteractiveGizmo* AddPlaneTranslationGizmo(
+		IGizmoAxisSource* InAxisSource,
+		IGizmoTransformSource* InTransformSource,
+		IGizmoStateTarget* InStateTarget);
+
+	/** @return a new instance of the standard axis-rotation Gizmo */
+	virtual UInteractiveGizmo* AddAxisRotationGizmo(
+		IGizmoAxisSource* InAxisSource,
+		IGizmoTransformSource* InTransformSource,
+		IGizmoStateTarget* InStateTarget,
+		EAxisList::Type InAxisType,
+		const FLinearColor InAxisColor);
+
+	/** @return a new instance of the standard axis-scaling Gizmo */
+	virtual UInteractiveGizmo* AddAxisScaleGizmo(
+		UGizmoArrowObject* InArrowObject,
+		IGizmoAxisSource* InGizmoAxisSource, IGizmoAxisSource* InParameterAxisSource,
+		IGizmoTransformSource* InTransformSource,
+		IGizmoStateTarget* InStateTarget,
+		EAxisList::Type InAxisType,
+		const FLinearColor InAxisColor);
+
+	/** @return a new instance of the standard plane-scaling Gizmo */
+	virtual UInteractiveGizmo* AddPlaneScaleGizmo(
+		IGizmoAxisSource* InGizmoAxisSource, IGizmoAxisSource* InParameterAxisSource,
+		IGizmoTransformSource* InTransformSource,
+		IGizmoStateTarget* InStateTarget);
+
+	/** @return a new instance of the standard plane-scaling Gizmo */
+	virtual UInteractiveGizmo* AddUniformScaleGizmo(
+		IGizmoAxisSource* InGizmoAxisSource, IGizmoAxisSource* InParameterAxisSource,
+		IGizmoTransformSource* InTransformSource,
+		IGizmoStateTarget* InStateTarget);
+
+	// Axis and Plane TransformSources use this function to execute worldgrid snap queries
+	bool PositionSnapFunction(const FVector& WorldPosition, FVector& SnappedPositionOut) const;
+	FQuat RotationSnapFunction(const FQuat& DeltaRotation) const;
+
+	/** Materials and colors to be used when drawing the items for each axis */
+	TObjectPtr<UMaterialInterface> TransparentPlaneMaterialXY;
+	TObjectPtr<UMaterialInterface> GridMaterial;
+
+	TObjectPtr<UMaterialInstanceDynamic> AxisMaterialX;
+	TObjectPtr<UMaterialInstanceDynamic> AxisMaterialY;
+	TObjectPtr<UMaterialInstanceDynamic> AxisMaterialZ;
+	TObjectPtr<UMaterialInstanceDynamic> CurrentAxisMaterial;
+	TObjectPtr<UMaterialInstanceDynamic> OpaquePlaneMaterialXY;
+};
diff --git a/Engine/Source/Editor/Experimental/EditorInteractiveToolsFramework/Public/EditorGizmos/TransformGizmoInterfaces.h b/Engine/Source/Editor/Experimental/EditorInteractiveToolsFramework/Public/EditorGizmos/TransformGizmoInterfaces.h
new file mode 100644
--- /dev/null
+++ b/Engine/Source/Editor/Experimental/EditorInteractiveToolsFramework/Public/EditorGizmos/TransformGizmoInterfaces.h
@@ -0,0 +1,94 @@
+// Copyright Epic Games, Inc. All Rights Reserved.
+
+#pragma once
+
+#include "CoreMinimal.h"
+#include "ToolContextInterfaces.h"
+#include "UnrealWidgetFwd.h"
+#include "TransformGizmoInterfaces.generated.h"
+
+//
+// Interface for the Transform gizmo.
+//
+UENUM()
+enum class EGizmoTransformMode : uint8
+{
+	None = 0,
+	Translate,
+	Rotate,
+	Scale,
+	Max
+};
+
+
+namespace FTransformGizmoUtil
+{
+	/** Convert UE::Widget::EWidgetMode to ETransformGizmoMode*/
+	EGizmoTransformMode GetGizmoMode(UE::Widget::EWidgetMode InWidgetMode);
+
+	/** Convert EEditorGizmoMode to UE::Widget::EWidgetMode*/
+	UE::Widget::EWidgetMode GetWidgetMode(EGizmoTransformMode InGizmoMode);
+};
+
+UINTERFACE()
+class EDITORINTERACTIVETOOLSFRAMEWORK_API UTransformGizmoSource : public UInterface
+{
+	GENERATED_BODY()
+};
+
+/**
+ * ITransformGizmoSource is an interface for providing gizmo mode configuration information.
+ */
+class EDITORINTERACTIVETOOLSFRAMEWORK_API ITransformGizmoSource
+{
+	GENERATED_BODY()
+public:
+	/**
+	 * Returns the current Editor gizmo mode
+	 */
+	virtual EGizmoTransformMode GetGizmoMode() const PURE_VIRTUAL(ITransformGizmoSource::GetGizmoMode, return EGizmoTransformMode::None;);
+
+	/**
+	 * Returns the current gizmo axes to draw
+	 */
+	virtual EAxisList::Type GetGizmoAxisToDraw(EGizmoTransformMode InGizmoMode) const PURE_VIRTUAL(ITransformGizmoSource::GetGizmoAxisToDraw, return EAxisList::None;);
+
+	/**
+	 * Returns the coordinate system space (world or local) in which to display the gizmo
+	 */
+	virtual EToolContextCoordinateSystem GetGizmoCoordSystemSpace() const PURE_VIRTUAL(ITransformGizmoSource::GetGizmoCoordSystemSpace, return EToolContextCoordinateSystem::World;);
+
+	/**
+	 * Returns a scale factor for the gizmo
+	 */
+	virtual float GetGizmoScale() const PURE_VIRTUAL(ITransformGizmoSource::GetGizmoScale, return 1.0f;);
+
+	/* 
+	 * Returns whether the gizmo should be visible. 
+	 */
+	virtual bool GetVisible() const PURE_VIRTUAL(ITransformGizmoSource::GetVisible, return false;);
+
+};
+
+
+UINTERFACE()
+class EDITORINTERACTIVETOOLSFRAMEWORK_API UTransformGizmoTarget : public UInterface
+{
+	GENERATED_BODY()
+};
+
+/**
+ * ITransformGizmoTarget is an interface for updating transform gizmo related state
+ * which matches interfaces provided in the Editor.
+ */
+class EDITORINTERACTIVETOOLSFRAMEWORK_API ITransformGizmoTarget
+{
+	GENERATED_BODY()
+public:
+
+	/*
+	 * Set the current axis when the gizmo is interacting
+	 */
+	virtual void SetCurrentAxis(EAxisList::Type CurrentAxis) PURE_VIRTUAL(ITransformGizmoTarget::SetCurrentAxis, );
+};
+
diff --git a/Engine/Source/Editor/UnrealEd/Private/UnrealWidgetRender.cpp b/Engine/Source/Editor/UnrealEd/Private/UnrealWidgetRender.cpp
--- a/Engine/Source/Editor/UnrealEd/Private/UnrealWidgetRender.cpp
+++ b/Engine/Source/Editor/UnrealEd/Private/UnrealWidgetRender.cpp
@@ -181,8 +181,17 @@ void FWidget::Render( const FSceneView* View,FPrimitiveDrawInterface* PDI, FEdit
 	const bool bShowFlagsSupportsWidgetDrawing = View->Family->EngineShowFlags.ModeWidgets;
 	const bool bEditorModeToolsSupportsWidgetDrawing = EditorModeTools ? EditorModeTools->GetShowWidget() : true;
 
-	static const auto UseLegacyWidgetCVar = IConsoleManager::Get().FindConsoleVariable(TEXT("Gizmos.UseLegacyWidget"));
-	const bool bUseLegacyWidget = UseLegacyWidgetCVar->GetInt() > 0;
+	// Use legacy widget for TranslateRotateZ or 2D modes or when UseLegacyWidget cvar is true
+	UE::Widget::EWidgetMode WidgetMode = ViewportClient->GetWidgetMode();
+	bool bUseLegacyWidget = (WidgetMode == UE::Widget::EWidgetMode::WM_TranslateRotateZ || WidgetMode == UE::Widget::EWidgetMode::WM_2D);
+	if (!bUseLegacyWidget)
+	{
+		static IConsoleVariable* const UseLegacyWidgetCVar = IConsoleManager::Get().FindConsoleVariable(TEXT("Gizmos.UseLegacyWidget"));
+		if (ensure(UseLegacyWidgetCVar))
+		{
+			bUseLegacyWidget = UseLegacyWidgetCVar->GetInt() > 0;
+		}
+	}
 
 	bool bDrawWidget;
 
@@ -219,7 +228,7 @@ void FWidget::Render( const FSceneView* View,FPrimitiveDrawInterface* PDI, FEdit
 		Origin = NewOrigin;
 	}
 
-	switch( ViewportClient->GetWidgetMode() )
+	switch( WidgetMode )
 	{
 		case UE::Widget::EWidgetMode::WM_Translate:
 			Render_Translate(View, PDI, ViewportClient, WidgetLocation, bDrawWidget);
diff --git a/Engine/Source/Runtime/CoreUObject/Public/UObject/NoExportTypes.h b/Engine/Source/Runtime/CoreUObject/Public/UObject/NoExportTypes.h
--- a/Engine/Source/Runtime/CoreUObject/Public/UObject/NoExportTypes.h
+++ b/Engine/Source/Runtime/CoreUObject/Public/UObject/NoExportTypes.h
@@ -120,6 +120,32 @@ namespace EAxis
 	};
 }
 
+/** Generic axis list enum (mirrored for native use in Axis.h). */
+UENUM()
+namespace EAxisList
+{
+	enum Type
+	{
+		None = 0,
+		X = 1,
+		Y = 2,
+		Z = 4,
+
+		Screen = 8,
+		XY = X | Y,
+		XZ = X | Z,
+		YZ = Y | Z,
+		XYZ = X | Y | Z,
+		All = XYZ | Screen,
+
+		/** alias over Axis YZ since it isn't used when the z-rotation widget is being used */
+		ZRotation = YZ,
+
+		/** alias over Screen since it isn't used when the 2d translate rotate widget is being used */
+		Rotate2D = Screen,
+	};
+}
+
 /** Describes shape of an interpolation curve (mirrored from InterpCurvePoint.h). */
 UENUM()
 enum EInterpCurveMode
