Non-unity fix.

diff --git a/Engine/Source/Editor/Experimental/EditorInteractiveToolsFramework/Private/EditorInteractiveGizmoSubsystem.cpp b/Engine/Source/Editor/Experimental/EditorInteractiveToolsFramework/Private/EditorInteractiveGizmoSubsystem.cpp
--- a/Engine/Source/Editor/Experimental/EditorInteractiveToolsFramework/Private/EditorInteractiveGizmoSubsystem.cpp
+++ b/Engine/Source/Editor/Experimental/EditorInteractiveToolsFramework/Private/EditorInteractiveGizmoSubsystem.cpp
@@ -1,6 +1,7 @@
 // Copyright Epic Games, Inc. All Rights Reserved.
 
 #include "EditorInteractiveGizmoSubsystem.h"
+#include "Engine/Engine.h"
 #include "EditorGizmos/EditorTransformGizmoBuilder.h"
 
 #define LOCTEXT_NAMESPACE "UEditorInteractiveGizmoSubsystem"
