InteractiveToolsFramework: Fix mouse wheel not being captured by UModeManagerInteractiveToolsContext::InputKey.

#rnx

diff --git a/Engine/Source/Editor/Experimental/EditorInteractiveToolsFramework/Private/EdModeInteractiveToolsContext.cpp b/Engine/Source/Editor/Experimental/EditorInteractiveToolsFramework/Private/EdModeInteractiveToolsContext.cpp
--- a/Engine/Source/Editor/Experimental/EditorInteractiveToolsFramework/Private/EdModeInteractiveToolsContext.cpp
+++ b/Engine/Source/Editor/Experimental/EditorInteractiveToolsFramework/Private/EdModeInteractiveToolsContext.cpp
@@ -948,7 +948,7 @@ bool UModeManagerInteractiveToolsContext::InputKey(FEditorViewportClient* Viewpo
 					: (Key == EKeys::MouseScrollUp) ? 1 
 					: -1;
 
-				InputRouter->PostInputEvent(InputState);
+				return InputRouter->PostInputEvent(InputState);
 			}
 		}
 		else if (Key.IsGamepadKey())
