Fix an issue where ITF tools would not clean up when their contexts were removed, which was causing a crash.

diff --git a/Engine/Source/Editor/Experimental/EditorInteractiveToolsFramework/Private/EdModeInteractiveToolsContext.cpp b/Engine/Source/Editor/Experimental/EditorInteractiveToolsFramework/Private/EdModeInteractiveToolsContext.cpp
--- a/Engine/Source/Editor/Experimental/EditorInteractiveToolsFramework/Private/EdModeInteractiveToolsContext.cpp
+++ b/Engine/Source/Editor/Experimental/EditorInteractiveToolsFramework/Private/EdModeInteractiveToolsContext.cpp
@@ -1141,6 +1141,7 @@ bool UModeManagerInteractiveToolsContext::OnChildEdModeDeactivated(UEdModeIntera
 	{
 		if (EdModeToolsContexts[k] == ChildToolsContext)
 		{
+			ChildToolsContext->DeactivateAllActiveTools(EToolShutdownType::Cancel);
 			EdModeToolsContexts.RemoveAt(k);
 			return true;
 		}
