Fix Compile warnings in InteractiveToolStack

diff --git a/Engine/Source/Editor/Experimental/EditorInteractiveToolsFramework/Private/InteractiveToolStack.cpp b/Engine/Source/Editor/Experimental/EditorInteractiveToolsFramework/Private/InteractiveToolStack.cpp
--- a/Engine/Source/Editor/Experimental/EditorInteractiveToolsFramework/Private/InteractiveToolStack.cpp
+++ b/Engine/Source/Editor/Experimental/EditorInteractiveToolsFramework/Private/InteractiveToolStack.cpp
@@ -178,50 +178,53 @@ void FInteractiveToolStack::AddToolToStack(UEdMode* EdMode, UEditorInteractiveTo
 		}
 	}
 
-	if (FInteractiveToolStack* ToolStack = ToolStackContext->AccessToolStack(ToolStackIdentifier))
+	if (ensure(UseToolsContext))
 	{
-		ToolStack->AddTool(ToolIdentifier, UICommand);
-		UseToolsContext->ToolManager->RegisterToolType(ToolIdentifier, Builder);
-
-		auto ExecuteLastToolForAction = [EdMode, ToolStackIdentifier, UseToolsContext, UICommand, ToolStackContext]()
+		if (FInteractiveToolStack* ToolStack = ToolStackContext->AccessToolStack(ToolStackIdentifier))
 		{
-			if (FInteractiveToolStack* ToolStack = ToolStackContext->AccessToolStack(ToolStackIdentifier))
-			{
-				const FInputChord& InputChord = UICommand->GetDefaultChord(EMultipleKeyBindingIndex::Primary);
-				const FString& Tool = ToolStack->GetLastActiveToolByChord(InputChord);
+			ToolStack->AddTool(ToolIdentifier, UICommand);
+			UseToolsContext->ToolManager->RegisterToolType(ToolIdentifier, Builder);
 
-				if (!UseToolsContext->IsToolActive(EToolSide::Mouse, Tool))
+			auto ExecuteLastToolForAction = [EdMode, ToolStackIdentifier, UseToolsContext, UICommand, ToolStackContext]()
+			{
+				if (FInteractiveToolStack* ToolStack = ToolStackContext->AccessToolStack(ToolStackIdentifier))
 				{
-					ToolStack->NotifyToolActivated(Tool);
-					UseToolsContext->StartTool(Tool);
+					const FInputChord& InputChord = UICommand->GetDefaultChord(EMultipleKeyBindingIndex::Primary);
+					const FString& Tool = ToolStack->GetLastActiveToolByChord(InputChord);
+
+					if (!UseToolsContext->IsToolActive(EToolSide::Mouse, Tool))
+					{
+						ToolStack->NotifyToolActivated(Tool);
+						UseToolsContext->StartTool(Tool);
+					}
 				}
-			}
-		};
+			};
 
-		auto CanExecuteLastToolForAction = [EdMode, ToolStackIdentifier, UseToolsContext, UICommand, ToolStackContext]()
-		{
-			if (FInteractiveToolStack* ToolStack = ToolStackContext->AccessToolStack(ToolStackIdentifier))
+			auto CanExecuteLastToolForAction = [EdMode, ToolStackIdentifier, UseToolsContext, UICommand, ToolStackContext]()
 			{
-				// @TODO: This leaves us in a possible "deadlock" where the last active tool for a chord can't be used,
-				// but other tools for a chord are valid. Seemingly unlikely scenario so address later if needed.
+				if (FInteractiveToolStack* ToolStack = ToolStackContext->AccessToolStack(ToolStackIdentifier))
+				{
+					// @TODO: This leaves us in a possible "deadlock" where the last active tool for a chord can't be used,
+					// but other tools for a chord are valid. Seemingly unlikely scenario so address later if needed.
 
-				const FInputChord& InputChord = UICommand->GetDefaultChord(EMultipleKeyBindingIndex::Primary);
-				const FString& Tool = ToolStack->GetLastActiveToolByChord(InputChord);
-				return EdMode->ShouldToolStartBeAllowed(Tool) && UseToolsContext->ToolManager->CanActivateTool(EToolSide::Mouse, Tool);
-			}
+					const FInputChord& InputChord = UICommand->GetDefaultChord(EMultipleKeyBindingIndex::Primary);
+					const FString& Tool = ToolStack->GetLastActiveToolByChord(InputChord);
+					return EdMode->ShouldToolStartBeAllowed(Tool) && UseToolsContext->ToolManager->CanActivateTool(EToolSide::Mouse, Tool);
+				}
 
-			return false;
-		};
+				return false;
+			};
 
-		TSharedPtr<FModeToolkit> Toolkit = EdMode->GetToolkit().Pin();
-		if (ensure(Toolkit))
-		{
-			const TSharedRef<FUICommandList>& CommandList = Toolkit->GetToolkitCommands();
-			CommandList->MapAction(UICommand,
-				FExecuteAction::CreateWeakLambda(UseToolsContext, ExecuteLastToolForAction),
-				FCanExecuteAction::CreateWeakLambda(UseToolsContext, CanExecuteLastToolForAction),
-				FIsActionChecked(),
-				EUIActionRepeatMode::RepeatDisabled);
+			TSharedPtr<FModeToolkit> Toolkit = EdMode->GetToolkit().Pin();
+			if (ensure(Toolkit))
+			{
+				const TSharedRef<FUICommandList>& CommandList = Toolkit->GetToolkitCommands();
+				CommandList->MapAction(UICommand,
+					FExecuteAction::CreateWeakLambda(UseToolsContext, ExecuteLastToolForAction),
+					FCanExecuteAction::CreateWeakLambda(UseToolsContext, CanExecuteLastToolForAction),
+					FIsActionChecked(),
+					EUIActionRepeatMode::RepeatDisabled);
+			}
 		}
 	}
 }
