Editor: Add delegates in EdModeInteractiveToolsContext that allow UEdModes to receive Render() and DrawHUD() calls.

diff --git a/Engine/Source/Editor/Experimental/EditorInteractiveToolsFramework/Private/EdModeInteractiveToolsContext.cpp b/Engine/Source/Editor/Experimental/EditorInteractiveToolsFramework/Private/EdModeInteractiveToolsContext.cpp
--- a/Engine/Source/Editor/Experimental/EditorInteractiveToolsFramework/Private/EdModeInteractiveToolsContext.cpp
+++ b/Engine/Source/Editor/Experimental/EditorInteractiveToolsFramework/Private/EdModeInteractiveToolsContext.cpp
@@ -562,6 +562,7 @@ void UEditorInteractiveToolsContext::Render(const FSceneView* View, FViewport* V
 	FEdModeTempRenderContext RenderContext(View, Viewport, ViewportClient, PDI, InteractionState);
 	ToolManager->Render(&RenderContext);
 	GizmoManager->Render(&RenderContext);
+	OnRender.Broadcast(&RenderContext);
 }
 
 void UEditorInteractiveToolsContext::DrawHUD(FViewportClient* ViewportClient,FViewport* Viewport,const FSceneView* View, FCanvas* Canvas)
@@ -581,6 +582,7 @@ void UEditorInteractiveToolsContext::DrawHUD(FViewportClient* ViewportClient,FVi
 	FEdModeTempRenderContext RenderContext(View, Viewport, EditorViewportClient, nullptr /*PDI*/, InteractionState);
 	ToolManager->DrawHUD(Canvas, &RenderContext);
 	GizmoManager->DrawHUD(Canvas, &RenderContext);
+	OnDrawHUD.Broadcast(Canvas, &RenderContext);
 }
 
 
diff --git a/Engine/Source/Editor/Experimental/EditorInteractiveToolsFramework/Public/EdModeInteractiveToolsContext.h b/Engine/Source/Editor/Experimental/EditorInteractiveToolsFramework/Public/EdModeInteractiveToolsContext.h
--- a/Engine/Source/Editor/Experimental/EditorInteractiveToolsFramework/Public/EdModeInteractiveToolsContext.h
+++ b/Engine/Source/Editor/Experimental/EditorInteractiveToolsFramework/Public/EdModeInteractiveToolsContext.h
@@ -77,6 +77,13 @@ public:
 	virtual void Render(const FSceneView* View, FViewport* Viewport, FPrimitiveDrawInterface* PDI);
 	virtual void DrawHUD(FViewportClient* ViewportClient,FViewport* Viewport,const FSceneView* View, FCanvas* Canvas);
 
+	// These delegates can be used to hook in to Render() and DrawHUD() calls. In particular, non-legacy UEdMode's
+	// don't normally receive Render() and DrawHUD() calls from the mode manager, but can attach to these.
+	DECLARE_MULTICAST_DELEGATE_OneParam(FOnRender, IToolsContextRenderAPI* RenderAPI);
+	FOnRender OnRender;
+	DECLARE_MULTICAST_DELEGATE_TwoParams(FOnDrawHUD, FCanvas* Canvas, IToolsContextRenderAPI* RenderAPI);
+	FOnDrawHUD OnDrawHUD;
+
 	/** @return true if selected actors/components can be deleted */
 	virtual bool ProcessEditDelete();
 
