Fix static analysis CIS nullptr guard EdModeInteractiveToolsContext Viewport

diff --git a/Engine/Source/Editor/Experimental/EditorInteractiveToolsFramework/Private/EdModeInteractiveToolsContext.cpp b/Engine/Source/Editor/Experimental/EditorInteractiveToolsFramework/Private/EdModeInteractiveToolsContext.cpp
--- a/Engine/Source/Editor/Experimental/EditorInteractiveToolsFramework/Private/EdModeInteractiveToolsContext.cpp
+++ b/Engine/Source/Editor/Experimental/EditorInteractiveToolsFramework/Private/EdModeInteractiveToolsContext.cpp
@@ -394,12 +394,12 @@ void UEditorInteractiveToolsContext::Tick(FEditorViewportClient* ViewportClient,
 {
 	// invalidate this viewport if it's timestamp is not current
 	const int32* FoundTimestamp = InvalidationMap.Find(ViewportClient);
-	if (FoundTimestamp == nullptr && ViewportClient)
+	if (ViewportClient && FoundTimestamp == nullptr)
 	{
 		ViewportClient->Invalidate(false, false);
 		InvalidationMap.Add(ViewportClient, InvalidationTimestamp);
 	}
-	if (FoundTimestamp != nullptr && *FoundTimestamp < InvalidationTimestamp)
+	if (ViewportClient && FoundTimestamp != nullptr && *FoundTimestamp < InvalidationTimestamp)
 	{
 		ViewportClient->Invalidate(false, false);
 		InvalidationMap[ViewportClient] = InvalidationTimestamp;
