Non-unity fix.

diff --git a/Engine/Source/Editor/Experimental/EditorInteractiveToolsFramework/Private/EditorGizmos/GizmoBoxObject.cpp b/Engine/Source/Editor/Experimental/EditorInteractiveToolsFramework/Private/EditorGizmos/GizmoBoxObject.cpp
--- a/Engine/Source/Editor/Experimental/EditorInteractiveToolsFramework/Private/EditorGizmos/GizmoBoxObject.cpp
+++ b/Engine/Source/Editor/Experimental/EditorInteractiveToolsFramework/Private/EditorGizmos/GizmoBoxObject.cpp
@@ -5,6 +5,7 @@
 #include "BaseGizmos/GizmoRenderingUtil.h"
 #include "Materials/MaterialInterface.h"
 #include "InputState.h"
+#include "SceneManagement.h"
 
 void UGizmoBoxObject::Render(IToolsContextRenderAPI* RenderAPI)
 {
diff --git a/Engine/Source/Editor/Experimental/EditorInteractiveToolsFramework/Private/EditorGizmos/GizmoConeObject.cpp b/Engine/Source/Editor/Experimental/EditorInteractiveToolsFramework/Private/EditorGizmos/GizmoConeObject.cpp
--- a/Engine/Source/Editor/Experimental/EditorInteractiveToolsFramework/Private/EditorGizmos/GizmoConeObject.cpp
+++ b/Engine/Source/Editor/Experimental/EditorInteractiveToolsFramework/Private/EditorGizmos/GizmoConeObject.cpp
@@ -5,6 +5,7 @@
 #include "BaseGizmos/GizmoMath.h"
 #include "InputState.h"
 #include "Materials/MaterialInterface.h"
+#include "SceneManagement.h"
 
 void UGizmoConeObject::Render(IToolsContextRenderAPI* RenderAPI)
 {
