Non-unity fix

diff --git a/Engine/Source/Editor/Experimental/EditorInteractiveToolsFramework/Private/EditorGizmos/GizmoCylinderObject.cpp b/Engine/Source/Editor/Experimental/EditorInteractiveToolsFramework/Private/EditorGizmos/GizmoCylinderObject.cpp
--- a/Engine/Source/Editor/Experimental/EditorInteractiveToolsFramework/Private/EditorGizmos/GizmoCylinderObject.cpp
+++ b/Engine/Source/Editor/Experimental/EditorInteractiveToolsFramework/Private/EditorGizmos/GizmoCylinderObject.cpp
@@ -5,6 +5,7 @@
 #include "BaseGizmos/GizmoMath.h"
 #include "InputState.h"
 #include "Materials/MaterialInterface.h"
+#include "SceneManagement.h"
 
 void UGizmoCylinderObject::Render(IToolsContextRenderAPI* RenderAPI)
 {
