Deprecate the typed element selection set on the tool builder state. This can be replaced with the context object that exposes the typed element set within an asset editor.

# this fixes CIS errors and sets us up to remove this next release

diff --git a/Engine/Source/Editor/Experimental/EditorInteractiveToolsFramework/Private/EdModeInteractiveToolsContext.cpp b/Engine/Source/Editor/Experimental/EditorInteractiveToolsFramework/Private/EdModeInteractiveToolsContext.cpp
--- a/Engine/Source/Editor/Experimental/EditorInteractiveToolsFramework/Private/EdModeInteractiveToolsContext.cpp
+++ b/Engine/Source/Editor/Experimental/EditorInteractiveToolsFramework/Private/EdModeInteractiveToolsContext.cpp
@@ -121,7 +121,10 @@ public:
 		StateOut.World = EditorModeManager->GetWorld();
 		EditorModeManager->GetSelectedActors()->GetSelectedObjects(StateOut.SelectedActors);
 		EditorModeManager->GetSelectedComponents()->GetSelectedObjects(StateOut.SelectedComponents);
+
+		PRAGMA_DISABLE_DEPRECATION_WARNINGS
 		StateOut.TypedElementSelectionSet = EditorModeManager->GetEditorSelectionSet();
+		PRAGMA_ENABLE_DEPRECATION_WARNINGS
 	}
 
 	virtual void GetCurrentViewState(FViewCameraState& StateOut) const override
diff --git a/Engine/Source/Editor/Experimental/EditorInteractiveToolsFramework/Private/EditorInteractiveGizmoSelectionBuilder.cpp b/Engine/Source/Editor/Experimental/EditorInteractiveToolsFramework/Private/EditorInteractiveGizmoSelectionBuilder.cpp
--- a/Engine/Source/Editor/Experimental/EditorInteractiveToolsFramework/Private/EditorInteractiveGizmoSelectionBuilder.cpp
+++ b/Engine/Source/Editor/Experimental/EditorInteractiveToolsFramework/Private/EditorInteractiveGizmoSelectionBuilder.cpp
@@ -2,42 +2,47 @@
 
 #include "EditorInteractiveGizmoSelectionBuilder.h"
 #include "BaseGizmos/TransformProxy.h"
+#include "ContextObjectStore.h"
+#include "Elements/Framework/TypedElementSelectionSet.h"
 #include "Elements/Interfaces/TypedElementObjectInterface.h"
 #include "Elements/Interfaces/TypedElementWorldInterface.h"
-#include "Engine/Selection.h"
 #include "Engine/World.h"
+#include "Tools/AssetEditorContextInterface.h"
 
 UTransformProxy* FEditorGizmoSelectionBuilderHelper::CreateTransformProxyForSelection(const FToolBuilderState& SceneState)
 {
 	// @todo - once UTransformProxy supports typed elements, update this to use the normalized typed
 	// element selection set.
-	if (UTypedElementSelectionSet* SelectionSet = SceneState.TypedElementSelectionSet.Get())
+	if (IAssetEditorContextInterface* AssetEditorContext = SceneState.ToolManager->GetContextObjectStore()->FindContext<IAssetEditorContextInterface>())
 	{
-		if (SelectionSet->GetNumSelectedElements() > 0)
+		if (const UTypedElementSelectionSet* SelectionSet = AssetEditorContext->GetSelectionSet())
 		{
-			bool bHasSelectedElements = false;
-			UTransformProxy* TransformProxy = NewObject<UTransformProxy>();
+			if (SelectionSet->GetNumSelectedElements() > 0)
+			{
+				bool bHasSelectedElements = false;
+				UTransformProxy* TransformProxy = NewObject<UTransformProxy>();
 
-			SelectionSet->ForEachSelectedElement<ITypedElementWorldInterface>([TransformProxy, SelectionSet, &bHasSelectedElements](const TTypedElement<ITypedElementWorldInterface>& InWorldElement)
-				{
-					if (InWorldElement.CanMoveElement(ETypedElementWorldType::Editor))
+				SelectionSet->ForEachSelectedElement<ITypedElementWorldInterface>([TransformProxy, SelectionSet, &bHasSelectedElements](const TTypedElement<ITypedElementWorldInterface>& InWorldElement)
 					{
-						if (TTypedElement<ITypedElementObjectInterface> ObjectTypedElement = SelectionSet->GetElementList()->GetElement<ITypedElementObjectInterface>(InWorldElement))
+						if (InWorldElement.CanMoveElement(ETypedElementWorldType::Editor))
 						{
-							if (AActor* Actor = ObjectTypedElement.GetObjectAs<AActor>())
+							if (TTypedElement<ITypedElementObjectInterface> ObjectTypedElement = SelectionSet->GetElementList()->GetElement<ITypedElementObjectInterface>(InWorldElement))
 							{
-								USceneComponent* SceneComponent = Actor->GetRootComponent();
-								TransformProxy->AddComponent(SceneComponent);
-								bHasSelectedElements = true;
+								if (AActor* Actor = ObjectTypedElement.GetObjectAs<AActor>())
+								{
+									USceneComponent* SceneComponent = Actor->GetRootComponent();
+									TransformProxy->AddComponent(SceneComponent);
+									bHasSelectedElements = true;
+								}
 							}
 						}
-					}
-					return true;
-				});
+						return true;
+					});
 
-			if (bHasSelectedElements)
-			{
-				return TransformProxy;
+				if (bHasSelectedElements)
+				{
+					return TransformProxy;
+				}
 			}
 		}
 	}
diff --git a/Engine/Source/Runtime/InteractiveToolsFramework/Public/ToolContextInterfaces.h b/Engine/Source/Runtime/InteractiveToolsFramework/Public/ToolContextInterfaces.h
--- a/Engine/Source/Runtime/InteractiveToolsFramework/Public/ToolContextInterfaces.h
+++ b/Engine/Source/Runtime/InteractiveToolsFramework/Public/ToolContextInterfaces.h
@@ -8,6 +8,10 @@
 #include "ComponentSourceInterfaces.h"
 #include "UObject/ObjectMacros.h"
 
+// For deprecated FToolBuilderState::TypedElementSelectionSet member
+#include "Elements/Framework/TypedElementSelectionSet.h"
+#include "UObject/WeakObjectPtr.h"
+
 
 // predeclarations so we don't have to include these in all tools
 class UWorld;
@@ -22,7 +26,6 @@ class UInteractiveGizmoManager;
 class UToolTargetManager;
 struct FMeshDescription;
 struct FTypedElementHandle;
-class UTypedElementSelectionSet;
 
 /**
  * FToolBuilderState is a bucket of state information that a ToolBuilder might need
@@ -45,6 +48,7 @@ struct INTERACTIVETOOLSFRAMEWORK_API FToolBuilderState
 	/** Current selected Components. May be empty or nullptr. */
 	TArray<UActorComponent*> SelectedComponents;
 
+	UE_DEPRECATED(5.1, "This has moved to a context object. See IAssetEditorContextInterface")
 	TWeakObjectPtr<UTypedElementSelectionSet> TypedElementSelectionSet;
 };
 
