Update gizmo element library with support for common view-dependent functionality, and rectangle, circle, torus objects. Implemented new TransformGizmo UX design using gizmo primitive library objects.

diff --git a/Engine/Source/Editor/Experimental/EditorInteractiveToolsFramework/Private/EditorGizmos/GizmoArrowObject.cpp b/Engine/Source/Editor/Experimental/EditorInteractiveToolsFramework/Private/EditorGizmos/GizmoArrowObject.cpp
deleted file mode 100644
diff --git a/Engine/Source/Editor/Experimental/EditorInteractiveToolsFramework/Private/EditorGizmos/GizmoBoxObject.cpp b/Engine/Source/Editor/Experimental/EditorInteractiveToolsFramework/Private/EditorGizmos/GizmoBoxObject.cpp
deleted file mode 100644
diff --git a/Engine/Source/Editor/Experimental/EditorInteractiveToolsFramework/Private/EditorGizmos/GizmoConeObject.cpp b/Engine/Source/Editor/Experimental/EditorInteractiveToolsFramework/Private/EditorGizmos/GizmoConeObject.cpp
deleted file mode 100644
diff --git a/Engine/Source/Editor/Experimental/EditorInteractiveToolsFramework/Private/EditorGizmos/GizmoCylinderObject.cpp b/Engine/Source/Editor/Experimental/EditorInteractiveToolsFramework/Private/EditorGizmos/GizmoCylinderObject.cpp
deleted file mode 100644
diff --git a/Engine/Source/Editor/Experimental/EditorInteractiveToolsFramework/Private/EditorGizmos/GizmoGroupObject.cpp b/Engine/Source/Editor/Experimental/EditorInteractiveToolsFramework/Private/EditorGizmos/GizmoGroupObject.cpp
deleted file mode 100644
diff --git a/Engine/Source/Editor/Experimental/EditorInteractiveToolsFramework/Private/EditorGizmos/GizmoObjectHitTargets.cpp b/Engine/Source/Editor/Experimental/EditorInteractiveToolsFramework/Private/EditorGizmos/GizmoObjectHitTargets.cpp
deleted file mode 100644
diff --git a/Engine/Source/Editor/Experimental/EditorInteractiveToolsFramework/Private/EditorGizmos/GizmoObjectTransformSources.cpp b/Engine/Source/Editor/Experimental/EditorInteractiveToolsFramework/Private/EditorGizmos/GizmoObjectTransformSources.cpp
deleted file mode 100644
diff --git a/Engine/Source/Editor/Experimental/EditorInteractiveToolsFramework/Private/EditorGizmos/TransformGizmo.cpp b/Engine/Source/Editor/Experimental/EditorInteractiveToolsFramework/Private/EditorGizmos/TransformGizmo.cpp
--- a/Engine/Source/Editor/Experimental/EditorInteractiveToolsFramework/Private/EditorGizmos/TransformGizmo.cpp
+++ b/Engine/Source/Editor/Experimental/EditorInteractiveToolsFramework/Private/EditorGizmos/TransformGizmo.cpp
@@ -1,17 +1,15 @@
 // Copyright Epic Games, Inc. All Rights Reserved.
 
 #include "EditorGizmos/TransformGizmo.h"
-#include "BaseGizmos/AxisPositionGizmo.h"
+#include "BaseGizmos/AxisSources.h"
+#include "BaseGizmos/GizmoElementGroup.h"
+#include "BaseGizmos/GizmoElementShapes.h"
+#include "BaseGizmos/GizmoRenderingUtil.h"
 #include "BaseGizmos/TransformSources.h"
 #include "EditorGizmos/EditorAxisSources.h"
 #include "EditorGizmos/EditorTransformGizmoSource.h"
 #include "EditorGizmos/EditorTransformProxy.h"
 #include "EditorGizmos/EditorParameterToTransformAdapters.h"
-#include "EditorGizmos/GizmoArrowObject.h"
-#include "EditorGizmos/GizmoBoxObject.h"
-#include "EditorGizmos/GizmoConeObject.h"
-#include "EditorGizmos/GizmoCylinderObject.h"
-#include "EditorGizmos/GizmoGroupObject.h"
 #include "Elements/Interfaces/TypedElementObjectInterface.h"
 #include "Elements/Interfaces/TypedElementWorldInterface.h"
 #include "Engine/CollisionProfile.h"
@@ -23,63 +21,18 @@
 
 #define LOCTEXT_NAMESPACE "UTransformGizmo"
 
-constexpr float UTransformGizmo::AXIS_LENGTH;
-constexpr float UTransformGizmo::AXIS_RADIUS;
-constexpr float UTransformGizmo::AXIS_CONE_ANGLE;
-constexpr float UTransformGizmo::AXIS_CONE_HEIGHT;
-constexpr float UTransformGizmo::AXIS_CONE_HEAD_OFFSET;
-constexpr float UTransformGizmo::AXIS_CUBE_SIZE;
-constexpr float UTransformGizmo::AXIS_CUBE_HEAD_OFFSET;
-constexpr float UTransformGizmo::TRANSLATE_ROTATE_AXIS_CIRCLE_RADIUS;
-constexpr float UTransformGizmo::TWOD_AXIS_CIRCLE_RADIUS;
-constexpr float UTransformGizmo::INNER_AXIS_CIRCLE_RADIUS;
-constexpr float UTransformGizmo::OUTER_AXIS_CIRCLE_RADIUS;
-constexpr float UTransformGizmo::ROTATION_TEXT_RADIUS;
-constexpr int32 UTransformGizmo::AXIS_CIRCLE_SIDES;
-constexpr float UTransformGizmo::ARCALL_RELATIVE_INNER_SIZE;
-constexpr float UTransformGizmo::AXIS_LENGTH_SCALE;
-constexpr float UTransformGizmo::AXIS_LENGTH_SCALE_OFFSET;
-
-constexpr FLinearColor UTransformGizmo::AxisColorX;
-constexpr FLinearColor UTransformGizmo::AxisColorY;
-constexpr FLinearColor UTransformGizmo::AxisColorZ;
-constexpr FLinearColor UTransformGizmo::ScreenAxisColor;
-constexpr FColor UTransformGizmo::PlaneColorXY;
-constexpr FColor UTransformGizmo::ArcBallColor;
-constexpr FColor UTransformGizmo::ScreenSpaceColor;
-constexpr FColor UTransformGizmo::CurrentColor;
-
 void UTransformGizmo::SetDisallowNegativeScaling(bool bDisallow)
 {
-	if (bDisallowNegativeScaling != bDisallow)
-	{
-		bDisallowNegativeScaling = bDisallow;
-		for (UInteractiveGizmo* SubGizmo : this->ActiveGizmos)
-		{
-			if (UAxisPositionGizmo* CastGizmo = Cast<UAxisPositionGizmo>(SubGizmo))
-			{
-				if (UGizmoAxisScaleParameterSource* ParamSource = Cast<UGizmoAxisScaleParameterSource>(CastGizmo->ParameterSource.GetObject()))
-				{
-					ParamSource->bClampToZero = bDisallow;
-				}
-			}
-			/* @todo
-			if (UPlanePositionGizmo* CastGizmo = Cast<UEditorPlanePositionGizmo>(SubGizmo))
-			{
-				if (UGizmoPlaneScaleParameterSource* ParamSource = Cast<UGizmoPlaneScaleParameterSource>(CastGizmo->ParameterSource.GetObject()))
-				{
-					ParamSource->bClampToZero = bDisallow;
-				}
-			}
-			*/
-		}
-	}
+	bDisallowNegativeScaling = bDisallow;
 }
 
 void UTransformGizmo::Setup()
 {
 	UInteractiveGizmo::Setup();
 
+	GizmoElementRoot = NewObject<UGizmoElementGroup>();
+	GizmoElementRoot->SetConstantScale(true);
+
 	UMaterial* AxisMaterialBase = GEngine->ArrowMaterial;
 
 	AxisMaterialX = UMaterialInstanceDynamic::Create(AxisMaterialBase, NULL);
@@ -91,13 +44,19 @@ void UTransformGizmo::Setup()
 	AxisMaterialZ = UMaterialInstanceDynamic::Create(AxisMaterialBase, NULL);
 	AxisMaterialZ->SetVectorParameterValue("GizmoColor", AxisColorZ);
 
+	GreyMaterial = UMaterialInstanceDynamic::Create(AxisMaterialBase, NULL);
+	GreyMaterial->SetVectorParameterValue("GizmoColor", GreyColor);
+
+	WhiteMaterial = UMaterialInstanceDynamic::Create(AxisMaterialBase, NULL);
+	WhiteMaterial->SetVectorParameterValue("GizmoColor", WhiteColor);
+
 	CurrentAxisMaterial = UMaterialInstanceDynamic::Create(AxisMaterialBase, NULL);
 	CurrentAxisMaterial->SetVectorParameterValue("GizmoColor", CurrentColor);
 
 	OpaquePlaneMaterialXY = UMaterialInstanceDynamic::Create(AxisMaterialBase, NULL);
 	OpaquePlaneMaterialXY->SetVectorParameterValue("GizmoColor", FLinearColor::White);
 
-	TransparentPlaneMaterialXY = (UMaterial*)StaticLoadObject(
+	TransparentVertexColorMaterial = (UMaterial*)StaticLoadObject(
 		UMaterial::StaticClass(), NULL,
 		TEXT("/Engine/EditorMaterials/WidgetVertexColorMaterial.WidgetVertexColorMaterial"), NULL, LOAD_None, NULL);
 
@@ -107,8 +66,11 @@ void UTransformGizmo::Setup()
 		LOAD_None, NULL);
 	if (!GridMaterial)
 	{
-		GridMaterial = TransparentPlaneMaterialXY;
+		GridMaterial = TransparentVertexColorMaterial;
 	}
+
+	GizmoElementRoot->SetHoverMaterial(CurrentAxisMaterial);
+	GizmoElementRoot->SetInteractMaterial(CurrentAxisMaterial);
 }
 
 
@@ -119,9 +81,28 @@ void UTransformGizmo::Shutdown()
 
 void UTransformGizmo::Render(IToolsContextRenderAPI* RenderAPI)
 {
-	if (bVisible && GizmoGroupObject)
+	if (bVisible && GizmoElementRoot && RenderAPI)
 	{
-		GizmoGroupObject->Render(RenderAPI);
+		EToolContextCoordinateSystem Space = EToolContextCoordinateSystem::World;
+		float Scale = 1.0f;
+
+		if (TransformSource)
+		{
+			Space = TransformSource->GetGizmoCoordSystemSpace();
+			Scale = TransformSource->GetGizmoScale();
+		}
+
+		FTransform LocalToWorldTransform = ActiveTarget->GetTransform();
+		if (Space == EToolContextCoordinateSystem::World)
+		{
+			LocalToWorldTransform.SetRotation(FQuat::Identity);
+		}
+		LocalToWorldTransform.SetScale3D(FVector(Scale, Scale, Scale));
+
+		UGizmoElementBase::FRenderTraversalState RenderState;
+		RenderState.LocalToWorldTransform = LocalToWorldTransform;
+
+		GizmoElementRoot->Render(RenderAPI, RenderState);
 	}
 }
 
@@ -134,7 +115,6 @@ void UTransformGizmo::UpdateMode()
 
 		if (NewMode != CurrentMode)
 		{
-			ActiveObjects.Empty();
 			EnableMode(CurrentMode, EAxisList::None);
 			EnableMode(NewMode, NewAxisToDraw);
 
@@ -143,48 +123,17 @@ void UTransformGizmo::UpdateMode()
 		}
 		else if (NewAxisToDraw != CurrentAxisToDraw)
 		{
-			ActiveObjects.Empty();
 			EnableMode(CurrentMode, NewAxisToDraw);
 			CurrentAxisToDraw = NewAxisToDraw;
 		}
 	}
 	else
 	{
-		ActiveObjects.Empty();
 		EnableMode(CurrentMode, EAxisList::None);
 		CurrentMode = EGizmoTransformMode::None;
 	}
 }
 
-void UTransformGizmo::UpdateCoordSystem()
-{
-	// Note: the following will change with upcoming changes to gizmo objects
-	if (ActiveTarget && ActiveObjects.Num() > 0)
-	{
-		EToolContextCoordinateSystem Space = EToolContextCoordinateSystem::World;
-		float Scale = 1.0f;
-
-		if (TransformSource)
-		{
-			Space = TransformSource->GetGizmoCoordSystemSpace();
-			Scale = TransformSource->GetGizmoScale();
-		}
-
-		FTransform LocalToWorldTransform = ActiveTarget->GetTransform();
-		if (Space == EToolContextCoordinateSystem::World)
-		{
-			LocalToWorldTransform.SetRotation(FQuat::Identity);
-		}
-
-		for (UGizmoBaseObject* Object : ActiveObjects)
-		{
-			Object->SetWorldLocalState(Space == EToolContextCoordinateSystem::World);
-			Object->SetLocalToWorldTransform(LocalToWorldTransform);
-			Object->SetGizmoScale(Scale);
-		}
-	}
-}
-
 void UTransformGizmo::EnableMode(EGizmoTransformMode InMode, EAxisList::Type InAxisListToDraw)
 {
 	if (InMode == EGizmoTransformMode::Translate)
@@ -201,229 +150,300 @@ void UTransformGizmo::EnableMode(EGizmoTransformMode InMode, EAxisList::Type InA
 	}
 }
 
-void UTransformGizmo::EnableObject(UGizmoBaseObject* InGizmoObject, EAxisList::Type InGizmoAxis, EAxisList::Type InAxisListToDraw)
+void UTransformGizmo::EnableTranslate(EAxisList::Type InAxisListToDraw)
 {
-	if (static_cast<uint8>(InAxisListToDraw) & static_cast<uint8>(InGizmoAxis))
+	check(GizmoElementRoot);
+
+	const bool bEnableX = static_cast<uint8>(InAxisListToDraw) & static_cast<uint8>(EAxisList::X);
+	const bool bEnableY = static_cast<uint8>(InAxisListToDraw) & static_cast<uint8>(EAxisList::Y);
+	const bool bEnableZ = static_cast<uint8>(InAxisListToDraw) & static_cast<uint8>(EAxisList::Z);
+	const bool bEnableAny = bEnableX || bEnableY || bEnableZ;
+
+	if (bEnableX && TranslateXAxisElement == nullptr)
 	{
-		InGizmoObject->SetVisibility(true);
-		ActiveObjects.Add(InGizmoObject);
+		TranslateXAxisElement = MakeTranslateAxis(FVector(1.0f, 0.0f, 0.0f), FVector(0.0f, 1.0f, 0.0f), AxisMaterialX);
+		GizmoElementRoot->Add(TranslateXAxisElement);
 	}
-	else
+
+	if (bEnableY && TranslateYAxisElement == nullptr)
 	{
-		InGizmoObject->SetVisibility(false);
+		TranslateYAxisElement = MakeTranslateAxis(FVector(0.0f, 1.0f, 0.0f), FVector(0.0f, 0.0f, 1.0f), AxisMaterialY);
+		GizmoElementRoot->Add(TranslateYAxisElement);
 	}
-}
 
-void UTransformGizmo::EnableTranslate(EAxisList::Type InAxisListToDraw)
-{
-	if (ensure(AxisXObject && AxisYObject && AxisZObject))
+	if (bEnableZ && TranslateZAxisElement == nullptr)
 	{
-		EnableObject(AxisXObject, EAxisList::X, InAxisListToDraw);
-		EnableObject(AxisYObject, EAxisList::Y, InAxisListToDraw);
-		EnableObject(AxisZObject, EAxisList::Z, InAxisListToDraw);
+		TranslateZAxisElement = MakeTranslateAxis(FVector(0.0f, 0.0f, 1.0f), FVector(1.0f, 0.0f, 0.0f), AxisMaterialZ);
+		GizmoElementRoot->Add(TranslateZAxisElement);
 	}
-}
 
-void UTransformGizmo::EnableRotate(EAxisList::Type InAxisListToDraw)
-{
-	// @todo
-}
-
-void UTransformGizmo::EnableScale(EAxisList::Type InAxisListToDraw)
-{
-	if (ensure(ScaleAxisXObject && ScaleAxisYObject && ScaleAxisZObject))
+	if (bEnableAny && TranslateScreenSpaceElement == nullptr)
 	{
-		EnableObject(ScaleAxisXObject, EAxisList::X, InAxisListToDraw);
-		EnableObject(ScaleAxisYObject, EAxisList::Y, InAxisListToDraw);
-		EnableObject(ScaleAxisZObject, EAxisList::Z, InAxisListToDraw);
+		TranslateScreenSpaceElement = MakeTranslateScreenSpaceHandle();
+		GizmoElementRoot->Add(TranslateScreenSpaceElement);
 	}
-}
 
-void UTransformGizmo::UpdateCameraAxisSource()
-{
-	FViewCameraState CameraState;
-	GetGizmoManager()->GetContextQueriesAPI()->GetCurrentViewState(CameraState);
-	if (CameraAxisSource != nullptr)
+	if (TranslateXAxisElement)
 	{
-		CameraAxisSource->Origin = ActiveTarget ? ActiveTarget->GetTransform().GetLocation() : FVector::ZeroVector;
-		CameraAxisSource->Direction = -CameraState.Forward();
-		CameraAxisSource->TangentX = CameraState.Right();
-		CameraAxisSource->TangentY = CameraState.Up();
+		TranslateXAxisElement->SetEnabled(bEnableX);
 	}
-}
 
+	if (TranslateYAxisElement)
+	{
+		TranslateYAxisElement->SetEnabled(bEnableY);
+	}
 
-void UTransformGizmo::Tick(float DeltaTime)
-{
-	UpdateMode();
+	if (TranslateZAxisElement)
+	{
+		TranslateZAxisElement->SetEnabled(bEnableZ);
+	}
 
-	UpdateCoordSystem();
+	if (TranslateScreenSpaceElement)
+	{
+		TranslateScreenSpaceElement->SetEnabled(bEnableAny);
+	}
 
-	UpdateCameraAxisSource();
+	EnablePlanarObjects(bEnableX, bEnableY, bEnableZ);
 }
 
-void UTransformGizmo::SetActiveTarget(UTransformProxy* Target, IToolContextTransactionProvider* TransactionProvider)
+void UTransformGizmo::EnablePlanarObjects(bool bEnableX, bool bEnableY, bool bEnableZ)
 {
-	if (ActiveTarget != nullptr)
+	check(GizmoElementRoot);
+
+	const bool bEnableXY = bEnableX && bEnableY;
+	const bool bEnableYZ = bEnableY && bEnableZ;
+	const bool bEnableXZ = bEnableX && bEnableZ;
+
+	const FVector XAxis(1.0f, 0.0f, 0.0f);
+	const FVector YAxis(0.0f, 1.0f, 0.0f);
+	const FVector ZAxis(0.0f, 0.0f, 1.0f);
+
+	if (bEnableXY && PlanarXYElement == nullptr)
 	{
-		ClearActiveTarget();
+		PlanarXYElement = MakePlanarHandle(XAxis, YAxis, ZAxis, TransparentVertexColorMaterial, AxisColorZ);
+		GizmoElementRoot->Add(PlanarXYElement);
 	}
 
-	ActiveTarget = Target;
+	if (bEnableYZ && PlanarYZElement == nullptr)
+	{
+		PlanarYZElement = MakePlanarHandle(YAxis, ZAxis, XAxis, TransparentVertexColorMaterial, AxisColorX);
+		GizmoElementRoot->Add(PlanarYZElement);
+	}
 
-	if (!ActiveTarget)
+	if (bEnableXZ && PlanarXZElement == nullptr)
 	{
-		return;
+		PlanarXZElement = MakePlanarHandle(ZAxis, XAxis, YAxis, TransparentVertexColorMaterial, AxisColorY);
+		GizmoElementRoot->Add(PlanarXZElement);
 	}
 
-	// create group object to which all active objects will be added
-	GizmoGroupObject = NewObject<UGizmoGroupObject>();
+	if (PlanarXYElement)
+	{
+		PlanarXYElement->SetEnabled(bEnableXY);
+	}
 
-	// root component provides local X/Y/Z axis, identified by AxisIndex
-	AxisXSource = UGizmoEditorAxisSource::Construct(0, true, this);
-	AxisYSource = UGizmoEditorAxisSource::Construct(1, true, this);
-	AxisZSource = UGizmoEditorAxisSource::Construct(2, true, this);
+	if (PlanarYZElement)
+	{
+		PlanarYZElement->SetEnabled(bEnableYZ);
+	}
 
-	auto MakeArrowObjectFunc = [](FVector Axis, UMaterialInterface* Material, UMaterialInterface* CurrentMaterial)
+	if (PlanarXZElement)
 	{
-		UGizmoArrowObject* ArrowObject = NewObject<UGizmoArrowObject>();
-		ArrowObject->CylinderObject->Direction = Axis;
-		ArrowObject->ConeObject->Direction = -Axis;
-		ArrowObject->SetMaterial(Material);
-		ArrowObject->SetCurrentMaterial(CurrentMaterial);
-		return ArrowObject;
-	};
+		PlanarXZElement->SetEnabled(bEnableXZ);
+	}
+}
 
-	AxisXObject = MakeArrowObjectFunc(FVector(1.0f, 0.0f, 0.0f), AxisMaterialX, CurrentAxisMaterial);
-	AxisYObject = MakeArrowObjectFunc(FVector(0.0f, 1.0f, 0.0f), AxisMaterialY, CurrentAxisMaterial);
-	AxisZObject = MakeArrowObjectFunc(FVector(0.0f, 0.0f, 1.0f), AxisMaterialZ, CurrentAxisMaterial);
+void UTransformGizmo::EnableRotate(EAxisList::Type InAxisListToDraw)
+{
+	const bool bEnableX = static_cast<uint8>(InAxisListToDraw) & static_cast<uint8>(EAxisList::X);
+	const bool bEnableY = static_cast<uint8>(InAxisListToDraw) & static_cast<uint8>(EAxisList::Y);
+	const bool bEnableZ = static_cast<uint8>(InAxisListToDraw) & static_cast<uint8>(EAxisList::Z);
+	const bool bEnableAll = bEnableX && bEnableY && bEnableZ;
 
-	UGizmoScaledAndUnscaledTransformSources* AxisTransformSource = UGizmoScaledAndUnscaledTransformSources::Construct(
-		UGizmoTransformProxyTransformSource::Construct(ActiveTarget, this),
-		UGizmoObjectWorldTransformSource::Construct(GizmoGroupObject, this));
-	// This state target emits an explicit FChange that moves the GizmoActor root component during undo/redo.
-	// It also opens/closes the Transaction that saves/restores the target object locations.
-	if (TransactionProvider == nullptr)
+	const FVector XAxis(1.0f, 0.0f, 0.0f);
+	const FVector YAxis(0.0f, 1.0f, 0.0f);
+	const FVector ZAxis(0.0f, 0.0f, 1.0f);
+
+	if (bEnableX && RotateXAxisElement == nullptr)
 	{
-		TransactionProvider = GetGizmoManager();
+		RotateXAxisElement = MakeRotateAxis(XAxis, YAxis, ZAxis, AxisMaterialX, CurrentAxisMaterial);
+		GizmoElementRoot->Add(RotateXAxisElement);
 	}
-	StateTarget = UGizmoObjectTransformChangeStateTarget::Construct(GizmoGroupObject,
-		LOCTEXT("UTransformGizmoTransaction", "Transform"), TransactionProvider, this);
-	StateTarget->DependentChangeSources.Add(MakeUnique<FTransformProxyChangeSource>(Target));
 
-	CameraAxisSource = NewObject<UGizmoConstantFrameAxisSource>(this);
+	if (bEnableY && RotateYAxisElement == nullptr)
+	{
+		RotateYAxisElement = MakeRotateAxis(YAxis, ZAxis, XAxis, AxisMaterialY, CurrentAxisMaterial);
+		GizmoElementRoot->Add(RotateYAxisElement);
+	}
 
+	if (bEnableZ && RotateZAxisElement == nullptr)
 	{
-		AddAxisTranslationGizmo(AxisXObject, AxisXSource, AxisTransformSource, StateTarget, EAxisList::X, AxisColorX);
-		GizmoGroupObject->Add(AxisXObject);
+		RotateZAxisElement = MakeRotateAxis(ZAxis, XAxis, YAxis, AxisMaterialZ, CurrentAxisMaterial);
+		GizmoElementRoot->Add(RotateZAxisElement);
 	}
 
+	if (bEnableAll)
 	{
-		AddAxisTranslationGizmo(AxisYObject, AxisYSource, AxisTransformSource, StateTarget, EAxisList::Y, AxisColorY);
-		GizmoGroupObject->Add(AxisYObject);
+		if (RotateScreenSpaceElement == nullptr)
+		{
+			RotateScreenSpaceElement = MakeRotateCircleHandle(RotateScreenSpaceRadius, RotateScreenSpaceCircleColor, false);
+			GizmoElementRoot->Add(RotateScreenSpaceElement);
+		}
+
+		if (RotateOuterCircleElement == nullptr)
+		{
+			RotateOuterCircleElement = MakeRotateCircleHandle(RotateOuterCircleRadius, RotateOuterCircleColor, false);
+			GizmoElementRoot->Add(RotateOuterCircleElement);
+		}
+
+		if (RotateArcballOuterElement == nullptr)
+		{
+			RotateArcballOuterElement = MakeRotateCircleHandle(RotateArcballOuterRadius, RotateArcballCircleColor, false);
+			GizmoElementRoot->Add(RotateArcballOuterElement);
+		}
+
+		if (RotateArcballInnerElement == nullptr)
+		{
+			RotateArcballInnerElement = MakeRotateCircleHandle(RotateArcballInnerRadius, RotateArcballCircleColor, true);
+			GizmoElementRoot->Add(RotateArcballInnerElement);
+		}
 	}
 
+	if (RotateXAxisElement)
 	{
-		AddAxisTranslationGizmo(AxisZObject,  AxisZSource, AxisTransformSource, StateTarget, EAxisList::Z, AxisColorZ);
-		GizmoGroupObject->Add(AxisZObject);
+		RotateXAxisElement->SetEnabled(bEnableX);
 	}
 
-/*
-	// @todo: add plane translation
-	if (GizmoActor->TranslateYZ != nullptr)
+	if (RotateYAxisElement)
 	{
-		AddPlaneTranslationGizmo(GizmoActor->TranslateYZ, GizmoComponent, AxisXSource, TransformSource, StateTarget);
-		ActiveComponents.Add(GizmoActor->TranslateYZ);
+		RotateYAxisElement->SetEnabled(bEnableY);
 	}
-	if (GizmoActor->TranslateXZ != nullptr)
+
+	if (RotateZAxisElement)
 	{
-		AddPlaneTranslationGizmo(GizmoActor->TranslateXZ, GizmoComponent, AxisYSource, TransformSource, StateTarget);
-		ActiveComponents.Add(GizmoActor->TranslateXZ);
+		RotateZAxisElement->SetEnabled(bEnableZ);
 	}
-	if (GizmoActor->TranslateXY != nullptr)
+
+	if (RotateScreenSpaceElement)
 	{
-		AddPlaneTranslationGizmo(GizmoActor->TranslateXY, GizmoComponent, AxisZSource, TransformSource, StateTarget);
-		ActiveComponents.Add(GizmoActor->TranslateXY);
+		RotateScreenSpaceElement->SetEnabled(bEnableAll);
 	}
 
-	// @todo: finish rotation implementation
-	if (GizmoActor->RotateX != nullptr)
+	if (RotateOuterCircleElement)
 	{
-		AddAxisRotationGizmo(GizmoActor->RotateX, GizmoComponent, AxisXSource, SpaceSource, TransformSource, StateTarget, EAxisList::X, AxisColorX);
-		ActiveComponents.Add(GizmoActor->RotateX);
+		RotateOuterCircleElement->SetEnabled(bEnableAll);
 	}
-	if (GizmoActor->RotateY != nullptr)
+
+	if (RotateArcballOuterElement)
 	{
-		AddAxisRotationGizmo(GizmoActor->RotateY, GizmoComponent, AxisYSource, SpaceSource, TransformSource, StateTarget, EAxisList::Y, AxisColorY);
-		ActiveComponents.Add(GizmoActor->RotateY);
+		RotateArcballOuterElement->SetEnabled(bEnableAll);
+	}
+
+	if (RotateArcballInnerElement)
+	{ 
+		RotateArcballInnerElement->SetEnabled(bEnableAll);
 	}
-	if (GizmoActor->RotateZ != nullptr)
+}
+
+void UTransformGizmo::EnableScale(EAxisList::Type InAxisListToDraw)
+{
+	check(GizmoElementRoot);
+
+	bool bEnableX = static_cast<uint8>(InAxisListToDraw) & static_cast<uint8>(EAxisList::X);
+	bool bEnableY = static_cast<uint8>(InAxisListToDraw) & static_cast<uint8>(EAxisList::Y);
+	bool bEnableZ = static_cast<uint8>(InAxisListToDraw) & static_cast<uint8>(EAxisList::Z);
+	
+	if (bEnableX && ScaleXAxisElement == nullptr)
 	{
-		AddAxisRotationGizmo(GizmoActor->RotateZ, GizmoComponent, AxisZSource, SpaceSource, TransformSource, StateTarget, EAxisList::Z, AxisColorZ);
-		ActiveComponents.Add(GizmoActor->RotateZ);
+		ScaleXAxisElement = MakeScaleAxis(FVector(1.0f, 0.0f, 0.0f), FVector(0.0f, 1.0f, 0.0f), AxisMaterialX);
+		GizmoElementRoot->Add(ScaleXAxisElement);
 	}
-*/
 
-	// Create objects for scale gizmo
-	ScaleAxisXObject = MakeArrowObjectFunc(FVector(1.0f, 0.0f, 0.0f), AxisMaterialX, CurrentAxisMaterial); 
-	ScaleAxisYObject = MakeArrowObjectFunc(FVector(0.0f, 1.0f, 0.0f), AxisMaterialY, CurrentAxisMaterial); 
-	ScaleAxisZObject = MakeArrowObjectFunc(FVector(0.0f, 0.0f, 1.0f), AxisMaterialZ, CurrentAxisMaterial); 
+	if (bEnableY && ScaleYAxisElement == nullptr)
+	{
+		ScaleYAxisElement = MakeScaleAxis(FVector(0.0f, 1.0f, 0.0f), FVector(0.0f, 0.0f, 1.0f), AxisMaterialY);
+		GizmoElementRoot->Add(ScaleYAxisElement);
+	}
 
-	// only need these if scaling enabled. Essentially these are just the unit axes, regardless
-	// of what 3D axis is in use, we will tell the ParameterSource-to-3D-Scale mapper to
-	// use the coordinate axes
-	UnitAxisXSource = UGizmoEditorAxisSource::Construct(0, false, this);
-	UnitAxisYSource = UGizmoEditorAxisSource::Construct(1, false, this);
-	UnitAxisZSource = UGizmoEditorAxisSource::Construct(2, false, this);
+	if (bEnableZ && ScaleZAxisElement == nullptr)
+	{
+		ScaleZAxisElement = MakeScaleAxis(FVector(0.0f, 0.0f, 1.0f), FVector(1.0f, 0.0f, 0.0f), AxisMaterialZ);
+		GizmoElementRoot->Add(ScaleZAxisElement);
+	}
 
-/*
-	// @todo: add uniform scale handle
-	if (GizmoActor->UniformScale != nullptr)
+	if ((bEnableX || bEnableY || bEnableZ) && ScaleUniformElement == nullptr)
 	{
-		AddUniformScaleGizmo(GizmoActor->UniformScale, GizmoComponent, CameraAxisSource, CameraAxisSource, TransformSource, StateTarget);
-		ActiveComponents.Add(GizmoActor->UniformScale);
+		ScaleUniformElement = MakeUniformScaleHandle();
+		GizmoElementRoot->Add(ScaleUniformElement);
 	}
-*/
 
+	if (ScaleXAxisElement)
 	{
-		AddAxisScaleGizmo(ScaleAxisXObject, AxisXSource, UnitAxisXSource, AxisTransformSource, StateTarget, EAxisList::X, AxisColorX);
-		GizmoGroupObject->Add(ScaleAxisXObject);
+		ScaleXAxisElement->SetEnabled(bEnableX);
 	}
 
+	if (ScaleYAxisElement)
 	{
-		AddAxisScaleGizmo(ScaleAxisYObject, AxisYSource, UnitAxisYSource, AxisTransformSource, StateTarget, EAxisList::Y, AxisColorY);
-		GizmoGroupObject->Add(ScaleAxisYObject);
+		ScaleYAxisElement->SetEnabled(bEnableY);
 	}
 
+	if (ScaleZAxisElement)
 	{
-		AddAxisScaleGizmo(ScaleAxisZObject, AxisZSource, UnitAxisZSource, AxisTransformSource, StateTarget, EAxisList::Z, AxisColorZ);
-		GizmoGroupObject->Add(ScaleAxisZObject);
+		ScaleZAxisElement->SetEnabled(bEnableZ);
 	}
 
-/*
-	// @todo: add plane scale
-	if (GizmoActor->PlaneScaleYZ != nullptr)
+	if (ScaleUniformElement)
 	{
-		AddPlaneScaleGizmo(GizmoActor->PlaneScaleYZ, GizmoComponent, AxisXSource, UnitAxisXSource, TransformSource, StateTarget);
-		ActiveComponents.Add(GizmoActor->PlaneScaleYZ);
-		NonuniformScaleComponents.Add(GizmoActor->PlaneScaleYZ);
+		ScaleUniformElement->SetEnabled(bEnableX || bEnableY || bEnableZ);
 	}
-	if (GizmoActor->PlaneScaleXZ != nullptr)
+
+	EnablePlanarObjects(bEnableX, bEnableY, bEnableZ);
+}
+
+void UTransformGizmo::UpdateCameraAxisSource()
+{
+	FViewCameraState CameraState;
+	GetGizmoManager()->GetContextQueriesAPI()->GetCurrentViewState(CameraState);
+	if (CameraAxisSource != nullptr)
 	{
-		AddPlaneScaleGizmo(GizmoActor->PlaneScaleXZ, GizmoComponent, AxisYSource, UnitAxisYSource, TransformSource, StateTarget);
-		ActiveComponents.Add(GizmoActor->PlaneScaleXZ);
-		NonuniformScaleComponents.Add(GizmoActor->PlaneScaleXZ);
+		CameraAxisSource->Origin = ActiveTarget ? ActiveTarget->GetTransform().GetLocation() : FVector::ZeroVector;
+		CameraAxisSource->Direction = -CameraState.Forward();
+		CameraAxisSource->TangentX = CameraState.Right();
+		CameraAxisSource->TangentY = CameraState.Up();
 	}
-	if (GizmoActor->PlaneScaleXY != nullptr)
+}
+
+void UTransformGizmo::Tick(float DeltaTime)
+{
+	UpdateMode();
+
+	UpdateCameraAxisSource();
+}
+
+void UTransformGizmo::SetActiveTarget(UTransformProxy* Target, IToolContextTransactionProvider* TransactionProvider)
+{
+	if (ActiveTarget != nullptr)
 	{
-		AddPlaneScaleGizmo(GizmoActor->PlaneScaleXY, GizmoComponent, AxisZSource, UnitAxisZSource, TransformSource, StateTarget);
-		ActiveComponents.Add(GizmoActor->PlaneScaleXY);
-		NonuniformScaleComponents.Add(GizmoActor->PlaneScaleXY);
+		ClearActiveTarget();
 	}
-*/ 
 
-	GizmoGroupObject->SetVisibility(false);
+	ActiveTarget = Target;
+
+	// Set current mode to none, mode will be updated next Tick()
 	CurrentMode = EGizmoTransformMode::None;
+
+	if (!ActiveTarget)
+	{
+		return;
+	}
+
+	// This state target emits an explicit FChange that moves the GizmoActor root component during undo/redo.
+	// It also opens/closes the Transaction that saves/restores the target object locations.
+	if (TransactionProvider == nullptr)
+	{
+		TransactionProvider = GetGizmoManager();
+	}
+
+	CameraAxisSource = NewObject<UGizmoConstantFrameAxisSource>(this);
 }
 
 
@@ -461,308 +481,152 @@ void UTransformGizmo::SetNewChildScale(const FVector& NewChildScale)
 	ActiveTarget->SetTransform(NewTransform);
 }
 
-
 void UTransformGizmo::SetVisibility(bool bVisibleIn)
 {
 	bVisible = bVisibleIn;
 }
 
-UInteractiveGizmo* UTransformGizmo::AddAxisTranslationGizmo(
-	UGizmoArrowObject * InArrowObject,
-	IGizmoAxisSource * InAxisSource,
-	IGizmoTransformSource * InTransformSource,
-	IGizmoStateTarget * InStateTarget,
-	EAxisList::Type InAxisType,
-	const FLinearColor InAxisColor)
+UGizmoElementArrow* UTransformGizmo::MakeTranslateAxis(const FVector& InAxisDir, const FVector& InSideDir, UMaterialInterface* InMaterial)
 {
-	check(InArrowObject);
-	check(InArrowObject->CylinderObject);
-	check(InArrowObject->ConeObject);
-
-	// create axis-position gizmo, axis-position parameter will drive translation
-	UAxisPositionGizmo* TranslateGizmo = Cast<UAxisPositionGizmo>(GetGizmoManager()->CreateGizmo(
-		UInteractiveGizmoManager::DefaultAxisPositionBuilderIdentifier, FString(), this));
-	check(TranslateGizmo);
-
-	InArrowObject->CylinderObject->Length = AXIS_LENGTH;
-	InArrowObject->CylinderObject->Radius = AXIS_RADIUS;
-	InArrowObject->bHasConeHead = true;
-
-	InArrowObject->ConeObject->Angle = FMath::DegreesToRadians(AXIS_CONE_ANGLE);
-	InArrowObject->ConeObject->Height = AXIS_CONE_HEIGHT;
-	InArrowObject->ConeObject->Offset = -(AXIS_LENGTH + AXIS_CONE_HEAD_OFFSET);
-
-	// axis source provides the translation axis
-	TranslateGizmo->AxisSource = Cast<UObject>(InAxisSource);
-
-	// parameter source maps axis-parameter-change to translation of TransformSource's transform
-	UGizmoEditorAxisTranslationParameterSource* ParamSource = UGizmoEditorAxisTranslationParameterSource::Construct(InAxisSource, InTransformSource, this);
-	ParamSource->AxisTranslationParameterSource->PositionConstraintFunction = [this](const FVector& Pos, FVector& Snapped) { return PositionSnapFunction(Pos, Snapped); };
-	TranslateGizmo->ParameterSource = ParamSource;
-
-	UGizmoObjectHitTarget* HitTarget = UGizmoObjectHitTarget::Construct(InArrowObject, this);
-
-	TranslateGizmo->HitTarget = HitTarget;
-	TranslateGizmo->StateTarget = Cast<UObject>(InStateTarget);
-
-	TranslateGizmo->ShouldUseCustomDestinationFunc = [this]() { return ShouldAlignDestination(); };
-	TranslateGizmo->CustomDestinationFunc =
-		[this](const UAxisPositionGizmo::FCustomDestinationParams& Params, FVector& OutputPoint) {
-		return DestinationAlignmentRayCaster(*Params.WorldRay, OutputPoint);
-	};
-
-	ActiveGizmos.Add(TranslateGizmo);
-
-	return TranslateGizmo;
+	UGizmoElementArrow* ArrowElement = NewObject<UGizmoElementArrow>();
+	ArrowElement->SetHeadType(EGizmoElementArrowHeadType::Cone);
+	ArrowElement->SetBase(InAxisDir * AxisLengthOffset);
+	ArrowElement->SetDirection(InAxisDir);
+	ArrowElement->SetSideDirection(InSideDir);
+	ArrowElement->SetBodyLength(TranslateAxisLength);
+	ArrowElement->SetBodyRadius(AxisRadius);
+	ArrowElement->SetHeadLength(TranslateAxisConeHeight);
+	ArrowElement->SetHeadRadius(TranslateAxisConeRadius);
+	ArrowElement->SetNumSides(32);
+	ArrowElement->SetMaterial(InMaterial);
+	ArrowElement->SetViewDependentType(EGizmoElementViewDependentType::Axis);
+	ArrowElement->SetViewDependentAxis(InAxisDir);
+	return ArrowElement;
 }
 
-UInteractiveGizmo* UTransformGizmo::AddPlaneTranslationGizmo(
-	IGizmoAxisSource* InAxisSource,
-	IGizmoTransformSource* InTransformSource,
-	IGizmoStateTarget* InStateTarget)
+UGizmoElementArrow* UTransformGizmo::MakeScaleAxis(const FVector& InAxisDir, const FVector& InSideDir, UMaterialInterface* InMaterial)
 {
-	/* @todo
-	// create axis-position gizmo, axis-position parameter will drive translation
-	UEditorPlanePositionGizmo* TranslateGizmo = Cast<UEditorPlanePositionGizmo>(GetGizmoManager()->CreateGizmo(
-		UEditorInteractiveGizmoManager::DefaultEditorPlanePositionBuilderIdentifier, FString(), this));
-	check(TranslateGizmo);
-
-	// axis source provides the translation axis
-	TranslateGizmo->AxisSource = Cast<UObject>(InAxisSource);
-
-	// parameter source maps axis-parameter-change to translation of TransformSource's transform
-	UGizmoPlaneTranslationParameterSource* ParamSource = UGizmoPlaneTranslationParameterSource::Construct(InAxisSource, InTransformSource, this);
-	ParamSource->PositionConstraintFunction = [this](const FVector& Pos, FVector& Snapped) { return PositionSnapFunction(Pos, Snapped); };
-	TranslateGizmo->ParameterSource = ParamSource;
-
-	// sub-component provides hit target
-	UGizmoComponentHitTarget* HitTarget = UGizmoComponentHitTarget::Construct(InAxisComponent, this);
-	if (this->UpdateHoverFunctionComponent)
-	{
-		HitTarget->UpdateHoverFunction = [InAxisComponent, this](bool bHovering) { this->UpdateHoverFunctionComponent(InAxisComponent, bHovering); };
-	}
-	TranslateGizmo->HitTarget = HitTarget;
-
-	TranslateGizmo->StateTarget = Cast<UObject>(InStateTarget);
-
-	TranslateGizmo->ShouldUseCustomDestinationFunc = [this]() { return ShouldAlignDestination(); };
-	TranslateGizmo->CustomDestinationFunc =
-		[this](const UEditorPlanePositionGizmo::FCustomDestinationParams& Params, FVector& OutputPoint) {
-		return DestinationAlignmentRayCaster(*Params.WorldRay, OutputPoint);
-	};
-
-	ActiveGizmos.Add(TranslateGizmo);
-	return TranslateGizmo;
-	*/
-	return nullptr;
+	UGizmoElementArrow* ArrowElement = NewObject<UGizmoElementArrow>();
+	ArrowElement->SetHeadType(EGizmoElementArrowHeadType::Cube);
+	ArrowElement->SetBase(InAxisDir * AxisLengthOffset);
+	ArrowElement->SetDirection(InAxisDir);
+	ArrowElement->SetSideDirection(InSideDir);
+	ArrowElement->SetBodyLength(TranslateAxisLength);
+	ArrowElement->SetBodyRadius(AxisRadius);
+	ArrowElement->SetHeadLength(ScaleAxisCubeDim);
+	ArrowElement->SetNumSides(32);
+	ArrowElement->SetMaterial(InMaterial);
+	ArrowElement->SetViewDependentType(EGizmoElementViewDependentType::Axis);
+	ArrowElement->SetViewDependentAxis(InAxisDir);
+	return ArrowElement;
 }
 
-UInteractiveGizmo* UTransformGizmo::AddAxisRotationGizmo(
-	IGizmoAxisSource* InAxisSource,
-	IGizmoTransformSource* InTransformSource,
-	IGizmoStateTarget* InStateTarget,
-	EAxisList::Type InAxisType,
-	const FLinearColor InAxisColor)
+UGizmoElementBox* UTransformGizmo::MakeUniformScaleHandle()
 {
-	/* @todo
-	// create axis-angle gizmo, angle will drive axis-rotation
-	UEditorAxisAngleGizmo* RotateGizmo = Cast<UEditorAxisAngleGizmo>(GetGizmoManager()->CreateGizmo(
-		UEditorInteractiveGizmoManager::DefaultEditorAxisAngleBuilderIdentifier, FString(), this));
-	check(RotateGizmo);
-
-	// axis source provides the translation axis
-	RotateGizmo->Initialize(InAxisType, InAxisColor, CurrentColor);
-
-	// axis source provides the rotation axis
-	RotateGizmo->AxisSource = Cast<UObject>(InAxisSource);
-
-	// space source provides coordinate space info specific to the gizmo
-	RotateGizmo->GizmoSpaceSource = Cast<UObject>(SpaceSource);
-
-	// parameter source maps angle-parameter-change to rotation of TransformSource's transform
-	UGizmoAxisRotationParameterSource* AngleSource = UGizmoAxisRotationParameterSource::Construct(InAxisSource, InTransformSource, this);
-	AngleSource->RotationConstraintFunction = [this](const FQuat& DeltaRotation){ return RotationSnapFunction(DeltaRotation); };
-	RotateGizmo->AngleSource = AngleSource;
-
-	// sub-component provides hit target
-	UGizmoComponentHitTarget* HitTarget = UGizmoComponentHitTarget::Construct(InAxisComponent, this);
-	if (this->UpdateHoverFunctionComponent)
-	{
-		HitTarget->UpdateHoverFunction = [InAxisComponent, this](bool bHovering) { this->UpdateHoverFunctionComponent(InAxisComponent, bHovering); };
-	}
-	RotateGizmo->HitTarget = HitTarget;
-
-	RotateGizmo->StateTarget = Cast<UObject>(StateTarget);
-
-	RotateGizmo->ShouldUseCustomDestinationFunc = [this]() { return ShouldAlignDestination(); };
-	RotateGizmo->CustomDestinationFunc =
-		[this](const UEditorAxisAngleGizmo::FCustomDestinationParams& Params, FVector& OutputPoint) {
-		return DestinationAlignmentRayCaster(*Params.WorldRay, OutputPoint);
-	};
-
-	ActiveGizmos.Add(RotateGizmo);
-
-	return RotateGizmo;
-	*/
-	return nullptr;
+	UGizmoElementBox* BoxElement = NewObject<UGizmoElementBox>();
+	BoxElement->SetCenter(FVector::ZeroVector);
+	BoxElement->SetUpDirection(FVector::UpVector);
+	BoxElement->SetSideDirection(FVector::RightVector);
+	BoxElement->SetDimensions(FVector(ScaleAxisCubeDim, ScaleAxisCubeDim, ScaleAxisCubeDim));
+	BoxElement->SetMaterial(GreyMaterial);
+	return BoxElement;
 }
 
-
-UInteractiveGizmo* UTransformGizmo::AddAxisScaleGizmo(
-	UGizmoArrowObject* InArrowObject,
-	IGizmoAxisSource* InGizmoAxisSource, IGizmoAxisSource* InParameterAxisSource,
-	IGizmoTransformSource* InTransformSource,
-	IGizmoStateTarget* InStateTarget,
-	EAxisList::Type InAxisType,
-	const FLinearColor InAxisColor)
+UGizmoElementRectangle* UTransformGizmo::MakePlanarHandle(const FVector& InUpDirection, const FVector& InSideDirection, const FVector& InPlaneNormal,
+	UMaterialInterface* InMaterial, const FLinearColor& InVertexColor)
 {
-	check(InArrowObject);
-	check(InArrowObject->CylinderObject);
-	check(InArrowObject->BoxObject);
-
-	InArrowObject->CylinderObject->Length = AXIS_LENGTH_SCALE;
-	InArrowObject->CylinderObject->Radius = AXIS_RADIUS;
-	InArrowObject->CylinderObject->Offset = AXIS_LENGTH_SCALE_OFFSET;
-	InArrowObject->bHasConeHead = false;
-
-	InArrowObject->BoxObject->Dimensions = FVector(AXIS_CUBE_SIZE, AXIS_CUBE_SIZE, AXIS_CUBE_SIZE);
-	InArrowObject->BoxObject->Offset = AXIS_LENGTH_SCALE + AXIS_LENGTH_SCALE_OFFSET + AXIS_CUBE_HEAD_OFFSET;
-	if (InAxisType == EAxisList::X)
-	{
-		InArrowObject->BoxObject->UpDirection = FVector(1, 0, 0);
-		InArrowObject->BoxObject->SideDirection = FVector(0, 1, 0);
-	}
-	else if (InAxisType == EAxisList::Y)
-	{
-		InArrowObject->BoxObject->UpDirection = FVector(0, 1, 0);
-		InArrowObject->BoxObject->SideDirection = FVector(0, 0, 1);
-	}
-	else
-	{
-		InArrowObject->BoxObject->UpDirection = FVector(0, 0, 1);
-		InArrowObject->BoxObject->SideDirection = FVector(1, 0, 0);
-	}
-
-	// create axis-position gizmo, axis-position parameter will drive scale
-	UAxisPositionGizmo* ScaleGizmo = Cast<UAxisPositionGizmo>(GetGizmoManager()->CreateGizmo(
-		UInteractiveGizmoManager::DefaultAxisPositionBuilderIdentifier, FString(), this));
-	ScaleGizmo->bEnableSignedAxis = true;
-	check(ScaleGizmo);
-
-	// axis source provides the translation axis
-	ScaleGizmo->AxisSource = Cast<UObject>(InGizmoAxisSource);
-
-	// parameter source maps axis-parameter-change to translation of TransformSource's transform
-	UGizmoAxisScaleParameterSource* ParamSource = UGizmoAxisScaleParameterSource::Construct(InParameterAxisSource, InTransformSource, this);
-	ParamSource->bClampToZero = bDisallowNegativeScaling;
-	ScaleGizmo->ParameterSource = ParamSource;
-
-	UGizmoObjectHitTarget* HitTarget = UGizmoObjectHitTarget::Construct(InArrowObject, this);
-
-	ScaleGizmo->HitTarget = HitTarget;
-	ScaleGizmo->StateTarget = Cast<UObject>(StateTarget);
-
-	ActiveGizmos.Add(ScaleGizmo);
-	return ScaleGizmo;
+	FVector PlanarHandleCenter = (InUpDirection + InSideDirection) * PlanarHandleOffset;
+
+	FColor LineColor = InVertexColor.ToFColor(false);
+	FColor VertexColor = LineColor;
+	VertexColor.A = LargeOuterAlpha;
+
+	UGizmoElementRectangle* RectangleElement = NewObject<UGizmoElementRectangle>();
+	RectangleElement->SetUpDirection(InUpDirection);
+	RectangleElement->SetSideDirection(InSideDirection);
+	RectangleElement->SetCenter(PlanarHandleCenter);
+	RectangleElement->SetHeight(PlanarHandleSize);
+	RectangleElement->SetWidth(PlanarHandleSize);
+	RectangleElement->SetMaterial(InMaterial);
+	RectangleElement->SetVertexColor(VertexColor);
+	RectangleElement->SetLineColor(LineColor);
+	RectangleElement->SetDrawLine(true);
+	RectangleElement->SetDrawMesh(true);
+	RectangleElement->SetHitMesh(true);
+	RectangleElement->SetViewDependentType(EGizmoElementViewDependentType::Plane);
+	RectangleElement->SetViewDependentAxis(InPlaneNormal);
+	return RectangleElement;
 }
 
-
-UInteractiveGizmo* UTransformGizmo::AddPlaneScaleGizmo(
-	IGizmoAxisSource* InGizmoAxisSource, IGizmoAxisSource* InParameterAxisSource,
-	IGizmoTransformSource* InTransformSource,
-	IGizmoStateTarget* InStateTarget)
+UGizmoElementRectangle* UTransformGizmo::MakeTranslateScreenSpaceHandle()
 {
-	/* @todo
-	// create axis-position gizmo, axis-position parameter will drive scale
-	UEditorPlanePositionGizmo* ScaleGizmo = Cast<UEditorPlanePositionGizmo>(GetGizmoManager()->CreateGizmo(
-		UEditorInteractiveGizmoManager::DefaultEditorPlanePositionBuilderIdentifier, FString(), this));
-	ScaleGizmo->bEnableSignedAxis = true;
-	check(ScaleGizmo);
-
-	// axis source provides the translation axis
-	ScaleGizmo->AxisSource = Cast<UObject>(InGizmoAxisSource);
-
-	// parameter source maps axis-parameter-change to translation of TransformSource's transform
-	UGizmoPlaneScaleParameterSource* ParamSource = UGizmoPlaneScaleParameterSource::Construct(InParameterAxisSource, InTransformSource, this);
-	ParamSource->bClampToZero = bDisallowNegativeScaling;
-	ParamSource->bUseEqualScaling = true;
-	ScaleGizmo->ParameterSource = ParamSource;
-
-	// sub-component provides hit target
-	UGizmoComponentHitTarget* HitTarget = UGizmoComponentHitTarget::Construct(InAxisComponent, this);
-	if (this->UpdateHoverFunctionComponent)
-	{
-		HitTarget->UpdateHoverFunction = [InAxisComponent, this](bool bHovering) { this->UpdateHoverFunctionComponent(InAxisComponent, bHovering); };
-	}
-	ScaleGizmo->HitTarget = HitTarget;
-
-	ScaleGizmo->StateTarget = Cast<UObject>(StateTarget);
-
-	ActiveGizmos.Add(ScaleGizmo);
-	return ScaleGizmo;
-	*/
-	return nullptr;
+	UGizmoElementRectangle* RectangleElement = NewObject<UGizmoElementRectangle>();
+	RectangleElement->SetUpDirection(FVector::UpVector);
+	RectangleElement->SetSideDirection(FVector::RightVector);
+	RectangleElement->SetCenter(FVector::ZeroVector);
+	RectangleElement->SetHeight(TranslateScreenSpaceHandleSize);
+	RectangleElement->SetWidth(TranslateScreenSpaceHandleSize);
+	RectangleElement->SetScreenSpace(true);
+	RectangleElement->SetMaterial(TransparentVertexColorMaterial);
+	RectangleElement->SetLineColor(ScreenSpaceColor);
+	RectangleElement->SetHitMesh(true);
+	RectangleElement->SetDrawMesh(false);
+	RectangleElement->SetDrawLine(true);
+	return RectangleElement;
 }
 
-
-UInteractiveGizmo* UTransformGizmo::AddUniformScaleGizmo(
-	IGizmoAxisSource* InGizmoAxisSource, IGizmoAxisSource* InParameterAxisSource,
-	IGizmoTransformSource* InTransformSource,
-	IGizmoStateTarget* InStateTarget)
+UGizmoElementTorus* UTransformGizmo::MakeRotateAxis(const FVector& Normal, const FVector& TorusAxis0, const FVector& TorusAxis1,
+	UMaterialInterface* InMaterial, UMaterialInterface* InCurrentMaterial)
 {
-	/* @todo
-	// create plane-position gizmo, plane-position parameter will drive scale
-	UEditorPlanePositionGizmo* ScaleGizmo = Cast<UEditorPlanePositionGizmo>(GetGizmoManager()->CreateGizmo(
-		UEditorInteractiveGizmoManager::DefaultEditorPlanePositionBuilderIdentifier, FString(), this));
-	check(ScaleGizmo);
-
-	// axis source provides the translation plane
-	ScaleGizmo->AxisSource = Cast<UObject>(InGizmoAxisSource);
+	UGizmoElementTorus* RotateAxisElement = NewObject<UGizmoElementTorus>();
+	RotateAxisElement->SetCenter(FVector::ZeroVector);
+	RotateAxisElement->SetOuterRadius(UTransformGizmo::RotateAxisOuterRadius);
+	RotateAxisElement->SetOuterSegments(UTransformGizmo::RotateAxisOuterSegments);
+	RotateAxisElement->SetInnerRadius(UTransformGizmo::RotateAxisInnerRadius);
+	RotateAxisElement->SetInnerSlices(UTransformGizmo::RotateAxisInnerSlices);
+	RotateAxisElement->SetNormal(Normal);
+	RotateAxisElement->SetBeginAxis(TorusAxis0);
+	RotateAxisElement->SetPartial(true);
+	RotateAxisElement->SetAngle(PI);
+	RotateAxisElement->SetViewDependentType(EGizmoElementViewDependentType::Plane);
+	RotateAxisElement->SetViewDependentAxis(Normal);
+	RotateAxisElement->SetViewAlignType(EGizmoElementViewAlignType::Axial);
+	RotateAxisElement->SetViewAlignAxis(Normal);
+	RotateAxisElement->SetViewAlignNormal(TorusAxis1);
+	RotateAxisElement->SetMaterial(InMaterial);
+	return RotateAxisElement;
+}
 
-	// parameter source maps axis-parameter-change to translation of TransformSource's transform
-	UGizmoUniformScaleParameterSource* ParamSource = UGizmoUniformScaleParameterSource::Construct(InParameterAxisSource, InTransformSource, this);
-	//ParamSource->PositionConstraintFunction = [this](const FVector& Pos, FVector& Snapped) { return PositionSnapFunction(Pos, Snapped); };
-	ScaleGizmo->ParameterSource = ParamSource;
+UGizmoElementCircle* UTransformGizmo::MakeRotateCircleHandle(float InRadius, const FLinearColor& InColor, float bFill)
+{
+	UGizmoElementCircle* CircleElement = NewObject<UGizmoElementCircle>();
+	CircleElement->SetCenter(FVector::ZeroVector);
+	CircleElement->SetRadius(InRadius);
+	CircleElement->SetNormal(-FVector::ForwardVector);
+	CircleElement->SetLineColor(InColor);
+	CircleElement->SetViewAlignType(EGizmoElementViewAlignType::PointOnly);
+	CircleElement->SetViewAlignNormal(-FVector::ForwardVector);
 
-	// sub-component provides hit target
-	UGizmoComponentHitTarget* HitTarget = UGizmoComponentHitTarget::Construct(InScaleComponent, this);
-	if (this->UpdateHoverFunctionComponent)
+	if (bFill)
 	{
-		HitTarget->UpdateHoverFunction = [InScaleComponent, this](bool bHovering) { this->UpdateHoverFunctionComponent(InScaleComponent, bHovering); };
+		CircleElement->SetVertexColor(InColor.ToFColor(true));
+		CircleElement->SetMaterial(WhiteMaterial);
+	}
+	else
+	{
+		CircleElement->SetDrawLine(true);
+		CircleElement->SetHitLine(true);
+		CircleElement->SetDrawMesh(false);
+		CircleElement->SetHitMesh(false);
 	}
-	ScaleGizmo->HitTarget = HitTarget;
-
-	ScaleGizmo->StateTarget = Cast<UObject>(InStateTarget);
 
-	ActiveGizmos.Add(ScaleGizmo);
-	return ScaleGizmo;
-	*/
-	return nullptr;
+	return CircleElement;
 }
 
 
 void UTransformGizmo::ClearActiveTarget()
 {
-	for (UInteractiveGizmo* Gizmo : ActiveGizmos)
-	{
-		GetGizmoManager()->DestroyGizmo(Gizmo);
-	}
-	ActiveGizmos.SetNum(0);
-	ActiveObjects.SetNum(0);
-
-	CameraAxisSource = nullptr;
-	GizmoGroupObject = nullptr;
-	AxisXSource = nullptr;
-	AxisYSource = nullptr;
-	AxisZSource = nullptr;
-	AxisXObject = nullptr;
-	AxisYObject = nullptr;
-	AxisZObject = nullptr;
-	UnitAxisXSource = nullptr;
-	UnitAxisYSource = nullptr;
-	UnitAxisZSource = nullptr;
 	StateTarget = nullptr;
-
 	ActiveTarget = nullptr;
 }
 
diff --git a/Engine/Source/Editor/Experimental/EditorInteractiveToolsFramework/Public/EditorGizmos/GizmoArrowObject.h b/Engine/Source/Editor/Experimental/EditorInteractiveToolsFramework/Public/EditorGizmos/GizmoArrowObject.h
deleted file mode 100644
diff --git a/Engine/Source/Editor/Experimental/EditorInteractiveToolsFramework/Public/EditorGizmos/GizmoBaseObject.h b/Engine/Source/Editor/Experimental/EditorInteractiveToolsFramework/Public/EditorGizmos/GizmoBaseObject.h
deleted file mode 100644
diff --git a/Engine/Source/Editor/Experimental/EditorInteractiveToolsFramework/Public/EditorGizmos/GizmoBoxObject.h b/Engine/Source/Editor/Experimental/EditorInteractiveToolsFramework/Public/EditorGizmos/GizmoBoxObject.h
deleted file mode 100644
diff --git a/Engine/Source/Editor/Experimental/EditorInteractiveToolsFramework/Public/EditorGizmos/GizmoConeObject.h b/Engine/Source/Editor/Experimental/EditorInteractiveToolsFramework/Public/EditorGizmos/GizmoConeObject.h
deleted file mode 100644
diff --git a/Engine/Source/Editor/Experimental/EditorInteractiveToolsFramework/Public/EditorGizmos/GizmoCylinderObject.h b/Engine/Source/Editor/Experimental/EditorInteractiveToolsFramework/Public/EditorGizmos/GizmoCylinderObject.h
deleted file mode 100644
diff --git a/Engine/Source/Editor/Experimental/EditorInteractiveToolsFramework/Public/EditorGizmos/GizmoGroupObject.h b/Engine/Source/Editor/Experimental/EditorInteractiveToolsFramework/Public/EditorGizmos/GizmoGroupObject.h
deleted file mode 100644
diff --git a/Engine/Source/Editor/Experimental/EditorInteractiveToolsFramework/Public/EditorGizmos/GizmoObjectHitTargets.h b/Engine/Source/Editor/Experimental/EditorInteractiveToolsFramework/Public/EditorGizmos/GizmoObjectHitTargets.h
deleted file mode 100644
diff --git a/Engine/Source/Editor/Experimental/EditorInteractiveToolsFramework/Public/EditorGizmos/GizmoObjectStateTargets.h b/Engine/Source/Editor/Experimental/EditorInteractiveToolsFramework/Public/EditorGizmos/GizmoObjectStateTargets.h
deleted file mode 100644
diff --git a/Engine/Source/Editor/Experimental/EditorInteractiveToolsFramework/Public/EditorGizmos/GizmoObjectTransformSources.h b/Engine/Source/Editor/Experimental/EditorInteractiveToolsFramework/Public/EditorGizmos/GizmoObjectTransformSources.h
deleted file mode 100644
diff --git a/Engine/Source/Editor/Experimental/EditorInteractiveToolsFramework/Public/EditorGizmos/TransformGizmo.h b/Engine/Source/Editor/Experimental/EditorInteractiveToolsFramework/Public/EditorGizmos/TransformGizmo.h
--- a/Engine/Source/Editor/Experimental/EditorInteractiveToolsFramework/Public/EditorGizmos/TransformGizmo.h
+++ b/Engine/Source/Editor/Experimental/EditorInteractiveToolsFramework/Public/EditorGizmos/TransformGizmo.h
@@ -3,14 +3,10 @@
 #pragma once
 
 #include "CoreMinimal.h"
+#include "BaseGizmos/GizmoElementHitTargets.h"
+#include "BaseGizmos/GizmoElementStateTargets.h"
 #include "BaseGizmos/TransformProxy.h"
 #include "EditorGizmos/EditorAxisSources.h"
-#include "EditorGizmos/GizmoArrowObject.h"
-#include "EditorGizmos/GizmoBaseObject.h"
-#include "EditorGizmos/GizmoGroupObject.h"
-#include "EditorGizmos/GizmoObjectHitTargets.h"
-#include "EditorGizmos/GizmoObjectStateTargets.h"
-#include "EditorGizmos/GizmoObjectTransformSources.h"
 #include "EditorGizmos/TransformGizmoInterfaces.h"
 #include "InteractiveGizmo.h"
 #include "InteractiveToolObjects.h"
@@ -23,10 +19,19 @@ class UInteractiveGizmoManager;
 class IGizmoAxisSource;
 class IGizmoTransformSource;
 class IGizmoStateTarget;
+class FEditorTransformGizmoTransformChange;
 class UGizmoConstantFrameAxisSource;
 class UGizmoTransformChangeStateTarget;
-class FEditorTransformGizmoTransformChange;
-
+class UGizmoElementArrow;
+class UGizmoElementBase;
+class UGizmoElementBox;
+class UGizmoElementCircle;
+class UGizmoElementCone;
+class UGizmoElementCylinder;
+class UGizmoElementGroup;
+class UGizmoElementRectangle;
+class UGizmoElementRoot;
+class UGizmoElementTorus;
 
 /**
  * UTransformGizmo provides standard Transformation Gizmo interactions,
@@ -39,31 +44,56 @@ class EDITORINTERACTIVETOOLSFRAMEWORK_API UTransformGizmo : public UInteractiveG
 	GENERATED_BODY()
 
 public:
-	constexpr static float AXIS_LENGTH = 35.0f;
-	constexpr static float AXIS_RADIUS = 1.2f;
-	constexpr static float AXIS_CONE_ANGLE = 15.7f;
-	constexpr static float AXIS_CONE_HEIGHT = 13;
-	constexpr static float AXIS_CONE_HEAD_OFFSET = 12;
-	constexpr static float AXIS_CUBE_SIZE = 4;
-	constexpr static float AXIS_CUBE_HEAD_OFFSET= 3;
-	constexpr static float TRANSLATE_ROTATE_AXIS_CIRCLE_RADIUS = 20.0f;
-	constexpr static float TWOD_AXIS_CIRCLE_RADIUS = 10.0f;
-	constexpr static float INNER_AXIS_CIRCLE_RADIUS = 48.0f;
-	constexpr static float OUTER_AXIS_CIRCLE_RADIUS = 56.0f;
-	constexpr static float ROTATION_TEXT_RADIUS = 75.0f;
-	constexpr static int32 AXIS_CIRCLE_SIDES = 24;
-	constexpr static float ARCALL_RELATIVE_INNER_SIZE = 0.75f;
-	constexpr static float AXIS_LENGTH_SCALE = 25.0f;
-	constexpr static float AXIS_LENGTH_SCALE_OFFSET = 5.0f;
-
-	constexpr static FLinearColor AxisColorX = FLinearColor(0.594f, 0.0197f, 0.0f);
-	constexpr static FLinearColor AxisColorY = FLinearColor(0.1349f, 0.3959f, 0.0f);
-	constexpr static FLinearColor AxisColorZ = FLinearColor(0.0251f, 0.207f, 0.85f);
-	constexpr static FLinearColor ScreenAxisColor = FLinearColor(0.76, 0.72, 0.14f);
-	constexpr static FColor PlaneColorXY = FColor(255, 255, 0); // FColor::Yellow
-	constexpr static FColor ArcBallColor = FColor(128, 128, 128, 6);
-	constexpr static FColor ScreenSpaceColor = FColor(196, 196, 196);
-	constexpr static FColor CurrentColor = FColor(255, 255, 0); // FColor::Yellow
+
+	static constexpr float AxisRadius = 1.5f;
+	static constexpr float AxisLengthOffset = 20.0f;
+
+	static constexpr float TranslateAxisLength = 70.0f;
+	static constexpr float TranslateAxisConeAngle = 16.0f;
+	static constexpr float TranslateAxisConeHeight = 22.0f;
+	static constexpr float TranslateAxisConeRadius = 7.0f;
+	static constexpr float TranslateScreenSpaceHandleSize = 14.0f;
+
+	// Rotate constants
+	static constexpr float RotateArcballInnerRadius = 8.0f;
+	static constexpr float RotateArcballOuterRadius = 10.0f;
+	static constexpr float RotateArcballSphereRadius = 70.0f;
+	static constexpr float RotateAxisOuterRadius = 73.0f;
+	static constexpr float RotateAxisInnerRadius = 1.25f;
+	static constexpr int32 RotateAxisOuterSegments = 64;
+	static constexpr int32 RotateAxisInnerSlices = 8;
+	static constexpr float RotateOuterCircleRadius = 73.0f;
+	static constexpr float RotateScreenSpaceRadius = 83.0f;
+
+	static constexpr float ScaleAxisLength = 35.0f;
+	static constexpr float ScaleAxisCubeSize = 3.0f;
+	static constexpr float ScaleAxisCubeDim = 12.0f;
+
+	static constexpr float PlanarHandleOffset = 55.0f;
+	static constexpr float PlanarHandleSize = 15.0f;
+
+	static constexpr float AxisTransp = 0.8f;
+	static constexpr FLinearColor AxisColorX = FLinearColor(0.594f, 0.0197f, 0.0f);
+	static constexpr FLinearColor AxisColorY = FLinearColor(0.1349f, 0.3959f, 0.0f);
+	static constexpr FLinearColor AxisColorZ = FLinearColor(0.0251f, 0.207f, 0.85f);
+	static constexpr FLinearColor ScreenAxisColor = FLinearColor(0.76, 0.72, 0.14f);
+	static constexpr FColor PlaneColorXY = FColor(255, 255, 0); // FColor::Yellow
+	static constexpr FColor ArcBallColor = FColor(128, 128, 128, 6);
+	static constexpr FColor ScreenSpaceColor = FColor(196, 196, 196);
+	static constexpr FColor CurrentColor = FColor(255, 255, 0); // FColor::Yellow
+
+	static constexpr FLinearColor GreyColor = FLinearColor(0.50f, 0.50f, 0.50f);
+	static constexpr FLinearColor WhiteColor = FLinearColor(1.0f, 1.0f, 1.0f);
+
+	static constexpr FLinearColor RotateScreenSpaceCircleColor = WhiteColor;
+	static constexpr FLinearColor RotateOuterCircleColor = GreyColor;
+	static constexpr FLinearColor RotateArcballCircleColor = WhiteColor;
+
+	static constexpr uint8 LargeInnerAlpha = 0x3f;
+	static constexpr uint8 SmallInnerAlpha = 0x0f;
+	static constexpr uint8 LargeOuterAlpha = 0x7f;
+	static constexpr uint8 SmallOuterAlpha = 0x0f;
+
 
 public:
 
@@ -164,80 +194,95 @@ public:
 
 protected:
 
-	/** List of current-active gizmo objects */
+	//
+	// Gizmo Objects, used for rendering and hit testing
+	//
+
+	/** Root of renderable gizmo elements */
 	UPROPERTY()
-	TArray<TObjectPtr<UGizmoBaseObject>> ActiveObjects;
+	TObjectPtr<UGizmoElementGroup> GizmoElementRoot;
 
-	/** list of currently-active child gizmos */
+	/** Translate X Axis */
 	UPROPERTY()
-	TArray<TObjectPtr<UInteractiveGizmo>> ActiveGizmos;
+	TObjectPtr<UGizmoElementArrow> TranslateXAxisElement;
 
-	//
-	// Axis Sources
-	//
-	/** Axis that points towards camera, X/Y plane tangents aligned to right/up. Shared across Gizmos, and created internally during SetActiveTarget() */
+	/** Translate Y Axis */
 	UPROPERTY()
-	TObjectPtr<UGizmoConstantFrameAxisSource> CameraAxisSource;
+	TObjectPtr<UGizmoElementArrow> TranslateYAxisElement;
 
-	// internal function that updates CameraAxisSource by getting current view state from GizmoManager
-	void UpdateCameraAxisSource();
+	/** Translate Z Axis */
+	UPROPERTY()
+	TObjectPtr<UGizmoElementArrow> TranslateZAxisElement;
 
-	/** Gizmo group object contains all gizmos created internally during SetActiveTarget() */
+	/** Translate screen-space */
 	UPROPERTY()
-	TObjectPtr<UGizmoGroupObject> GizmoGroupObject;
+	TObjectPtr<UGizmoElementRectangle> TranslateScreenSpaceElement;
 
-	// @todo - these properties can probably be removed since they are contained within the group now.
-	/** X-axis source is shared across Gizmos, and created internally during SetActiveTarget() */
+	/** Planar XY handle */
 	UPROPERTY()
-	TObjectPtr<UGizmoEditorAxisSource> AxisXSource;
+	TObjectPtr<UGizmoElementRectangle> PlanarXYElement;
 
-	/** Y-axis source is shared across Gizmos, and created internally during SetActiveTarget() */
+	/** Planar YZ handle */
 	UPROPERTY()
-	TObjectPtr<UGizmoEditorAxisSource> AxisYSource;
+	TObjectPtr<UGizmoElementRectangle> PlanarYZElement;
 
-	/** Z-axis source is shared across Gizmos, and created internally during SetActiveTarget() */
+	/** Planar XZ handle */
 	UPROPERTY()
-	TObjectPtr<UGizmoEditorAxisSource> AxisZSource;
+	TObjectPtr<UGizmoElementRectangle> PlanarXZElement;
 
-	/** X-axis arrow object is shared across Gizmos, and created internally during SetActiveTarget() */
+	/** Rotate X Axis */
 	UPROPERTY()
-	TObjectPtr<UGizmoArrowObject> AxisXObject;
+	TObjectPtr<UGizmoElementTorus> RotateXAxisElement;
 
-	/** Y-axis arrow object is shared across Gizmos, and created internally during SetActiveTarget() */
+	/** Rotate Y Axis */
 	UPROPERTY()
-	TObjectPtr<UGizmoArrowObject> AxisYObject;
+	TObjectPtr<UGizmoElementTorus> RotateYAxisElement;
 
-	/** Z-axis arrow object is shared across Gizmos, and created internally during SetActiveTarget() */
+	/** Rotate Z Axis */
 	UPROPERTY()
-	TObjectPtr<UGizmoArrowObject> AxisZObject;
-	//
-	// Scaling support. 
-	// UE Components only support scaling in local coordinates, so we have to create separate sources for that.
-	//
+	TObjectPtr<UGizmoElementTorus> RotateZAxisElement;
+
+	/** Rotate outer circle */
+	UPROPERTY()
+	TObjectPtr<UGizmoElementCircle> RotateOuterCircleElement;
 
-	/** Local X-axis source (ie 1,0,0) is shared across Scale Gizmos, and created internally during SetActiveTarget() */
+	/** Rotate arcball outer circle */
 	UPROPERTY()
-	TObjectPtr<UGizmoEditorAxisSource> UnitAxisXSource;
+	TObjectPtr<UGizmoElementCircle> RotateArcballOuterElement;
 
-	/** Y-axis source (ie 0,1,0) is shared across Scale Gizmos, and created internally during SetActiveTarget() */
+	/** Rotate arcball inner circle */
 	UPROPERTY()
-	TObjectPtr<UGizmoEditorAxisSource> UnitAxisYSource;
+	TObjectPtr<UGizmoElementCircle> RotateArcballInnerElement;
 
-	/** Z-axis source (ie 0,0,1) is shared across Scale Gizmos, and created internally during SetActiveTarget() */
+	/** Rotate screen space circle */
 	UPROPERTY()
-	TObjectPtr<UGizmoEditorAxisSource> UnitAxisZSource;
+	TObjectPtr<UGizmoElementCircle> RotateScreenSpaceElement;
 
-	/** X-axis arrow object is shared across Gizmos, and created internally during SetActiveTarget() */
+	/** Scale X Axis object */
 	UPROPERTY()
-	TObjectPtr<UGizmoArrowObject> ScaleAxisXObject;
+	TObjectPtr<UGizmoElementArrow> ScaleXAxisElement;
 
-	/** Y-axis arrow object is shared across Gizmos, and created internally during SetActiveTarget() */
+	/** Scale Y Axis object */
 	UPROPERTY()
-	TObjectPtr<UGizmoArrowObject> ScaleAxisYObject;
+	TObjectPtr<UGizmoElementArrow> ScaleYAxisElement;
 
-	/** Z-axis arrow object is shared across Gizmos, and created internally during SetActiveTarget() */
+	/** Scale Z Axis object */
+	UPROPERTY()
+	TObjectPtr<UGizmoElementArrow> ScaleZAxisElement;
+
+	/** Uniform scale object */
+	UPROPERTY()
+	TObjectPtr<UGizmoElementBox> ScaleUniformElement;
+
+	//
+	// Axis Sources
+	//
+	/** Axis that points towards camera, X/Y plane tangents aligned to right/up. Shared across Gizmos, and created internally during SetActiveTarget() */
 	UPROPERTY()
-	TObjectPtr<UGizmoArrowObject> ScaleAxisZObject;
+	TObjectPtr<UGizmoConstantFrameAxisSource> CameraAxisSource;
+
+	// internal function that updates CameraAxisSource by getting current view state from GizmoManager
+	void UpdateCameraAxisSource();
 
 	/** 
 	 * State target is shared across gizmos, and created internally during SetActiveTarget(). 
@@ -245,7 +290,7 @@ protected:
 	 * itself (IToolCommandChangeSource implementation above is called)
 	 */
 	UPROPERTY()
-	TObjectPtr<UGizmoObjectTransformChangeStateTarget> StateTarget;
+	TObjectPtr<UGizmoDependentTransformChangeStateTarget> StateTarget;
 
 	/**
 	 * These are used to let the translation subgizmos use raycasts into the scene to align the gizmo with scene geometry.
@@ -260,15 +305,9 @@ protected:
 	/** Update current gizmo mode based on transform source */
 	void UpdateMode();
 
-	/** Update current coord system based on the active target and transform source */
-	void UpdateCoordSystem();
-
 	/** Enable the given mode with the specified axes, EAxisList::Type::None will hide objects associated with mode */
 	void EnableMode(EGizmoTransformMode InGizmoMode, EAxisList::Type InAxisListToDraw);
 
-	/** Enable object if the specified InGizmoAxis is enabled in InAxisListToDraw*/
-	void EnableObject(UGizmoBaseObject* InGizmoObject, EAxisList::Type InGizmoAxis, EAxisList::Type InAxisListToDraw);
-
 	/** Enable translate using specified axis list */
 	void EnableTranslate(EAxisList::Type InAxisListToDraw);
 
@@ -278,67 +317,60 @@ protected:
 	/** Enable scale using specified axis list */
 	void EnableScale(EAxisList::Type InAxisListToDraw);
 
+	/** Enable planar handles used by translate and scale */
+	void EnablePlanarObjects(bool bEnableX, bool bEnableY, bool bEnableZ);
+
+
 	UPROPERTY()
 	EGizmoTransformMode CurrentMode = EGizmoTransformMode::None;
 
 	UPROPERTY()
 	TEnumAsByte<EAxisList::Type> CurrentAxisToDraw = EAxisList::None;
 
-	/** @return a new instance of the standard axis-translation Gizmo */
-	virtual UInteractiveGizmo* AddAxisTranslationGizmo(
-		UGizmoArrowObject* InArrowObject,
-		IGizmoAxisSource* InAxisSource,
-		IGizmoTransformSource* InTransformSource,
-		IGizmoStateTarget* InStateTarget,
-		EAxisList::Type InAxisType,
-		const FLinearColor InAxisColor);
-
-	/** @return a new instance of the standard plane-translation Gizmo */
-	virtual UInteractiveGizmo* AddPlaneTranslationGizmo(
-		IGizmoAxisSource* InAxisSource,
-		IGizmoTransformSource* InTransformSource,
-		IGizmoStateTarget* InStateTarget);
-
-	/** @return a new instance of the standard axis-rotation Gizmo */
-	virtual UInteractiveGizmo* AddAxisRotationGizmo(
-		IGizmoAxisSource* InAxisSource,
-		IGizmoTransformSource* InTransformSource,
-		IGizmoStateTarget* InStateTarget,
-		EAxisList::Type InAxisType,
-		const FLinearColor InAxisColor);
-
-	/** @return a new instance of the standard axis-scaling Gizmo */
-	virtual UInteractiveGizmo* AddAxisScaleGizmo(
-		UGizmoArrowObject* InArrowObject,
-		IGizmoAxisSource* InGizmoAxisSource, IGizmoAxisSource* InParameterAxisSource,
-		IGizmoTransformSource* InTransformSource,
-		IGizmoStateTarget* InStateTarget,
-		EAxisList::Type InAxisType,
-		const FLinearColor InAxisColor);
-
-	/** @return a new instance of the standard plane-scaling Gizmo */
-	virtual UInteractiveGizmo* AddPlaneScaleGizmo(
-		IGizmoAxisSource* InGizmoAxisSource, IGizmoAxisSource* InParameterAxisSource,
-		IGizmoTransformSource* InTransformSource,
-		IGizmoStateTarget* InStateTarget);
-
-	/** @return a new instance of the standard plane-scaling Gizmo */
-	virtual UInteractiveGizmo* AddUniformScaleGizmo(
-		IGizmoAxisSource* InGizmoAxisSource, IGizmoAxisSource* InParameterAxisSource,
-		IGizmoTransformSource* InTransformSource,
-		IGizmoStateTarget* InStateTarget);
+	/** Construct translate axis handle */
+	virtual UGizmoElementArrow* MakeTranslateAxis(const FVector& InAxisDir, const FVector& InSideDir, UMaterialInterface* InMaterial);
+
+	/** Construct scale axis handle */
+	virtual UGizmoElementArrow* MakeScaleAxis(const FVector& InAxisDir, const FVector& InSideDir, UMaterialInterface* InMaterial);
+
+	/** Construct rotate axis handle */
+	virtual UGizmoElementTorus* MakeRotateAxis(const FVector& Normal, const FVector& TorusAxis0, const FVector& TorusAxis1, 
+		UMaterialInterface* InMaterial, UMaterialInterface* InCurrentMaterial);
+
+	/** Construct uniform scale handle */
+	virtual UGizmoElementBox* MakeUniformScaleHandle();
+
+	/** Construct planar axis handle */
+	virtual UGizmoElementRectangle* MakePlanarHandle(const FVector& InUpDirection, const FVector& InSideDirection, const FVector& InPlaneNormal,
+		UMaterialInterface* InMaterial, const FLinearColor& InVertexColor);
+
+	/** Construct translate screen space handle */
+	virtual UGizmoElementRectangle* MakeTranslateScreenSpaceHandle();
+
+	/** Construct rotate screen space handle */
+	virtual UGizmoElementCircle* MakeRotateCircleHandle(float InRadius, const FLinearColor& InColor, float bFill);
 
 	// Axis and Plane TransformSources use this function to execute worldgrid snap queries
 	bool PositionSnapFunction(const FVector& WorldPosition, FVector& SnappedPositionOut) const;
 	FQuat RotationSnapFunction(const FQuat& DeltaRotation) const;
 
 	/** Materials and colors to be used when drawing the items for each axis */
-	TObjectPtr<UMaterialInterface> TransparentPlaneMaterialXY;
+	UPROPERTY()
+	TObjectPtr<UMaterialInterface> TransparentVertexColorMaterial;
+	UPROPERTY()
 	TObjectPtr<UMaterialInterface> GridMaterial;
-
+	UPROPERTY()
 	TObjectPtr<UMaterialInstanceDynamic> AxisMaterialX;
+	UPROPERTY()
 	TObjectPtr<UMaterialInstanceDynamic> AxisMaterialY;
+	UPROPERTY()
 	TObjectPtr<UMaterialInstanceDynamic> AxisMaterialZ;
+	UPROPERTY()
 	TObjectPtr<UMaterialInstanceDynamic> CurrentAxisMaterial;
+	UPROPERTY()
+	TObjectPtr<UMaterialInstanceDynamic> GreyMaterial;
+	UPROPERTY()
+	TObjectPtr<UMaterialInstanceDynamic> WhiteMaterial;
+	UPROPERTY()
 	TObjectPtr<UMaterialInstanceDynamic> OpaquePlaneMaterialXY;
 };
diff --git a/Engine/Source/Runtime/InteractiveToolsFramework/Private/BaseGizmos/GizmoElementArrow.cpp b/Engine/Source/Runtime/InteractiveToolsFramework/Private/BaseGizmos/GizmoElementArrow.cpp
new file mode 100644
--- /dev/null
+++ b/Engine/Source/Runtime/InteractiveToolsFramework/Private/BaseGizmos/GizmoElementArrow.cpp
@@ -0,0 +1,246 @@
+// Copyright Epic Games, Inc. All Rights Reserved.
+ 
+#include "BaseGizmos/GizmoElementArrow.h"
+#include "BaseGizmos/GizmoElementBox.h"
+#include "BaseGizmos/GizmoElementCone.h"
+#include "BaseGizmos/GizmoElementCylinder.h"
+#include "BaseGizmos/GizmoMath.h"
+#include "BaseGizmos/GizmoRenderingUtil.h"
+#include "InputState.h"
+#include "Materials/MaterialInterface.h"
+
+UGizmoElementArrow::UGizmoElementArrow()
+{
+	HeadType = EGizmoElementArrowHeadType::Cone;
+	CylinderElement = NewObject<UGizmoElementCylinder>();
+	ConeElement = NewObject<UGizmoElementCone>();
+	BoxElement = nullptr;
+}
+
+void UGizmoElementArrow::Render(IToolsContextRenderAPI* RenderAPI, const FRenderTraversalState& RenderState)
+{
+	if (!IsVisible())
+	{
+		return;
+	}
+
+	const FSceneView* View = RenderAPI->GetSceneView();
+	bool bVisibleViewDependent = GetViewDependentVisibility(View, RenderState.LocalToWorldTransform, Base);
+
+	FRenderTraversalState RenderStateCopy = RenderState;
+
+	if (bVisibleViewDependent)
+	{
+		FQuat AlignRot;
+		if (GetViewAlignRot(View, RenderState.LocalToWorldTransform, Base, AlignRot))
+		{
+			RenderStateCopy.LocalToWorldTransform = FTransform(AlignRot, Base) * RenderState.LocalToWorldTransform;
+		}
+		else
+		{
+			RenderStateCopy.LocalToWorldTransform = FTransform(Base) * RenderState.LocalToWorldTransform;
+		}
+
+		UpdateRenderTraversalState(RenderStateCopy);
+
+		check(CylinderElement);
+		UGizmoElementBase* Element = Cast<UGizmoElementBase>(CylinderElement);
+		Element->Render(RenderAPI, RenderStateCopy);
+
+		if (HeadType == EGizmoElementArrowHeadType::Cone)
+		{
+			check(ConeElement);
+			ConeElement->Render(RenderAPI, RenderStateCopy);
+		}
+		else // (HeadType == EGizmoElementArrowHeadType::Cube)
+		{
+			check(BoxElement);
+			BoxElement->Render(RenderAPI, RenderStateCopy);
+		}
+	}
+
+	CacheRenderState(RenderStateCopy.LocalToWorldTransform, bVisibleViewDependent);
+}
+
+FInputRayHit UGizmoElementArrow::LineTrace(const FVector RayOrigin, const FVector RayDirection)
+{
+	if (!IsHittableInView())
+	{
+		return FInputRayHit();
+	}
+
+	check(CylinderElement);
+	FInputRayHit Hit = CylinderElement->LineTrace(RayOrigin, RayDirection);
+
+	if (!Hit.bHit)
+	{
+		if (HeadType == EGizmoElementArrowHeadType::Cone)
+		{
+			check(ConeElement);
+			Hit = ConeElement->LineTrace(RayOrigin, RayDirection);
+		}
+		else // (HeadType == EGizmoElementArrowHeadType::Cube)
+		{
+			check(BoxElement);
+			Hit = BoxElement->LineTrace(RayOrigin, RayDirection);
+		}
+	}
+
+	return Hit;
+}
+
+FBoxSphereBounds UGizmoElementArrow::CalcBounds(const FTransform& LocalToWorld) const
+{
+	// @todo - implement box-sphere bounds calculation
+	return FBoxSphereBounds();
+}
+
+void UGizmoElementArrow::SetBase(const FVector& InBase)
+{
+	Base = InBase;
+	UpdateArrowBody();
+	UpdateArrowHead();
+}
+
+FVector UGizmoElementArrow::GetBase() const
+{
+	return Base;
+}
+
+void UGizmoElementArrow::SetDirection(const FVector& InDirection)
+{
+	Direction = InDirection;
+	Direction.Normalize();
+	UpdateArrowBody();
+	UpdateArrowHead();
+}
+
+FVector UGizmoElementArrow::GetDirection() const
+{
+	return Direction;
+}
+
+void UGizmoElementArrow::SetSideDirection(const FVector& InSideDirection)
+{
+	SideDirection = InSideDirection;
+	SideDirection.Normalize();
+	UpdateArrowHead();
+}
+
+FVector UGizmoElementArrow::GetSideDirection() const
+{
+	return SideDirection;
+}
+
+void UGizmoElementArrow::SetBodyLength(float InBodyLength)
+{
+	BodyLength = InBodyLength;
+	UpdateArrowBody();
+	UpdateArrowHead();
+}
+
+float UGizmoElementArrow::GetBodyLength() const
+{
+	return BodyLength;
+}
+
+void UGizmoElementArrow::SetBodyRadius(float InBodyRadius)
+{
+	BodyRadius = InBodyRadius;
+	UpdateArrowBody();
+	UpdateArrowHead();
+}
+
+float UGizmoElementArrow::GetBodyRadius() const
+{
+	return BodyRadius;
+}
+
+void UGizmoElementArrow::SetHeadLength(float InHeadLength)
+{
+	HeadLength = InHeadLength;
+	UpdateArrowHead();
+}
+
+float UGizmoElementArrow::GetHeadLength() const
+{
+	return HeadLength;
+}
+
+void UGizmoElementArrow::SetHeadRadius(float InHeadRadius)
+{
+	HeadRadius = InHeadRadius;
+	UpdateArrowHead();
+}
+
+float UGizmoElementArrow::GetHeadRadius() const
+{
+	return HeadRadius;
+}
+
+void UGizmoElementArrow::SetNumSides(int32 InNumSides)
+{
+	NumSides = InNumSides;
+	UpdateArrowBody();
+	UpdateArrowHead();
+}
+
+int32 UGizmoElementArrow::GetNumSides() const
+{
+	return NumSides;
+}
+
+void UGizmoElementArrow::SetHeadType(EGizmoElementArrowHeadType InHeadType)
+{
+	if (InHeadType != HeadType)
+	{
+		HeadType = InHeadType;
+
+		if (HeadType == EGizmoElementArrowHeadType::Cone)
+		{
+			ConeElement = NewObject<UGizmoElementCone>();
+			BoxElement = nullptr;
+		}
+		else // (HeadType == EGizmoElementArrowHeadType::Cube)
+		{
+			BoxElement = NewObject<UGizmoElementBox>();
+			ConeElement = nullptr;
+		}
+		UpdateArrowHead();
+	}
+}
+
+EGizmoElementArrowHeadType UGizmoElementArrow::GetHeadType() const
+{
+	return HeadType;
+}
+
+void UGizmoElementArrow::UpdateArrowBody()
+{
+	CylinderElement->SetBase(FVector::ZeroVector);
+	CylinderElement->SetDirection(Direction);
+	CylinderElement->SetHeight(BodyLength);
+	CylinderElement->SetNumSides(NumSides);
+	CylinderElement->SetRadius(BodyRadius);
+}
+
+void UGizmoElementArrow::UpdateArrowHead()
+{
+	if (HeadType == EGizmoElementArrowHeadType::Cone)
+	{
+		check(ConeElement);
+		ConeElement->SetBase(Direction * (BodyLength - HeadLength * 0.1f));
+		ConeElement->SetDirection(Direction);
+		ConeElement->SetHeight(HeadLength);
+		ConeElement->SetRadius(HeadRadius);
+		ConeElement->SetNumSides(NumSides);
+	}
+	else // (HeadType == EGizmoElementArrowHeadType::Cube)
+	{
+		check(BoxElement);
+		BoxElement->SetCenter(Direction * (BodyLength + HeadLength * 0.5f));
+		BoxElement->SetUpDirection(Direction);
+		BoxElement->SetSideDirection(SideDirection);
+		BoxElement->SetDimensions(FVector(HeadLength, HeadLength, HeadLength));
+	}
+}
diff --git a/Engine/Source/Runtime/InteractiveToolsFramework/Private/BaseGizmos/GizmoElementBase.cpp b/Engine/Source/Runtime/InteractiveToolsFramework/Private/BaseGizmos/GizmoElementBase.cpp
new file mode 100644
--- /dev/null
+++ b/Engine/Source/Runtime/InteractiveToolsFramework/Private/BaseGizmos/GizmoElementBase.cpp
@@ -0,0 +1,365 @@
+// Copyright Epic Games, Inc. All Rights Reserved.
+ 
+#include "BaseGizmos/GizmoElementBase.h"
+#include "BaseGizmos/GizmoRenderingUtil.h"
+#include "Materials/MaterialInterface.h"
+
+DEFINE_LOG_CATEGORY_STATIC(LogGizmoElements, Log, All);
+
+bool UGizmoElementBase::GetViewDependentVisibility(const FSceneView* View, const FTransform& InLocalToWorldTransform, const FVector& InLocalCenter) const
+{
+	if (ViewDependentType == EGizmoElementViewDependentType::None || ViewAlignType == EGizmoElementViewAlignType::PointOnly || ViewAlignType == EGizmoElementViewAlignType::PointEye)
+	{
+		return true;
+	}
+
+	FVector ViewDir;
+
+	if (View->IsPerspectiveProjection())
+	{
+
+		FVector WorldCenter = InLocalToWorldTransform.TransformPosition(InLocalCenter);
+		ViewDir = WorldCenter - View->ViewLocation;
+		ViewDir.Normalize();
+	}
+	else
+	{
+		ViewDir = View->GetViewDirection();
+	}
+
+	bool bVisibleViewDependent;
+	if (ViewDependentType == EGizmoElementViewDependentType::Axis)
+	{
+		bVisibleViewDependent = FMath::Abs(FVector::DotProduct(ViewDependentAxis, ViewDir)) < DefaultViewAlignAxialMaxCosAngleTol;
+	}
+	else // (ViewDependentType == EGizmoElementViewDependentType::Plane)
+	{
+		bVisibleViewDependent = FMath::Abs(FVector::DotProduct(ViewDependentAxis, ViewDir)) > DefaultViewAlignPlanarMinCosAngleTol;
+	}
+
+	return bVisibleViewDependent;
+
+}
+
+
+bool UGizmoElementBase::GetViewAlignRot(const FSceneView* View, const FTransform& InLocalToWorldTransform, const FVector& InLocalCenter, FQuat& OutAlignRot) const
+{
+	if (ViewAlignType == EGizmoElementViewAlignType::None)
+	{
+		return false;
+	}
+
+	FVector Scale = InLocalToWorldTransform.GetScale3D();
+	if (!FMath::IsNearlyEqual(Scale.X, Scale.Y, KINDA_SMALL_NUMBER) || !FMath::IsNearlyEqual(Scale.X, Scale.Z, KINDA_SMALL_NUMBER))
+	{
+		// Warn that non-uniform scale is not currently supported 
+		bool bNonUniformScaleWarning = true;
+		if (bNonUniformScaleWarning)
+		{
+			UE_LOG(LogGizmoElements, Warning, TEXT("Gizmo element library view-dependent alignment does not currently support non-uniform scale (%f, %f %f)."),
+				Scale.X, Scale.Y, Scale.Z);
+			bNonUniformScaleWarning = false;
+		}
+		return false;
+	}
+
+	FVector LocalViewDir;
+
+	FTransform WorldToLocalTransform = InLocalToWorldTransform.Inverse();
+
+	if (View->IsPerspectiveProjection())
+	{
+		FVector LocalViewLocation = WorldToLocalTransform.TransformPosition(View->ViewLocation);
+		LocalViewDir = InLocalCenter - LocalViewLocation;
+		LocalViewDir.Normalize();
+	}
+	else
+	{
+		FVector WorldViewDir = View->GetViewDirection();
+		LocalViewDir = WorldToLocalTransform.GetRotation().RotateVector(WorldViewDir);
+		LocalViewDir.Normalize();
+	}
+		
+	if (ViewAlignType == EGizmoElementViewAlignType::PointOnly)
+	{
+		OutAlignRot = FQuat::FindBetweenNormals(ViewAlignNormal, -LocalViewDir);
+	}
+	else if (ViewAlignType == EGizmoElementViewAlignType::PointEye)
+	{
+		FVector Right = ViewAlignAxis ^ ViewAlignNormal;
+		Right.Normalize();
+		FVector Up = ViewAlignNormal ^ Right;
+
+		FMatrix LocalToCanonicalBasis(
+			FPlane(ViewAlignNormal.X, Right.X, Up.X, 0.0),
+			FPlane(ViewAlignNormal.Y, Right.Y, Up.Y, 0.0),
+			FPlane(ViewAlignNormal.Z, Right.Z, Up.Z, 0.0),
+			FPlane::ZeroVector);
+
+		FVector LocalViewUp = WorldToLocalTransform.TransformVector(View->GetViewUp());
+		FVector TargetFwd = -LocalViewDir;
+		FVector TargetRight = LocalViewUp ^ TargetFwd;
+		TargetRight.Normalize();
+		FVector TargetUp = TargetFwd ^ TargetRight;
+
+		FMatrix CanonicalToAlignRotBasis(
+			FPlane(TargetFwd.X, TargetFwd.Y, TargetFwd.Z, 0.0),
+			FPlane(TargetRight.X, TargetRight.Y, TargetRight.Z, 0.0),
+			FPlane(TargetUp.X, TargetUp.Y, TargetUp.Z, 0.0),
+			FPlane::ZeroVector);
+
+		FMatrix LocalToAligned = LocalToCanonicalBasis * CanonicalToAlignRotBasis;
+		OutAlignRot = LocalToAligned.ToQuat();
+		OutAlignRot.Normalize();
+	}
+	else if (ViewAlignType == EGizmoElementViewAlignType::Axial)
+	{
+		// if Axis and Dir are almost coincident, do not adjust the rotation
+		if ((FMath::Abs(FVector::DotProduct(ViewAlignAxis, -LocalViewDir))) >= DefaultViewAlignAxialMaxCosAngleTol)
+		{
+			return false;
+		}
+
+		FVector TargetRight = -LocalViewDir ^ ViewAlignAxis;
+		TargetRight.Normalize();
+		FVector TargetNormal = ViewAlignAxis ^ TargetRight;
+		TargetNormal.Normalize();
+		OutAlignRot = FQuat::FindBetweenNormals(ViewAlignNormal, TargetNormal);
+	}
+
+	return true;
+}
+
+const UMaterialInterface* UGizmoElementBase::GetCurrentMaterial(const FRenderTraversalState& RenderState, EGizmoElementInteractionState InState) const
+{
+	if (InState == EGizmoElementInteractionState::Hovering)
+	{
+		if (HoverMaterial)
+		{
+			return HoverMaterial;
+		}
+		else if (RenderState.HoverMaterial.IsValid())
+		{
+			return RenderState.HoverMaterial.Get();
+		}
+	}
+	else if (InState == EGizmoElementInteractionState::Interacting)
+	{
+		if (InteractMaterial)
+		{
+			return InteractMaterial;
+		}
+		else if (RenderState.InteractMaterial.IsValid())
+		{
+			return RenderState.InteractMaterial.Get();
+		}
+	}
+	else
+	{
+		if (Material)
+		{
+			return Material;
+		}
+		else if (RenderState.Material.IsValid())
+		{
+			return RenderState.Material.Get();
+		}
+
+	}
+
+	return nullptr;
+}
+
+void UGizmoElementBase::CacheRenderState(const FTransform& InLocalToWorldState, bool InVisibleViewDependent)
+{
+	CachedLocalToWorldTransform = InLocalToWorldState;
+	bHasCachedLocalToWorldTransform = true;
+	bCachedVisibleViewDependent = InVisibleViewDependent;
+}
+
+void UGizmoElementBase::ResetCachedRenderState()
+{
+	bHasCachedLocalToWorldTransform = false;
+	bCachedVisibleViewDependent = true;
+}
+
+void UGizmoElementBase::UpdateRenderTraversalState(FRenderTraversalState& InRenderTraversalState)
+{
+	if (HoverMaterial)
+	{
+		InRenderTraversalState.HoverMaterial = HoverMaterial;
+	}
+	if (InteractMaterial)
+	{
+		InRenderTraversalState.InteractMaterial = InteractMaterial;
+	}
+	if (Material)
+	{
+		InRenderTraversalState.Material = Material;
+	}
+}
+
+bool UGizmoElementBase::IsVisible() const
+{
+	bool bVisible = static_cast<uint8>(ElementState) & static_cast<uint8>(EGizmoElementState::Visible);
+	return (bEnabled && bVisible);
+}
+
+bool UGizmoElementBase::IsHittable() const
+{
+	bool bHittable = static_cast<uint8>(ElementState) & static_cast<uint8>(EGizmoElementState::Hittable);
+	return (bEnabled && bHittable);
+}
+
+bool UGizmoElementBase::IsHittableInView() const
+{
+	return (IsHittable() && bHasCachedLocalToWorldTransform && (!IsVisible() || bCachedVisibleViewDependent));
+}
+
+void UGizmoElementBase::SetEnabled(bool InEnabled)
+{
+	bEnabled = InEnabled;
+}
+
+bool UGizmoElementBase::GetEnabled() const
+{
+	return bEnabled;
+}
+
+void UGizmoElementBase::SetElementState(EGizmoElementState InElementState)
+{
+	ElementState = InElementState;
+}
+
+EGizmoElementState UGizmoElementBase::GetElementState() const
+{
+	return ElementState;
+}
+
+void UGizmoElementBase::SetElementInteractionState(EGizmoElementInteractionState InElementInteractionState)
+{
+	ElementInteractionState = InElementInteractionState;
+}
+
+EGizmoElementInteractionState UGizmoElementBase::GetElementInteractionState() const
+{
+	return ElementInteractionState;
+}
+
+void UGizmoElementBase::SetViewDependentType(EGizmoElementViewDependentType InViewDependentType)
+{
+	ViewDependentType = InViewDependentType;
+}
+
+EGizmoElementViewDependentType UGizmoElementBase::GetViewDependentType() const
+{
+	return ViewDependentType;
+}
+
+void UGizmoElementBase::SetViewDependentAngleTol(float InAngleTol)
+{
+	ViewDependentAngleTol = InAngleTol;
+	ViewDependentAxialMaxCosAngleTol = FMath::Abs(FMath::Cos(ViewDependentAngleTol));
+	ViewDependentPlanarMinCosAngleTol = FMath::Abs(FMath::Cos(HALF_PI + ViewDependentAngleTol));
+}
+
+float UGizmoElementBase::GetViewDependentAngleTol() const
+{
+	return ViewDependentAngleTol;
+}
+
+void UGizmoElementBase::SetViewDependentAxis(FVector InViewDependentAxis)
+{
+	ViewDependentAxis = InViewDependentAxis;
+	ViewDependentAxis.Normalize();
+}
+
+FVector UGizmoElementBase::GetViewDependentAxis() const
+{
+	return ViewDependentAxis;
+}
+
+void UGizmoElementBase::SetViewAlignType(EGizmoElementViewAlignType InViewAlignType)
+{
+	ViewAlignType = InViewAlignType;
+}
+
+EGizmoElementViewAlignType UGizmoElementBase::GetViewAlignType() const
+{
+	return ViewAlignType;
+}
+
+void UGizmoElementBase::SetViewAlignAxis(FVector InViewAlignAxis)
+{
+	ViewAlignAxis = InViewAlignAxis;
+	ViewAlignAxis.Normalize();
+}
+
+FVector UGizmoElementBase::GetViewAlignAxis() const
+{
+	return ViewAlignAxis;
+}
+
+void UGizmoElementBase::SetViewAlignNormal(FVector InViewAlignNormal)
+{
+	ViewAlignNormal = InViewAlignNormal;
+	ViewAlignNormal.Normalize();
+}
+
+FVector UGizmoElementBase::GetViewAlignNormal() const
+{
+	return ViewAlignNormal;
+}
+
+void UGizmoElementBase::SetViewAlignAxialAngleTol(float InAngleTol)
+{
+	ViewAlignAxialAngleTol = InAngleTol;
+	ViewAlignAxialMaxCosAngleTol = FMath::Abs(FMath::Cos(ViewAlignAxialAngleTol));
+}
+
+float UGizmoElementBase::GetViewAlignAxialAngleTol() const
+{
+	return ViewAlignAxialAngleTol;
+}
+
+void UGizmoElementBase::SetMaterial(UMaterialInterface* InMaterial)
+{
+	Material = InMaterial;
+}
+
+UMaterialInterface* UGizmoElementBase::GetMaterial() const
+{
+	return Material;
+}
+
+void UGizmoElementBase::SetHoverMaterial(UMaterialInterface* InHoverMaterial)
+{
+	HoverMaterial = InHoverMaterial;
+}
+
+UMaterialInterface* UGizmoElementBase::GetHoverMaterial() const
+{
+	return HoverMaterial;
+}
+
+void UGizmoElementBase::SetInteractMaterial(UMaterialInterface* InInteractMaterial)
+{
+	InteractMaterial = InInteractMaterial;
+}
+
+UMaterialInterface* UGizmoElementBase::GetInteractMaterial() const
+{
+	return InteractMaterial;
+}
+
+void UGizmoElementBase::SetVertexColor(const FColor& InVertexColor)
+{
+	VertexColor = InVertexColor;
+}
+
+FColor UGizmoElementBase::GetVertexColor() const
+{
+	return VertexColor;
+}
+
+
diff --git a/Engine/Source/Runtime/InteractiveToolsFramework/Private/BaseGizmos/GizmoElementBox.cpp b/Engine/Source/Runtime/InteractiveToolsFramework/Private/BaseGizmos/GizmoElementBox.cpp
new file mode 100644
--- /dev/null
+++ b/Engine/Source/Runtime/InteractiveToolsFramework/Private/BaseGizmos/GizmoElementBox.cpp
@@ -0,0 +1,130 @@
+// Copyright Epic Games, Inc. All Rights Reserved.
+ 
+#include "BaseGizmos/GizmoElementBox.h"
+#include "BaseGizmos/GizmoMath.h"
+#include "BaseGizmos/GizmoRenderingUtil.h"
+#include "Materials/MaterialInterface.h"
+#include "InputState.h"
+#include "SceneManagement.h"
+
+void UGizmoElementBox::Render(IToolsContextRenderAPI* RenderAPI, const FRenderTraversalState& RenderState)
+{
+	if (!IsVisible())
+	{
+		return;
+	}
+
+	check(RenderAPI);
+	const FSceneView* View = RenderAPI->GetSceneView();
+
+	FTransform LocalToWorldTransform = RenderState.LocalToWorldTransform;
+
+	bool bVisibleViewDependent = GetViewDependentVisibility(View, LocalToWorldTransform, Center);
+
+	if (bVisibleViewDependent)
+	{
+		const UMaterialInterface* UseMaterial = GetCurrentMaterial(RenderState, ElementInteractionState);
+
+		if (UseMaterial)
+		{
+			FQuat AlignRot;
+			FVector AdjustedSideDir, AdjustedUpDir;
+			if (GetViewAlignRot(View, LocalToWorldTransform, Center, AlignRot))
+			{
+				AdjustedSideDir = AlignRot.RotateVector(SideDirection);
+				AdjustedUpDir = AlignRot.RotateVector(UpDirection);
+			}
+			else
+			{
+				AdjustedSideDir = SideDirection;
+				AdjustedUpDir = UpDirection;
+			}
+
+			FQuat Rotation = FRotationMatrix::MakeFromYZ(AdjustedSideDir, AdjustedUpDir).ToQuat();
+			FPrimitiveDrawInterface* PDI = RenderAPI->GetPrimitiveDrawInterface();
+			LocalToWorldTransform = FTransform(Rotation, Center) * LocalToWorldTransform;
+			const FVector HalfDimensions = Dimensions * 0.5;
+			DrawBox(PDI, LocalToWorldTransform.ToMatrixWithScale(), HalfDimensions, UseMaterial->GetRenderProxy(), SDPG_Foreground);
+		}
+	}
+
+	CacheRenderState(LocalToWorldTransform, bVisibleViewDependent);
+}
+
+
+
+FInputRayHit UGizmoElementBox::LineTrace(const FVector RayCenter, const FVector RayDirection)
+{
+	if (IsHittableInView())
+	{
+		// @todo - modify ray-box intersection to work with updated properties
+
+#if 0
+		bool bIntersects = false;
+		float RayParam = 0.0f;
+
+
+		const FVector BoxOrigin = Origin + Offset;
+		GizmoMath::RayBoxIntersection(
+			BoxOrigin, CubeObject->Axis, CubeObject->Radius, CubeObject->Height,
+			ClickPos.WorldRay.Origin, ClickPos.WorldRay.Direction,
+			bIntersects, RayParam);
+
+		if (bIntersects)
+		{
+			return FInputRayHit(RayParam);
+		}
+#endif
+	}
+
+	return FInputRayHit();
+}
+
+FBoxSphereBounds UGizmoElementBox::CalcBounds(const FTransform& LocalToWorld) const
+{
+	// @todo - implement box-sphere bounds calculation
+	return FBoxSphereBounds();
+}
+
+void UGizmoElementBox::SetCenter(const FVector& InCenter)
+{
+	Center = InCenter;
+}
+
+FVector UGizmoElementBox::GetCenter() const
+{
+	return Center;
+}
+
+void UGizmoElementBox::SetUpDirection(const FVector& InUpDirection)
+{
+	UpDirection = InUpDirection;
+	UpDirection.Normalize();
+}
+
+FVector UGizmoElementBox::GetUpDirection() const
+{
+	return UpDirection;
+}
+
+void UGizmoElementBox::SetSideDirection(const FVector& InSideDirection)
+{
+	SideDirection = InSideDirection;
+	SideDirection.Normalize();
+}
+
+FVector UGizmoElementBox::GetSideDirection() const
+{
+	return SideDirection;
+}
+
+FVector UGizmoElementBox::GetDimensions() const
+{
+	return Dimensions;
+}
+
+void UGizmoElementBox::SetDimensions(const FVector& InDimensions)
+{
+	Dimensions = InDimensions;
+}
+
diff --git a/Engine/Source/Runtime/InteractiveToolsFramework/Private/BaseGizmos/GizmoElementCircle.cpp b/Engine/Source/Runtime/InteractiveToolsFramework/Private/BaseGizmos/GizmoElementCircle.cpp
new file mode 100644
--- /dev/null
+++ b/Engine/Source/Runtime/InteractiveToolsFramework/Private/BaseGizmos/GizmoElementCircle.cpp
@@ -0,0 +1,186 @@
+// Copyright Epic Games, Inc. All Rights Reserved.
+ 
+#include "BaseGizmos/GizmoElementCircle.h"
+#include "BaseGizmos/GizmoRenderingUtil.h"
+#include "BaseGizmos/GizmoMath.h"
+#include "InputState.h"
+#include "Materials/MaterialInterface.h"
+#include "SceneManagement.h"
+
+void UGizmoElementCircle::Render(IToolsContextRenderAPI* RenderAPI, const FRenderTraversalState& RenderState)
+{
+	if (!IsVisible())
+	{
+		return;
+	}
+
+	check(RenderAPI);
+	const FSceneView* View = RenderAPI->GetSceneView();
+
+	FTransform LocalToWorldTransform = RenderState.LocalToWorldTransform;
+	bool bVisibleViewDependent = GetViewDependentVisibility(View, LocalToWorldTransform, Center);
+
+	if (bVisibleViewDependent)
+	{
+		FVector Axis0, Axis1;
+
+		if (bScreenSpace)
+		{
+			Axis0 = View->GetViewUp();
+			Axis1 = View->GetViewRight();
+		}
+		else
+		{
+			FVector AdjustedNormal;
+			FQuat AlignRot;
+			if (GetViewAlignRot(View, LocalToWorldTransform, Center, AlignRot))
+			{
+				AdjustedNormal = AlignRot.RotateVector(Normal);
+			}
+			else
+			{
+				AdjustedNormal = Normal;
+			}
+
+			const FVector WorldNormal = LocalToWorldTransform.TransformVectorNoScale(AdjustedNormal);
+			GizmoMath::MakeNormalPlaneBasis(WorldNormal, Axis0, Axis1);
+			Axis0.Normalize();
+			Axis1.Normalize();
+		}
+
+		const float WorldRadius = Radius * LocalToWorldTransform.GetScale3D().X;
+		const FVector WorldCenter = LocalToWorldTransform.TransformPosition(Center);
+
+		FPrimitiveDrawInterface* PDI = RenderAPI->GetPrimitiveDrawInterface();
+
+		if (bDrawMesh)
+		{
+			if (const UMaterialInterface* UseMaterial = GetCurrentMaterial(RenderState, ElementInteractionState))
+			{
+				DrawDisc(PDI, WorldCenter, Axis0, Axis1, VertexColor, WorldRadius, NumSides, UseMaterial->GetRenderProxy(), SDPG_Foreground);
+			}
+		}
+
+		if (bDrawLine)
+		{
+			DrawCircle(PDI, WorldCenter, Axis0, Axis1, LineColor, WorldRadius, NumSides, SDPG_Foreground, 0.0f);
+		}
+	}
+	CacheRenderState(LocalToWorldTransform, bVisibleViewDependent);
+}
+
+FInputRayHit UGizmoElementCircle::LineTrace(const FVector RayOrigin, const FVector RayDirection)
+{
+	if (IsHittableInView())
+	{
+		// @todo - implement ray-circle intersection 
+	}
+	return FInputRayHit();
+}
+
+FBoxSphereBounds UGizmoElementCircle::CalcBounds(const FTransform& LocalToWorld) const
+{
+	// @todo - implement box-sphere bounds calculation
+	return FBoxSphereBounds();
+}
+
+void UGizmoElementCircle::SetCenter(FVector InCenter)
+{
+	Center = InCenter;
+}
+
+FVector UGizmoElementCircle::GetCenter() const
+{
+	return Center;
+}
+
+void UGizmoElementCircle::SetNormal(FVector InNormal)
+{
+	Normal = InNormal;
+	Normal.Normalize();
+}
+
+FVector UGizmoElementCircle::GetNormal() const
+{
+	return Normal;
+}
+
+void UGizmoElementCircle::SetRadius(float InRadius)
+{
+	Radius = InRadius;
+}
+
+float UGizmoElementCircle::GetRadius() const
+{
+	return Radius;
+}
+
+void UGizmoElementCircle::SetNumSides(int InNumSides)
+{
+	NumSides = InNumSides;
+}
+
+int UGizmoElementCircle::GetNumSides() const
+{
+	return NumSides;
+}
+
+void UGizmoElementCircle::SetLineColor(const FLinearColor& InColor)
+{
+	LineColor = InColor;
+}
+
+FLinearColor UGizmoElementCircle::GetLineColor() const
+{
+	return LineColor;
+}
+
+void UGizmoElementCircle::SetDrawMesh(bool InDrawMesh)
+{
+	bDrawMesh = InDrawMesh;
+}
+
+bool UGizmoElementCircle::GetDrawMesh() const
+{
+	return bDrawMesh;
+}
+
+void UGizmoElementCircle::SetDrawLine(bool InDrawLine)
+{
+	bDrawLine = InDrawLine;
+}
+
+bool UGizmoElementCircle::GetDrawLine() const
+{
+	return bDrawLine;
+}
+
+void UGizmoElementCircle::SetHitMesh(bool InHitMesh)
+{
+	bHitMesh = InHitMesh;
+}
+
+bool UGizmoElementCircle::GetHitMesh() const
+{
+	return bHitMesh;
+}
+
+void UGizmoElementCircle::SetHitLine(bool InHitLine)
+{
+	bHitLine = InHitLine;
+}
+
+bool UGizmoElementCircle::GetHitLine() const
+{
+	return bHitLine;
+}
+
+void UGizmoElementCircle::SetScreenSpace(bool InScreenSpace)
+{
+	bScreenSpace = InScreenSpace;
+}
+
+bool UGizmoElementCircle::GetScreenSpace()
+{
+	return bScreenSpace;
+}
\ No newline at end of file
diff --git a/Engine/Source/Runtime/InteractiveToolsFramework/Private/BaseGizmos/GizmoElementCone.cpp b/Engine/Source/Runtime/InteractiveToolsFramework/Private/BaseGizmos/GizmoElementCone.cpp
new file mode 100644
--- /dev/null
+++ b/Engine/Source/Runtime/InteractiveToolsFramework/Private/BaseGizmos/GizmoElementCone.cpp
@@ -0,0 +1,143 @@
+// Copyright Epic Games, Inc. All Rights Reserved.
+ 
+#include "BaseGizmos/GizmoElementCone.h"
+#include "BaseGizmos/GizmoRenderingUtil.h"
+#include "BaseGizmos/GizmoMath.h"
+#include "InputState.h"
+#include "Materials/MaterialInterface.h"
+#include "SceneManagement.h"
+
+void UGizmoElementCone::Render(IToolsContextRenderAPI* RenderAPI, const FRenderTraversalState& RenderState)
+{
+	if (!IsVisible())
+	{
+		return;
+	}
+
+	check(RenderAPI);
+	const FSceneView* View = RenderAPI->GetSceneView();
+
+	FTransform LocalToWorldTransform = RenderState.LocalToWorldTransform;
+
+	bool bVisibleViewDependent = GetViewDependentVisibility(View, LocalToWorldTransform, Base);
+
+	if (bVisibleViewDependent)
+	{
+		const UMaterialInterface* UseMaterial = GetCurrentMaterial(RenderState, ElementInteractionState);
+
+		if (UseMaterial)
+		{
+			FQuat AlignRot;
+			const bool bHasAlignRot = GetViewAlignRot(View, LocalToWorldTransform, Base, AlignRot);
+			const FVector AdjustedDir = bHasAlignRot ? AlignRot.RotateVector(Direction) : Direction;
+
+			FQuat Rotation = FRotationMatrix::MakeFromX(-AdjustedDir).ToQuat();
+
+			const FVector Scale(Height);
+			const FVector Origin = Base + AdjustedDir * Height;
+
+			LocalToWorldTransform = FTransform(Rotation, Origin, Scale) * LocalToWorldTransform;
+			const float ConeSide = FMath::Sqrt(Height * Height + Radius * Radius);
+			const float Angle = FMath::Acos(Height / ConeSide);
+			FPrimitiveDrawInterface* PDI = RenderAPI->GetPrimitiveDrawInterface();
+			DrawCone(PDI, LocalToWorldTransform.ToMatrixWithScale(), Angle, Angle, NumSides, false, FColor::White, UseMaterial->GetRenderProxy(), SDPG_Foreground);
+		}
+	}
+	CacheRenderState(LocalToWorldTransform, bVisibleViewDependent);
+}
+
+
+
+FInputRayHit UGizmoElementCone::LineTrace(const FVector RayOrigin, const FVector RayDirection)
+{
+	if (IsHittableInView())
+	{
+		// @todo - modify ray-cone intersection to work with updated properties
+
+#if 0
+		bool bIntersects = false;
+		double RayParam = 0.0;
+
+		const FMatrix LocalToWorldMatrix = LocalToWorldTransform.ToMatrixNoScale();
+
+		const FVector Origin = LocalToWorldMatrix.TransformPosition(FVector::ZeroVector);
+		const FVector ConeDirection = (bWorld) ? Direction : FVector{ LocalToWorldMatrix.TransformVector(Direction) };
+		const double ConeHeight = Height * DynamicPixelToWorldScale;
+		const double ConeOffset = Offset * DynamicPixelToWorldScale;
+		const FVector ConeOrigin = Origin + ConeOffset * ConeDirection;
+
+		GizmoMath::RayConeIntersection(
+			ConeOrigin, 
+			ConeDirection, 
+			FMath::Cos(Angle), 
+			ConeHeight,
+			RayOrigin, RayDirection,
+			bIntersects, RayParam);
+
+		if (bIntersects)
+		{
+			return FInputRayHit(RayParam);
+		}
+#endif
+	}
+
+	return FInputRayHit();
+}
+
+FBoxSphereBounds UGizmoElementCone::CalcBounds(const FTransform& LocalToWorld) const
+{
+	// @todo - implement box-sphere bounds calculation
+	return FBoxSphereBounds();
+}
+
+void UGizmoElementCone::SetBase(const FVector& InBase)
+{
+	Base = InBase;
+}
+
+FVector UGizmoElementCone::GetBase() const
+{
+	return Base;
+}
+
+void UGizmoElementCone::SetDirection(const FVector& InDirection)
+{
+	Direction = InDirection;
+	Direction.Normalize();
+}
+
+FVector UGizmoElementCone::GetDirection() const
+{
+	return Direction;
+}
+
+void UGizmoElementCone::SetHeight(float InHeight)
+{
+	Height = InHeight;
+}
+
+float UGizmoElementCone::GetHeight() const
+{
+	return Height;
+}
+
+void UGizmoElementCone::SetRadius(float InRadius)
+{
+	Radius = InRadius;
+}
+
+float UGizmoElementCone::GetRadius() const
+{
+	return Radius;
+}
+
+void UGizmoElementCone::SetNumSides(int32 InNumSides)
+{
+	NumSides = InNumSides;
+}
+
+int32 UGizmoElementCone::GetNumSides() const
+{
+	return NumSides;
+}
+
diff --git a/Engine/Source/Runtime/InteractiveToolsFramework/Private/BaseGizmos/GizmoElementCylinder.cpp b/Engine/Source/Runtime/InteractiveToolsFramework/Private/BaseGizmos/GizmoElementCylinder.cpp
new file mode 100644
--- /dev/null
+++ b/Engine/Source/Runtime/InteractiveToolsFramework/Private/BaseGizmos/GizmoElementCylinder.cpp
@@ -0,0 +1,141 @@
+// Copyright Epic Games, Inc. All Rights Reserved.
+ 
+#include "BaseGizmos/GizmoElementCylinder.h"
+#include "BaseGizmos/GizmoRenderingUtil.h"
+#include "BaseGizmos/GizmoMath.h"
+#include "InputState.h"
+#include "Materials/MaterialInterface.h"
+#include "SceneManagement.h"
+
+void UGizmoElementCylinder::Render(IToolsContextRenderAPI* RenderAPI, const FRenderTraversalState& RenderState)
+{
+	if (!IsVisible())
+	{
+		return;
+	}
+
+	check(RenderAPI);
+	const FSceneView* View = RenderAPI->GetSceneView();
+
+	FTransform LocalToWorldTransform = RenderState.LocalToWorldTransform;
+
+	bool bVisibleViewDependent = GetViewDependentVisibility(View, LocalToWorldTransform, Base);
+
+	if (bVisibleViewDependent)
+	{
+		const UMaterialInterface* UseMaterial = GetCurrentMaterial(RenderState, ElementInteractionState);
+
+		if (UseMaterial)
+		{
+			FQuat AlignRot;
+			const bool bHasAlignRot = GetViewAlignRot(View, LocalToWorldTransform, Base, AlignRot);
+			const FVector AdjustedDir = bHasAlignRot ? AlignRot.RotateVector(Direction) : Direction;
+
+			const FQuat Rotation = FRotationMatrix::MakeFromZ(AdjustedDir).ToQuat();
+			const float HalfHeight = Height * 0.5f;
+			const FVector Origin = Base + AdjustedDir * HalfHeight;
+
+			LocalToWorldTransform = FTransform(Rotation, Origin) * LocalToWorldTransform;
+			FPrimitiveDrawInterface* PDI = RenderAPI->GetPrimitiveDrawInterface();
+			DrawCylinder(PDI, LocalToWorldTransform.ToMatrixWithScale(), FVector::ZeroVector, FVector(1, 0, 0), FVector(0, 1, 0), FVector(0, 0, 1), Radius, HalfHeight, NumSides, UseMaterial->GetRenderProxy(), SDPG_Foreground);
+
+		}
+	}
+	CacheRenderState(LocalToWorldTransform, bVisibleViewDependent);
+}
+
+
+
+FInputRayHit UGizmoElementCylinder::LineTrace(const FVector RayOrigin, const FVector RayDirectionm)
+{
+	if (IsHittableInView())
+	{
+		// @todo - modify ray-cylinder intersection to work with updated properties
+#if 0
+		bool bIntersects = false;
+		double RayParam = 0.0;
+
+		const FTransform LocalToWorldTransform = FTransform(Base) * InLocalToWorldTransform;
+		const double WorldConeHeight = Height * InLocalToWorldTransform.GetScale3D().X;
+		const FVector ConeOrigin = Base + Height * Direction;
+		const FVector WorldConeOrigin = LocalToWorldTransform.TransformPosition(ConeOrigin);
+		const FVector WorldConeDirection = LocalToWorldTransform.TransformVector(Direction);
+		const double ConeSide = FMath::Sqrt(Height * Height + Radius * Radius);
+		const double CosAngle = Radius / ConeSide;
+
+		GizmoMath::RayConeIntersection(
+			WorldConeOrigin,
+			WorldConeDirection,
+			CosAngle,
+			WorldConeHeight,
+			RayOrigin, RayDirection,
+			bIntersects, RayParam);
+
+		if (bIntersects)
+		{
+			FInputRayHit Hit;
+			Hit.HitObject = this;
+			return Hit;
+		}
+#endif
+	}
+
+	return FInputRayHit();
+}
+
+FBoxSphereBounds UGizmoElementCylinder::CalcBounds(const FTransform& LocalToWorld) const
+{
+	// @todo - implement box-sphere bounds calculation
+	return FBoxSphereBounds();
+}
+
+void UGizmoElementCylinder::SetBase(const FVector& InBase)
+{
+	Base = InBase;
+}
+
+FVector UGizmoElementCylinder::GetBase() const
+{
+	return Base;
+}
+
+void UGizmoElementCylinder::SetDirection(const FVector& InDirection)
+{
+	Direction = InDirection;
+	Direction.Normalize();
+}
+
+FVector UGizmoElementCylinder::GetDirection() const
+{
+	return Direction;
+}
+
+void UGizmoElementCylinder::SetHeight(float InHeight)
+{
+	Height = InHeight;
+}
+
+float UGizmoElementCylinder::GetHeight() const
+{
+	return Height;
+}
+
+void UGizmoElementCylinder::SetRadius(float InRadius)
+{
+	Radius = InRadius;
+}
+
+float UGizmoElementCylinder::GetRadius() const
+{
+	return Radius;
+}
+
+void UGizmoElementCylinder::SetNumSides(int32 InNumSides)
+{
+	NumSides = InNumSides;
+}
+
+int32 UGizmoElementCylinder::GetNumSides() const
+{
+	return NumSides;
+}
diff --git a/Engine/Source/Runtime/InteractiveToolsFramework/Private/BaseGizmos/GizmoElementGroup.cpp b/Engine/Source/Runtime/InteractiveToolsFramework/Private/BaseGizmos/GizmoElementGroup.cpp
new file mode 100644
--- /dev/null
+++ b/Engine/Source/Runtime/InteractiveToolsFramework/Private/BaseGizmos/GizmoElementGroup.cpp
@@ -0,0 +1,128 @@
+// Copyright Epic Games, Inc. All Rights Reserved.
+
+#include "BaseGizmos/GizmoElementGroup.h"
+#include "BaseGizmos/GizmoRenderingUtil.h"
+
+
+void UGizmoElementGroup::Render(IToolsContextRenderAPI* RenderAPI, const FRenderTraversalState& RenderState)
+{
+	if (!IsVisible())
+	{
+		return;
+	}
+
+	const FSceneView* View = RenderAPI->GetSceneView();
+
+	bool bVisibleViewDependent = GetViewDependentVisibility(View, RenderState.LocalToWorldTransform, FVector::ZeroVector);
+
+	FRenderTraversalState RenderStateCopy = RenderState;
+
+	if (bVisibleViewDependent)
+	{
+		// Compute constant scale, if applicable
+		float Scale = RenderStateCopy.LocalToWorldTransform.GetScale3D().X;
+		if (bConstantScale)
+		{
+			const float ConstantScaleFactor = GizmoRenderingUtil::CalculateLocalPixelToWorldScale(RenderAPI->GetSceneView(), RenderStateCopy.LocalToWorldTransform.GetTranslation());
+			Scale *= ConstantScaleFactor;
+		}
+		RenderStateCopy.LocalToWorldTransform.SetScale3D(FVector(Scale, Scale, Scale));
+
+		// Compute view alignment, if applicable
+		FQuat AlignRot;
+		if (GetViewAlignRot(View, RenderStateCopy.LocalToWorldTransform, FVector::ZeroVector, AlignRot))
+		{
+			RenderStateCopy.LocalToWorldTransform = FTransform(AlignRot) * RenderState.LocalToWorldTransform;
+		}
+
+		UpdateRenderTraversalState(RenderStateCopy);
+
+		// Continue render even if not visible so all transforms will be cached 
+		// for subsequent line tracing.
+		for (UGizmoElementBase* Element : Elements)
+		{
+			if (Element)
+			{
+				Element->Render(RenderAPI, RenderStateCopy);
+			}
+		}
+	}
+
+	CacheRenderState(RenderStateCopy.LocalToWorldTransform, bVisibleViewDependent);
+}
+
+FInputRayHit UGizmoElementGroup::LineTrace(const FVector Start, const FVector Direction)
+{
+	if (IsHittable())
+	{
+		for (UGizmoElementBase* Element : Elements)
+		{
+			if (Element)
+			{
+				return Element->LineTrace(Start, Direction);
+			}
+		}
+	}
+	return FInputRayHit();
+}
+
+FBoxSphereBounds UGizmoElementGroup::CalcBounds(const FTransform& LocalToWorld) const
+{
+	/*  @todo - accumulate box sphere bounds for all elements within the group
+
+	if (bEnabled)
+	{
+		for (UGizmoElementBase* Element : Elements)
+		{
+			if (Element)
+			{
+				CachedBoxSphereBounds = CachedBoxSphereBounds + Element->CalcBounds(LocalToWorld);
+			}
+		}
+	}
+	*/
+	return FBoxSphereBounds();
+}
+
+void UGizmoElementGroup::Add(UGizmoElementBase* InElement)
+{
+	if (!Elements.Contains(InElement))
+	{
+		Elements.Add(InElement);
+	}
+}
+
+void UGizmoElementGroup::Remove(UGizmoElementBase* InElement)
+{
+	if (InElement)
+	{
+		int32 Index;
+		if (Elements.Find(InElement, Index))
+		{
+			Elements.RemoveAtSwap(Index);
+		}
+	}
+}
+
+void UGizmoElementGroup::ResetCachedRenderState()
+{
+	Super::ResetCachedRenderState();
+
+	for (UGizmoElementBase* Element : Elements)
+	{
+		if (Element)
+		{
+			Element->ResetCachedRenderState();
+		}
+	}
+}
+
+void UGizmoElementGroup::SetConstantScale(bool bInConstantScale)
+{
+	bConstantScale = bInConstantScale;
+}
+
+bool UGizmoElementGroup::GetConstantScale() const
+{
+	return bConstantScale;
+}
\ No newline at end of file
diff --git a/Engine/Source/Runtime/InteractiveToolsFramework/Private/BaseGizmos/GizmoElementHitTargets.cpp b/Engine/Source/Runtime/InteractiveToolsFramework/Private/BaseGizmos/GizmoElementHitTargets.cpp
new file mode 100644
--- /dev/null
+++ b/Engine/Source/Runtime/InteractiveToolsFramework/Private/BaseGizmos/GizmoElementHitTargets.cpp
@@ -0,0 +1,41 @@
+// Copyright Epic Games, Inc. All Rights Reserved.
+
+#include "BaseGizmos/GizmoElementHitTargets.h"
+#include "BaseGizmos/GizmoElementBase.h"
+#include "Engine/EngineTypes.h"    // FHitResult
+
+
+FInputRayHit UGizmoElementHitTarget::IsHit(const FInputDeviceRay& ClickPos) const
+{
+	if (GizmoElement && (!Condition || Condition(ClickPos)))
+	{
+		return GizmoElement->LineTrace(ClickPos.WorldRay.Origin, ClickPos.WorldRay.Direction);
+	}
+	return FInputRayHit();
+}
+
+void UGizmoElementHitTarget::UpdateHoverState(bool bHovering)
+{
+	if (GizmoElement)
+	{
+		GizmoElement->SetElementInteractionState(EGizmoElementInteractionState::Hovering);
+	}
+}
+
+void UGizmoElementHitTarget::UpdateInteractingState(bool bInteracting)
+{
+	if (GizmoElement)
+	{
+		GizmoElement->SetElementInteractionState(EGizmoElementInteractionState::Interacting);
+	}
+}
+
+UGizmoElementHitTarget* UGizmoElementHitTarget::Construct(UGizmoElementBase* InGizmoElement, UObject* Outer)
+{
+	UGizmoElementHitTarget* NewHitTarget = NewObject<UGizmoElementHitTarget>(Outer);
+	NewHitTarget->GizmoElement = InGizmoElement;
+	return NewHitTarget;
+}
+
+
+
diff --git a/Engine/Source/Runtime/InteractiveToolsFramework/Private/BaseGizmos/GizmoElementRectangle.cpp b/Engine/Source/Runtime/InteractiveToolsFramework/Private/BaseGizmos/GizmoElementRectangle.cpp
new file mode 100644
--- /dev/null
+++ b/Engine/Source/Runtime/InteractiveToolsFramework/Private/BaseGizmos/GizmoElementRectangle.cpp
@@ -0,0 +1,196 @@
+// Copyright Epic Games, Inc. All Rights Reserved.
+ 
+#include "BaseGizmos/GizmoElementRectangle.h"
+#include "BaseGizmos/GizmoRenderingUtil.h"
+#include "BaseGizmos/GizmoMath.h"
+#include "InputState.h"
+#include "Materials/MaterialInterface.h"
+#include "SceneManagement.h"
+
+void UGizmoElementRectangle::Render(IToolsContextRenderAPI* RenderAPI, const FRenderTraversalState& RenderState)
+{
+	if (!IsVisible())
+	{
+		return;
+	}
+
+	check(RenderAPI);
+
+	const FSceneView* View = RenderAPI->GetSceneView();
+
+	FTransform LocalToWorldTransform = RenderState.LocalToWorldTransform;
+
+	bool bVisibleViewDependent = GetViewDependentVisibility(View, LocalToWorldTransform, Center);
+
+	if (bVisibleViewDependent)
+	{
+		FVector Axis0, Axis1;
+		if (bScreenSpace)
+		{
+			Axis0 = View->GetViewUp();
+			Axis1 = View->GetViewRight();
+		}
+		else
+		{
+			FQuat AlignRot;
+			if (GetViewAlignRot(View, LocalToWorldTransform, Center, AlignRot))
+			{
+				Axis0 = AlignRot.RotateVector(UpDirection);
+				Axis1 = AlignRot.RotateVector(SideDirection);
+			}
+			else
+			{
+				Axis0 = UpDirection;
+				Axis1 = SideDirection;
+			}
+
+			Axis0 = LocalToWorldTransform.TransformVectorNoScale(Axis0);
+			Axis1 = LocalToWorldTransform.TransformVectorNoScale(Axis1);
+		}
+
+		FVector WorldCenter = LocalToWorldTransform.TransformPosition(Center);
+		float WorldWidth = Width * LocalToWorldTransform.GetScale3D().X;
+		float WorldHeight = Height * LocalToWorldTransform.GetScale3D().X;
+
+		FPrimitiveDrawInterface* PDI = RenderAPI->GetPrimitiveDrawInterface();
+
+		if (bDrawMesh)
+		{
+			if (const UMaterialInterface* UseMaterial = GetCurrentMaterial(RenderState, ElementInteractionState))
+			{
+
+				DrawRectangleMesh(PDI, WorldCenter, Axis0, Axis1, VertexColor, WorldWidth, WorldHeight, UseMaterial->GetRenderProxy(), SDPG_Foreground);
+			}
+		}
+		if (bDrawLine)
+		{
+			DrawRectangle(PDI, WorldCenter, Axis0, Axis1, LineColor, WorldWidth, WorldHeight, SDPG_Foreground, 0.0f);
+		}
+	}
+	CacheRenderState(LocalToWorldTransform, bVisibleViewDependent);
+}
+
+FInputRayHit UGizmoElementRectangle::LineTrace(const FVector RayOrigin, const FVector RayDirection)
+{
+	if (IsHittableInView())
+	{
+		// @todo - implement ray-circle intersection
+	}
+	return FInputRayHit();
+}
+
+FBoxSphereBounds UGizmoElementRectangle::CalcBounds(const FTransform& LocalToWorld) const
+{
+	// @todo - implement box-sphere bounds calculation
+	return FBoxSphereBounds();
+}
+
+void UGizmoElementRectangle::SetCenter(FVector InCenter)
+{
+	Center = InCenter;
+}
+
+FVector UGizmoElementRectangle::GetCenter() const
+{
+	return Center;
+}
+
+void UGizmoElementRectangle::SetWidth(float InWidth)
+{
+	Width = InWidth;
+	
+}
+
+float UGizmoElementRectangle::GetWidth() const
+{
+	return Width;
+}
+
+void UGizmoElementRectangle::SetHeight(float InHeight)
+{
+	Height = InHeight;
+}
+
+float UGizmoElementRectangle::GetHeight() const
+{
+	return Height;
+}
+
+void UGizmoElementRectangle::SetUpDirection(const FVector& InUpDirection)
+{
+	UpDirection = InUpDirection.GetSafeNormal();
+}
+
+FVector UGizmoElementRectangle::GetUpDirection() const
+{
+	return UpDirection;
+}
+
+void UGizmoElementRectangle::SetSideDirection(const FVector& InSideDirection)
+{
+	SideDirection = InSideDirection.GetSafeNormal();
+}
+
+FVector UGizmoElementRectangle::GetSideDirection() const
+{
+	return SideDirection;
+}
+
+void UGizmoElementRectangle::SetScreenSpace(bool InScreenSpace)
+{
+	bScreenSpace = InScreenSpace;
+}
+
+bool UGizmoElementRectangle::GetScreenSpace()
+{
+	return bScreenSpace;
+}
+
+void UGizmoElementRectangle::SetLineColor(const FColor& InLineColor)
+{
+	LineColor = InLineColor;
+}
+
+FColor UGizmoElementRectangle::GetLineColor()
+{
+	return LineColor;
+}
+
+void UGizmoElementRectangle::SetDrawMesh(bool InDrawMesh)
+{
+	bDrawMesh = InDrawMesh;
+}
+bool UGizmoElementRectangle::GetDrawMesh() const
+{
+	return bDrawMesh;
+}
+
+void UGizmoElementRectangle::SetDrawLine(bool InDrawLine)
+{
+	bDrawLine = InDrawLine;
+}
+
+bool UGizmoElementRectangle::GetDrawLine() const
+{
+	return bDrawLine;
+}
+
+void UGizmoElementRectangle::SetHitMesh(bool InHitMesh)
+{
+	bHitMesh = InHitMesh;
+}
+
+bool UGizmoElementRectangle::GetHitMesh() const
+{
+	return bHitMesh;
+}
+
+void UGizmoElementRectangle::SetHitLine(bool InHitLine)
+{
+	bHitLine = InHitLine;
+}
+
+bool UGizmoElementRectangle::GetHitLine() const
+{
+	return bHitLine;
+}
diff --git a/Engine/Source/Runtime/InteractiveToolsFramework/Private/BaseGizmos/GizmoElementTorus.cpp b/Engine/Source/Runtime/InteractiveToolsFramework/Private/BaseGizmos/GizmoElementTorus.cpp
new file mode 100644
--- /dev/null
+++ b/Engine/Source/Runtime/InteractiveToolsFramework/Private/BaseGizmos/GizmoElementTorus.cpp
@@ -0,0 +1,180 @@
+// Copyright Epic Games, Inc. All Rights Reserved.
+ 
+#include "BaseGizmos/GizmoElementTorus.h"
+#include "BaseGizmos/GizmoRenderingUtil.h"
+#include "BaseGizmos/GizmoMath.h"
+#include "InputState.h"
+#include "Materials/MaterialInterface.h"
+#include "SceneManagement.h"
+
+void UGizmoElementTorus::Render(IToolsContextRenderAPI* RenderAPI, const FRenderTraversalState& RenderState)
+{
+	if (!IsVisible())
+	{
+		return;
+	}
+
+	check(RenderAPI);
+	const FSceneView* View = RenderAPI->GetSceneView();
+
+	FTransform LocalToWorldTransform = RenderState.LocalToWorldTransform;
+
+	bool bVisibleViewDependent = GetViewDependentVisibility(View, LocalToWorldTransform, Center);
+
+	if (bVisibleViewDependent)
+	{
+		const UMaterialInterface* UseMaterial = GetCurrentMaterial(RenderState, ElementInteractionState);
+			
+		if (UseMaterial)
+		{
+			FPrimitiveDrawInterface* PDI = RenderAPI->GetPrimitiveDrawInterface();
+
+			FVector TorusSideAxis = Normal ^ BeginAxis;
+			TorusSideAxis.Normalize();
+
+			FTransform TorusToLocal = FTransform::Identity;
+			TorusToLocal.SetTranslation(Center);
+
+			FQuat AlignRot;
+			if (GetViewAlignRot(View, LocalToWorldTransform, Center, AlignRot))
+			{
+				TorusToLocal.SetRotation(AlignRot);
+			}
+
+			FVector WorldAlignAxis = LocalToWorldTransform.TransformVector(ViewAlignAxis);
+
+			LocalToWorldTransform = TorusToLocal * LocalToWorldTransform;
+
+			DrawTorus(PDI, LocalToWorldTransform.ToMatrixWithScale(), BeginAxis, TorusSideAxis, OuterRadius, InnerRadius, OuterSegments, InnerSlices,  UseMaterial->GetRenderProxy(), SDPG_Foreground, bPartial, Angle, bEndCaps);
+		}
+	}
+
+	CacheRenderState(LocalToWorldTransform, bVisibleViewDependent);
+}
+
+
+
+FInputRayHit UGizmoElementTorus::LineTrace(const FVector RayOrigin, const FVector RayDirection)
+{
+	// Line trace is not supported for torus.
+	return FInputRayHit();
+}
+
+FBoxSphereBounds UGizmoElementTorus::CalcBounds(const FTransform& LocalToWorld) const
+{
+	// Box sphere bounds is not supported for torus.
+	return FBoxSphereBounds();
+}
+
+void UGizmoElementTorus::SetCenter(const FVector& InCenter)
+{
+	Center = InCenter;
+}
+
+FVector UGizmoElementTorus::GetCenter() const
+{
+	return Center;
+}
+
+void UGizmoElementTorus::SetNormal(const FVector& InNormal)
+{
+	Normal = InNormal;
+	Normal.Normalize();
+}
+
+FVector UGizmoElementTorus::GetNormal() const
+{
+	return Normal;
+}
+
+void UGizmoElementTorus::SetBeginAxis(const FVector& InBeginAxis)
+{
+	BeginAxis = InBeginAxis;
+	BeginAxis.Normalize();
+}
+
+FVector UGizmoElementTorus::GetBeginAxis() const
+{
+	return BeginAxis;
+}
+
+void UGizmoElementTorus::SetOuterRadius(float InOuterRadius)
+{
+	OuterRadius = InOuterRadius;
+}
+
+float UGizmoElementTorus::GetOuterRadius() const
+{
+	return OuterRadius;
+}
+
+void UGizmoElementTorus::SetInnerRadius(float InInnerRadius)
+{
+	InnerRadius = InInnerRadius;
+}
+
+float UGizmoElementTorus::GetInnerRadius() const
+{
+	return InnerRadius;
+}
+
+void UGizmoElementTorus::SetOuterSegments(int32 InOuterSegments)
+{
+	OuterSegments = InOuterSegments;
+}
+
+int32 UGizmoElementTorus::GetOuterSegments() const
+{
+	return OuterSegments;
+}
+
+void UGizmoElementTorus::SetInnerSlices(int32 InInnerSlices)
+{
+	InnerSlices = InInnerSlices;
+}
+
+int32 UGizmoElementTorus::GetInnerSlices() const
+{
+	return InnerSlices;
+}
+
+void UGizmoElementTorus::SetPartial(bool InPartial)
+{
+	bPartial = InPartial;
+}
+
+bool UGizmoElementTorus::GetPartial() const
+{
+	return bPartial;
+}
+
+void UGizmoElementTorus::SetScreenAlignPartial(bool InScreenAlignPartial)
+{
+	bScreenAlignPartial = InScreenAlignPartial;
+}
+
+bool UGizmoElementTorus::GetScreenAlignPartial() const
+{
+	return bPartial;
+}
+
+void UGizmoElementTorus::SetAngle(float InAngle)
+{
+	Angle = InAngle;
+}
+
+float UGizmoElementTorus::GetAngle() const
+{
+	return Angle;
+}
+
+void UGizmoElementTorus::SetEndCaps(bool InEndCaps)
+{
+	bEndCaps = InEndCaps;
+}
+
+bool UGizmoElementTorus::GetEndCaps() const
+{
+	return bEndCaps;
+}
+
diff --git a/Engine/Source/Runtime/InteractiveToolsFramework/Public/BaseGizmos/GizmoElementArrow.h b/Engine/Source/Runtime/InteractiveToolsFramework/Public/BaseGizmos/GizmoElementArrow.h
new file mode 100644
--- /dev/null
+++ b/Engine/Source/Runtime/InteractiveToolsFramework/Public/BaseGizmos/GizmoElementArrow.h
@@ -0,0 +1,133 @@
+// Copyright Epic Games, Inc. All Rights Reserved.
+
+#pragma once
+
+#include "CoreMinimal.h"
+#include "BaseGizmos/GizmoElementBase.h"
+#include "InputState.h"
+#include "UObject/ObjectMacros.h"
+#include "GizmoElementArrow.generated.h"
+
+class UGizmoElementBox;
+class UGizmoElementCone;
+class UGizmoElementCylinder;
+
+/**
+ * Simple object intended to be used as part of 3D Gizmos.
+ * Draws a 3D arrow based on parameters.
+ */
+
+UENUM()
+enum class EGizmoElementArrowHeadType
+{
+	Cone,
+	Cube
+};
+
+
+UCLASS(Transient)
+class INTERACTIVETOOLSFRAMEWORK_API UGizmoElementArrow : public UGizmoElementBase
+{
+	GENERATED_BODY()
+
+public:
+	UGizmoElementArrow();
+
+	//~ Begin UGizmoElementBase Interface.
+	virtual void Render(IToolsContextRenderAPI* RenderAPI, const FRenderTraversalState& RenderState) override;
+	virtual FInputRayHit LineTrace(const FVector Start, const FVector Direction) override;
+	virtual FBoxSphereBounds CalcBounds(const FTransform& LocalToWorld) const override;
+	//~ End UGizmoElementBase Interface.
+
+	// Location of base of arrow cylinder
+	virtual void SetBase(const FVector& InBase);
+	virtual FVector GetBase() const;
+
+	// Arrow direction.
+	virtual void SetDirection(const FVector& InDirection);
+	virtual FVector GetDirection() const;
+
+	// Arrow side direction for box head.
+	virtual void SetSideDirection(const FVector& InSideDirection);
+	virtual FVector GetSideDirection() const;
+
+	// Arrow body length.
+	virtual void SetBodyLength(float InBodyLength);
+	virtual float GetBodyLength() const;
+
+	// Arrow body cylinder radius.
+	virtual void SetBodyRadius(float InBodyRadius);
+	virtual float GetBodyRadius() const;
+
+	// Arrow head length, used for both cone and cube head.
+	virtual void SetHeadLength(float InHeadLength);
+	virtual float GetHeadLength() const;
+
+	// Arrow head radius, if cone.
+	virtual void SetHeadRadius(float InHeadRadius);
+	virtual float GetHeadRadius() const;
+
+	// Number of sides for cylinder and cone, if relevant.
+	virtual void SetNumSides(int32 InNumSides);
+	virtual int32 GetNumSides() const;
+
+	// Head type cone or cube.
+	virtual void SetHeadType(EGizmoElementArrowHeadType InHeadType);
+	virtual EGizmoElementArrowHeadType GetHeadType() const;
+
+protected:
+
+	// Update arrow cylinder based on parameters
+	virtual void UpdateArrowBody();
+
+	// Update arrow cone or box based on parameters
+	virtual void UpdateArrowHead();
+
+	// Arrow cylinder body
+	UPROPERTY()
+	TObjectPtr<UGizmoElementCylinder> CylinderElement;
+
+	// Arrow cone head
+	UPROPERTY()
+	TObjectPtr<UGizmoElementCone> ConeElement;
+
+	// Arrow box head
+	UPROPERTY()
+	TObjectPtr<UGizmoElementBox> BoxElement;
+
+	// Location of center of arrow base circle.
+	UPROPERTY()
+	FVector Base = FVector::ZeroVector;
+
+	// Direction of arrow axis
+	UPROPERTY()
+	FVector Direction = FVector(0.0f, 0.0f, 1.0f);
+
+	// Side direction for box head
+	UPROPERTY()
+	FVector SideDirection = FVector(0.0f, 1.0f, 0.0f);
+
+	// Arrow body length
+	UPROPERTY()
+	float BodyLength = 1.5f;
+
+	// Radius of arrow cylinder
+	UPROPERTY()
+	float BodyRadius = 0.5f;
+
+	// Length of head, cone or box
+	UPROPERTY()
+	float HeadLength = 0.5f;
+
+	// Radius of head cone
+	UPROPERTY()
+	float HeadRadius = 0.75f;
+
+	// Number of sides for tessellating cone and cylinder
+	UPROPERTY()
+	int32 NumSides = 32;
+
+	// Head type
+	UPROPERTY()
+	EGizmoElementArrowHeadType HeadType = EGizmoElementArrowHeadType::Cone;
+};
diff --git a/Engine/Source/Runtime/InteractiveToolsFramework/Public/BaseGizmos/GizmoElementBase.h b/Engine/Source/Runtime/InteractiveToolsFramework/Public/BaseGizmos/GizmoElementBase.h
new file mode 100644
--- /dev/null
+++ b/Engine/Source/Runtime/InteractiveToolsFramework/Public/BaseGizmos/GizmoElementBase.h
@@ -0,0 +1,311 @@
+// Copyright Epic Games, Inc. All Rights Reserved.
+
+#pragma once
+
+#include "CoreMinimal.h"
+#include "InputState.h"
+#include "ToolContextInterfaces.h"
+#include "UObject/GCObject.h"
+#include "GizmoElementBase.generated.h"
+
+class UMaterialInterface;
+
+// bitmask indicating whether object is visible, hittable or both
+UENUM()
+enum class EGizmoElementState : uint8
+{
+	None		= 0x00,
+	Visible		= 1<<1,
+	Hittable	= 1<<2,
+	VisibleAndHittable = Visible | Hittable
+};
+
+UENUM()
+enum class EGizmoElementInteractionState
+{
+	None,
+	Hovering,
+	Interacting
+};
+
+UENUM()
+enum class EGizmoElementViewDependentType
+{
+	None,
+	Axis, // Cull object when angle between axis and view direction is within a given tolerance
+	Plane // Cull object when angle between plane normal and view direction is perpendicular within a given tolerance
+};
+
+UENUM()
+enum class EGizmoElementViewAlignType
+{
+	None,
+	PointEye, // Align object forward axis to view direction, align object up axis to view up
+	PointOnly, // Align object forward axis to view direction only, useful for symmetrical objects such as a circle
+	Axial // Rotate object around up axis, minimizing angle between forward axis and view direction
+};
+
+/**
+ * Base class for simple objects intended to be used as part of 3D Gizmos.
+ * Contains common properties and utility functions.
+ * This class does nothing by itself, use subclasses like UGizmoElementCylinder
+ */
+UCLASS(Transient)
+class INTERACTIVETOOLSFRAMEWORK_API UGizmoElementBase : public UObject
+{
+	GENERATED_BODY()
+public:
+
+	static constexpr float DefaultViewDependentAngleTol = 0.052f;			// ~3 degrees
+	static constexpr float DefaultViewAlignAxialMaxCosAngleTol = 0.998f;	// Cos(DefaultViewDependentAngleTol)
+	static constexpr float DefaultViewAlignPlanarMinCosAngleTol = 0.052f;	// Cos(HALF_PI - DefaultViewDependentAngleTol)
+
+	static constexpr float DefaultViewAlignAngleTol = 0.052f;				// ~3 degrees
+	static constexpr float DefaultViewAlignMaxCosAngleTol = 0.998f;			// Cos(DefaultViewAlignAngleTol)
+
+	// Helper struct used to store traversal state during render
+	struct FRenderTraversalState
+	{
+		// LocalToWorld transform 
+		// Note: non-uniform scale is not supported and the X scale element will be used for uniform scaling.
+		FTransform LocalToWorldTransform;
+
+		// LocalToWorld translation 
+		FVector LocalToWorldTranslation;
+
+		// LocalToWorld rotation 
+		FQuat LocalToWorldRotation;
+
+		// LocalToWorld uniform scale, non-uniform scale is not supported
+		double LocalToWorldScale;
+
+		// Material, element's material should override parent's material
+		TWeakObjectPtr<UMaterialInterface> Material = nullptr;
+
+		// Hover material, element's hover material should override parent's hover material
+		TWeakObjectPtr<UMaterialInterface> HoverMaterial = nullptr;
+
+		// Interact material, element's interact material should override parent's interact material
+		TWeakObjectPtr<UMaterialInterface> InteractMaterial = nullptr;
+	};
+
+public:
+
+	// Render enabled visible element.
+	virtual void Render(IToolsContextRenderAPI* RenderAPI, const FRenderTraversalState& RenderState) PURE_VIRTUAL(UGizmoElementBase::Render);
+
+	// Line trace enabled hittable element.
+	virtual FInputRayHit LineTrace(const FVector Start, const FVector Direction) PURE_VIRTUAL(UGizmoElementBase::LineTrace, return FInputRayHit(););
+
+	// Calcute box sphere bounds for use when hit testing.
+	virtual FBoxSphereBounds CalcBounds(const FTransform & LocalToWorld) const PURE_VIRTUAL(UGizmoElementBase::CalcBounds, return FBoxSphereBounds(););
+
+	// Returns whether object is visible based on view-dependent visibility settings. 
+	virtual bool GetViewDependentVisibility(const FSceneView* View, const FTransform& InLocalToWorldTransform, const FVector& InLocalCenter) const;
+
+	// Returns true when view alignment is enabled. OutAlignRot is the rotation in local space which will align the object to the view base
+	// on view-dependent alignment settings. So it should be prepended to the local-to-world transform.
+	virtual bool GetViewAlignRot(const FSceneView* View, const FTransform& InLocalToWorldTransform, const FVector& InLocalCenter, FQuat& OutAlignRot) const;
+
+	// Returns the current material for given object state. The materials in the object will override those in the RenderState.
+	virtual const UMaterialInterface* GetCurrentMaterial(const FRenderTraversalState& RenderState, EGizmoElementInteractionState InState) const;
+
+	// Reset the cached render state.
+	virtual void ResetCachedRenderState();
+
+	// Whether this object is visible.
+	virtual bool IsVisible() const;
+
+	// Whether this object is hittable.
+	virtual bool IsHittable() const;
+
+	// Whether this object is hittable in the most recently cached view.
+	virtual bool IsHittableInView() const;
+
+	// Render and LineTrace should only occur when bEnabled is true.
+	virtual void SetEnabled(bool InEnabled);
+	virtual bool GetEnabled() const;
+
+	// Object type bitmask indicating whether this object is visible or hittable or both
+	virtual void SetElementState(EGizmoElementState InElementState);
+	virtual EGizmoElementState GetElementState() const;
+
+	// Object interaction state - None, Hovering or Interacting
+	virtual void SetElementInteractionState(EGizmoElementInteractionState InInteractionState);
+	virtual EGizmoElementInteractionState GetElementInteractionState() const;
+
+	// View-dependent type - None, Axis or Plane. 
+	virtual void SetViewDependentType(EGizmoElementViewDependentType ViewDependentType);
+	virtual EGizmoElementViewDependentType GetViewDependentType() const;
+
+	// View-dependent angle tolerance in radians 
+	//   For Axis, object is culled when angle between view dependent axis and view direction is less than tolerance angle.
+	//   For Planar, cos of angle between view dependent axis (plane normal) and view direction. 
+	// When the view direction is within this tolerance from the plane or axis, this object will be culled.
+	virtual void SetViewDependentAngleTol(float InMaxAngleTol);
+	virtual float GetViewDependentAngleTol() const;
+
+	// View-dependent axis or plane normal, based on the view-dependent type.
+	virtual void SetViewDependentAxis(FVector InAxis);
+	virtual FVector GetViewDependentAxis() const;
+
+	// View align type: None, PointEye, PointOnly, or Axial.
+	virtual void SetViewAlignType(EGizmoElementViewAlignType InViewAlignType);
+	virtual EGizmoElementViewAlignType GetViewAlignType() const;
+
+	// View align axis. 
+	// PointEye and PointOnly rotate this axis to align with view up.
+	// Axial rotates about this axis.
+	virtual void SetViewAlignAxis(FVector InAxis);
+	virtual FVector GetViewAlignAxis() const;
+
+	// View align normal.
+	// PointEye rotates the normal to align with view direction.
+	// Axial rotates the normal around the axis to align as closely as possible with the view direction.
+	virtual void SetViewAlignNormal(FVector InAxis);
+	virtual FVector GetViewAlignNormal() const;
+
+	// View-align angle tolerance in radians.
+	// Viewer alignment will not occur when the viewing angle is within this angle of view align axis.
+	virtual void SetViewAlignAxialAngleTol(float InMaxAngleTol);
+	virtual float GetViewAlignAxialAngleTol() const;
+
+	// Default material.
+	virtual void SetMaterial(UMaterialInterface* Material);
+	UMaterialInterface* GetMaterial() const;
+
+	// Hover material used when object state is hovering.
+	virtual void SetHoverMaterial(UMaterialInterface* Material);
+	UMaterialInterface* GetHoverMaterial() const;
+
+	// Interact material used when object state is interacting.
+	virtual void SetInteractMaterial(UMaterialInterface* Material);
+	UMaterialInterface* GetInteractMaterial() const;
+
+	// Vertex color where applicable.
+	virtual void SetVertexColor(const FColor& InVertexColor);
+	FColor GetVertexColor() const;
+
+protected:
+
+	// Render and LineTrace should only occur when bEnabled is true.
+	UPROPERTY()
+	bool bEnabled = true;
+
+	// Element state - indicates whether object is visible or hittable
+	UPROPERTY()
+	EGizmoElementState ElementState = EGizmoElementState::VisibleAndHittable;
+
+	// Element interaction state - None, Hovering or Interacting
+	UPROPERTY()
+	EGizmoElementInteractionState ElementInteractionState = EGizmoElementInteractionState::None;
+
+	// View-dependent type - None, Axis or Plane. 
+	UPROPERTY()
+	EGizmoElementViewDependentType ViewDependentType = EGizmoElementViewDependentType::None;
+
+	// View-dependent axis or plane normal, based on the view-dependent type.
+	UPROPERTY()
+	FVector ViewDependentAxis = FVector::UpVector;
+
+	// View-dependent angle tolerance based on :
+	//   For Axis, minimum radians between view dependent axis and view direction.
+	//   For Planar, minimum radians between view dependent axis and the plane where axis is its normal.
+	// When the angle between the view direction and the axis/plane is less than this tolerance, this object should be culled.
+	UPROPERTY()
+	float ViewDependentAngleTol = DefaultViewDependentAngleTol;
+
+	// Axial view alignment minimum cos angle tolerance, computed based on ViewDependentAngleTol. 
+	// When the cos of the angle between the view direction and the axis is less than this value, this object should not be culled.
+	UPROPERTY()
+	float ViewDependentAxialMaxCosAngleTol = DefaultViewAlignAxialMaxCosAngleTol;
+
+	// Planar view alignment minimum cos angle tolerance, computed based on ViewDependentAngleTol. 
+	// When the cos of the angle between the view direction and the axis is greater than this value, this object should not be culled.
+	UPROPERTY()
+	float ViewDependentPlanarMinCosAngleTol = DefaultViewAlignPlanarMinCosAngleTol;
+
+	// View align type: None, PointEye, or PointWorld.
+	// PointEye rotates this axis to align with the view up axis.
+	// PointWorld rotates this axis to align with the world up axis.
+	// Axial rotates around this axis to align the normal as closely as possible to the view direction.
+	UPROPERTY()
+	EGizmoElementViewAlignType ViewAlignType = EGizmoElementViewAlignType::None;
+
+	// View align axis. 
+	UPROPERTY()
+	FVector ViewAlignAxis = FVector::UpVector;
+
+	// View align normal.
+	// PointEye and PointWorld both rotate the normal to align with the view direction.
+	// Axial rotates the normal to align as closely as possible with view direction.
+	UPROPERTY()
+	FVector ViewAlignNormal = -FVector::ForwardVector;
+
+	// Axial view alignment angle tolerance in radians, based on angle between align normal and view direction. 
+	// When angle between the view align normal and the view direction is greater than this angle, the align rotation will be computed.
+	UPROPERTY()
+	float ViewAlignAxialAngleTol = DefaultViewAlignAngleTol;
+
+	// Axial view alignment minimum cos angle tolerance, computed based on ViewAlignAxialAngleTol. 
+	// When the cos of the angle between the view direction and the align normal is less than this value, the align rotation will be computed.
+	UPROPERTY()
+	float ViewAlignAxialMaxCosAngleTol = DefaultViewAlignMaxCosAngleTol;
+
+	// Default material.
+	UPROPERTY()
+	TObjectPtr<UMaterialInterface> Material;
+
+	// Hovering material.
+	UPROPERTY()
+	TObjectPtr<UMaterialInterface> HoverMaterial;
+
+	// Interacting material.
+	UPROPERTY()
+	TObjectPtr<UMaterialInterface> InteractMaterial;
+
+	// Vertex color where applicable.
+	UPROPERTY()
+	FColor VertexColor = FColor::White;
+
+	// Render stores the last cached transform, used by line trace.
+	UPROPERTY()
+	FTransform CachedLocalToWorldTransform = FTransform::Identity;
+
+	// Render stores the last cached translation, used by line trace.
+	UPROPERTY()
+	FVector CachedLocalToWorldTranslation = FVector::ZeroVector;
+
+	// Render stores the last cached rotation, used by line trace.
+	UPROPERTY()
+	FQuat CachedLocalToWorldRotation = FQuat::Identity;
+
+	// Render stores the last cached rotation, used by line trace.
+	UPROPERTY()
+	double CachedLocalToWorldScale = 1.0;
+
+	// Whether last transform has been cached.
+	UPROPERTY()
+	bool bHasCachedLocalToWorldTransform = false;
+
+	// Whether visible object was visible during the last render.
+	UPROPERTY()
+	bool bCachedVisibleViewDependent = true;
+
+	// Cached box sphere bounds.
+	UPROPERTY()
+	FBoxSphereBounds CachedBoxSphereBounds;
+
+	// Cached box sphere bounds.
+	UPROPERTY()
+	bool bHasCachedBoxSphereBounds;
+
+protected:
+	// Cache render state during render traversal, to be used subsequently when line tracing.
+	virtual void CacheRenderState(const FTransform& InLocalToWorldState, bool InVisibleViewDependent = true);
+
+	// Update render traversal state based on properties in the current object, should always be called at the beginning of Render()
+	virtual void UpdateRenderTraversalState(FRenderTraversalState& InRenderTraversalState);
+
+};
+
diff --git a/Engine/Source/Runtime/InteractiveToolsFramework/Public/BaseGizmos/GizmoElementBox.h b/Engine/Source/Runtime/InteractiveToolsFramework/Public/BaseGizmos/GizmoElementBox.h
new file mode 100644
--- /dev/null
+++ b/Engine/Source/Runtime/InteractiveToolsFramework/Public/BaseGizmos/GizmoElementBox.h
@@ -0,0 +1,63 @@
+// Copyright Epic Games, Inc. All Rights Reserved.
+
+#pragma once
+
+#include "CoreMinimal.h"
+#include "BaseGizmos/GizmoElementBase.h"
+#include "InputState.h"
+#include "UObject/ObjectMacros.h"
+#include "GizmoElementBox.generated.h"
+
+/**
+ * Simple object intended to be used as part of 3D Gizmos.
+ * Draws a solid 3D cylinder based on parameters.
+ */
+
+UCLASS(Transient)
+class INTERACTIVETOOLSFRAMEWORK_API UGizmoElementBox : public UGizmoElementBase
+{
+	GENERATED_BODY()
+
+public:
+
+	//~ Begin UGizmoElementBase Interface.
+	virtual void Render(IToolsContextRenderAPI* RenderAPI, const FRenderTraversalState& RenderState) override;
+	virtual FInputRayHit LineTrace(const FVector Start, const FVector Direction) override;
+	virtual FBoxSphereBounds CalcBounds(const FTransform& LocalToWorld) const override;
+	//~ End UGizmoElementBase Interface.
+
+	// Location of center of box's base circle.
+	virtual void SetCenter(const FVector& InCenter);
+	virtual FVector GetCenter() const;
+
+	// Box axis up direction.
+	virtual void SetUpDirection(const FVector& InUpDirection);
+	virtual FVector GetUpDirection() const;
+
+	// Box axis side direction.
+	virtual void SetSideDirection(const FVector& InSideDirection);
+	virtual FVector GetSideDirection() const;
+
+	// Box dimensions.
+	virtual void SetDimensions(const FVector& InDimensions);
+	virtual FVector GetDimensions() const;
+
+protected:
+
+	// Location of center of box
+	UPROPERTY()
+	FVector Center = FVector::ZeroVector;
+
+	// Box dimensions
+	UPROPERTY()
+	FVector Dimensions = FVector(1.0f, 1.0f, 1.0f);
+
+	// Box up axis
+	UPROPERTY()
+	FVector UpDirection = FVector::UpVector;
+
+	// Box side axis
+	UPROPERTY()
+	FVector SideDirection = FVector::RightVector;
+};
+
diff --git a/Engine/Source/Runtime/InteractiveToolsFramework/Public/BaseGizmos/GizmoElementCircle.h b/Engine/Source/Runtime/InteractiveToolsFramework/Public/BaseGizmos/GizmoElementCircle.h
new file mode 100644
--- /dev/null
+++ b/Engine/Source/Runtime/InteractiveToolsFramework/Public/BaseGizmos/GizmoElementCircle.h
@@ -0,0 +1,108 @@
+// Copyright Epic Games, Inc. All Rights Reserved.
+
+#pragma once
+
+#include "CoreMinimal.h"
+#include "BaseGizmos/GizmoElementBase.h"
+#include "InputState.h"
+#include "UObject/ObjectMacros.h"
+#include "GizmoElementCircle.generated.h"
+
+/**
+ * Simple object intended to be used as part of 3D Gizmos.
+ * Draws a filled or line circle based on parameters.
+ */
+UCLASS(Transient)
+class INTERACTIVETOOLSFRAMEWORK_API UGizmoElementCircle : public UGizmoElementBase
+{
+	GENERATED_BODY()
+
+public:
+	//~ Begin UGizmoElementBase Interface.
+	virtual void Render(IToolsContextRenderAPI* RenderAPI, const FRenderTraversalState& RenderState) override;
+	virtual FInputRayHit LineTrace(const FVector Start, const FVector Direction) override;
+	virtual FBoxSphereBounds CalcBounds(const FTransform& LocalToWorld) const override;
+	//~ End UGizmoElementBase Interface.
+
+	// Circle center.
+	virtual void SetCenter(FVector InCenter);
+	virtual FVector GetCenter() const;
+
+	// Normal to circle.
+	virtual void SetNormal(FVector InNormal);
+	virtual FVector GetNormal() const;
+
+	// Circle radius.
+	virtual void SetRadius(float Radius);
+	virtual float GetRadius() const;
+
+	// Number of sides for rendering circle.
+	virtual void SetNumSides(int InNumSides);
+	virtual int GetNumSides() const;
+
+	// Line color when bDrawLine is true.
+	virtual void SetLineColor(const FLinearColor& InColor);
+	virtual FLinearColor GetLineColor() const;
+
+	// Draw mesh
+	virtual void SetDrawMesh(bool InDrawMesh);
+	virtual bool GetDrawMesh() const;
+
+	// Draw line
+	virtual void SetDrawLine(bool InDrawLine);
+	virtual bool GetDrawLine() const;
+
+	// Hit mesh
+	virtual void SetHitMesh(bool InHitMesh);
+	virtual bool GetHitMesh() const;
+
+	// Hit line
+	virtual void SetHitLine(bool InHitLine);
+	virtual bool GetHitLine() const;
+
+	// Screen space, when true orients the circle to the screen up and side vectors
+	virtual void SetScreenSpace(bool InScreenSpace);
+	virtual bool GetScreenSpace();
+
+protected:
+
+	// Circle center.
+	UPROPERTY()
+	FVector Center = FVector::ZeroVector;
+
+	// Normal to circle.
+	UPROPERTY()
+	FVector Normal = FVector::ForwardVector;
+
+	// Circle radius.
+	UPROPERTY()
+	float Radius = 100.0f;
+
+	// Number of sides for rendering circle.
+	UPROPERTY()
+	int NumSides = 64;
+
+	// Line color when bFill is false.
+	UPROPERTY()
+	FLinearColor LineColor = FLinearColor::White;
+
+	// Whether to render solid circle.
+	UPROPERTY()
+	bool bDrawMesh = true;
+
+	// Whether to render line circle.
+	UPROPERTY()
+	bool bDrawLine = false;
+
+	// Whether to perform hit test on mesh.
+	UPROPERTY()
+	bool bHitMesh = true;
+
+	// Whether to perform hit test on line.
+	UPROPERTY()
+	bool bHitLine = false;
+
+	// Screen space flag
+	UPROPERTY()
+	bool bScreenSpace = false;
+};
\ No newline at end of file
diff --git a/Engine/Source/Runtime/InteractiveToolsFramework/Public/BaseGizmos/GizmoElementCone.h b/Engine/Source/Runtime/InteractiveToolsFramework/Public/BaseGizmos/GizmoElementCone.h
new file mode 100644
--- /dev/null
+++ b/Engine/Source/Runtime/InteractiveToolsFramework/Public/BaseGizmos/GizmoElementCone.h
@@ -0,0 +1,71 @@
+// Copyright Epic Games, Inc. All Rights Reserved.
+
+#pragma once
+
+#include "CoreMinimal.h"
+#include "BaseGizmos/GizmoElementBase.h"
+#include "InputState.h"
+#include "UObject/ObjectMacros.h"
+#include "GizmoElementCone.generated.h"
+
+/**
+ * Simple object intended to be used as part of 3D Gizmos.
+ * Draws a solid 3D cone based on parameters.
+ */
+
+UCLASS(Transient)
+class INTERACTIVETOOLSFRAMEWORK_API UGizmoElementCone : public UGizmoElementBase
+{
+	GENERATED_BODY()
+
+public:
+
+	//~ Begin UGizmoElementBase Interface.
+	virtual void Render(IToolsContextRenderAPI* RenderAPI, const FRenderTraversalState& RenderState) override;
+	virtual FInputRayHit LineTrace(const FVector Start, const FVector Direction) override;
+	virtual FBoxSphereBounds CalcBounds(const FTransform& LocalToWorld) const override;
+	//~ End UGizmoElementBase Interface.
+
+	// Location of center of cone's base circle.
+	virtual void SetBase(const FVector& InBase);
+	virtual FVector GetBase() const;
+
+	// Cone axis direction.
+	virtual void SetDirection(const FVector& InDirection);
+	virtual FVector GetDirection() const;
+
+	// Cone height.
+	virtual void SetHeight(float InHeight);
+	virtual float GetHeight() const;
+
+	// Cone radius.
+	virtual void SetRadius(float InRadius);
+	virtual float GetRadius() const;
+
+	// Number of sides to use when rendering cone.
+	virtual void SetNumSides(int32 InNumSides);
+	virtual int32 GetNumSides() const;
+
+protected:
+
+	// Location of center of cone's base circle.
+	UPROPERTY()
+	FVector Base = FVector::ZeroVector;
+
+	// Direction of cone's axis.
+	UPROPERTY()
+	FVector Direction = FVector(0.0f, 0.0f, 1.0f);
+
+	// Height of cone
+	UPROPERTY()
+	float Height = 1.0f;
+
+	// Radius of cone circle.
+	UPROPERTY()
+	float Radius = 0.5f;
+
+	// Number of sides for tessellating cone
+	UPROPERTY()
+	int32 NumSides = 32;
+};
+
diff --git a/Engine/Source/Runtime/InteractiveToolsFramework/Public/BaseGizmos/GizmoElementCylinder.h b/Engine/Source/Runtime/InteractiveToolsFramework/Public/BaseGizmos/GizmoElementCylinder.h
new file mode 100644
--- /dev/null
+++ b/Engine/Source/Runtime/InteractiveToolsFramework/Public/BaseGizmos/GizmoElementCylinder.h
@@ -0,0 +1,71 @@
+// Copyright Epic Games, Inc. All Rights Reserved.
+
+#pragma once
+
+#include "CoreMinimal.h"
+#include "BaseGizmos/GizmoElementBase.h"
+#include "InputState.h"
+#include "UObject/ObjectMacros.h"
+#include "GizmoElementCylinder.generated.h"
+
+/**
+ * Simple object intended to be used as part of 3D Gizmos.
+ * Draws a solid 3D cylinder based on parameters.
+ */
+
+UCLASS(Transient)
+class INTERACTIVETOOLSFRAMEWORK_API UGizmoElementCylinder : public UGizmoElementBase
+{
+	GENERATED_BODY()
+
+public:
+
+	//~ Begin UGizmoElementBase Interface.
+	virtual void Render(IToolsContextRenderAPI* RenderAPI, const FRenderTraversalState& RenderState) override;
+	virtual FInputRayHit LineTrace(const FVector Start, const FVector Direction) override;
+	virtual FBoxSphereBounds CalcBounds(const FTransform& LocalToWorld) const override;
+	//~ End UGizmoElementBase Interface.
+
+	// Location of center of cylinder's base circle.
+	virtual void SetBase(const FVector& InBase);
+	virtual FVector GetBase() const;
+
+	// Cylinder axis direction.
+	virtual void SetDirection(const FVector& InDirection);
+	virtual FVector GetDirection() const;
+
+	// Cylinder height.
+	virtual void SetHeight(float InHeight);
+	virtual float GetHeight() const;
+
+	// Cylinder radius.
+	virtual void SetRadius(float InRadius);
+	virtual float GetRadius() const;
+
+	// Number of sides to use when rendering cylinder.
+	virtual void SetNumSides(int32 InNumSides);
+	virtual int32 GetNumSides() const;
+
+protected:
+
+	// Location of center of cylinder's base circle.
+	UPROPERTY()
+	FVector Base = FVector::ZeroVector;
+
+	// Direction of cylinder's axis.
+	UPROPERTY()
+	FVector Direction = FVector(0.0f, 0.0f, 1.0f);
+
+	// Height of cylinder
+	UPROPERTY()
+	float Height = 1.0f;
+
+	// Radius of cylinder circle.
+	UPROPERTY()
+	float Radius = 0.5f;
+
+	// Number of sides for tessellating cylinder
+	UPROPERTY()
+	int32 NumSides = 32;
+};
+
diff --git a/Engine/Source/Runtime/InteractiveToolsFramework/Public/BaseGizmos/GizmoElementGroup.h b/Engine/Source/Runtime/InteractiveToolsFramework/Public/BaseGizmos/GizmoElementGroup.h
new file mode 100644
--- /dev/null
+++ b/Engine/Source/Runtime/InteractiveToolsFramework/Public/BaseGizmos/GizmoElementGroup.h
@@ -0,0 +1,50 @@
+// Copyright Epic Games, Inc. All Rights Reserved.
+
+#pragma once
+
+#include "CoreMinimal.h"
+#include "UObject/ObjectMacros.h"
+#include "BaseGizmos/GizmoElementBase.h"
+#include "InputState.h"
+#include "GizmoElementGroup.generated.h"
+
+/**
+ * Simple group object intended to be used as part of 3D Gizmos.
+ * Contains multiple gizmo objects.
+ */
+UCLASS(Transient)
+class INTERACTIVETOOLSFRAMEWORK_API UGizmoElementGroup : public UGizmoElementBase
+{
+	GENERATED_BODY()
+
+public:
+
+	//~ Begin UGizmoElementBase Interface.
+	virtual void Render(IToolsContextRenderAPI* RenderAPI, const FRenderTraversalState& RenderState) override;
+	virtual FInputRayHit LineTrace(const FVector Start, const FVector Direction) override;
+	virtual FBoxSphereBounds CalcBounds(const FTransform& LocalToWorld) const override;
+	//~ End UGizmoElementBase Interface.
+
+	// Add object to group.
+	virtual void Add(UGizmoElementBase* InElement);
+
+	// Remove object from group, if it exists
+	virtual void Remove(UGizmoElementBase* InElement);
+
+	// Reset cached render state
+	virtual void ResetCachedRenderState();
+
+	// When true, maintains view-dependent constant scale for this gizmo object hierarchy
+	virtual void SetConstantScale(bool InConstantScale);
+	virtual bool GetConstantScale() const;
+
+protected:
+
+	// When true, maintains view-dependent constant scale for this gizmo object hierarchy
+	UPROPERTY()
+	bool bConstantScale = false;
+
+	// Gizmo elements within this group
+	UPROPERTY()
+	TArray<TObjectPtr<UGizmoElementBase>> Elements;
+};
diff --git a/Engine/Source/Runtime/InteractiveToolsFramework/Public/BaseGizmos/GizmoElementHitTargets.h b/Engine/Source/Runtime/InteractiveToolsFramework/Public/BaseGizmos/GizmoElementHitTargets.h
new file mode 100644
--- /dev/null
+++ b/Engine/Source/Runtime/InteractiveToolsFramework/Public/BaseGizmos/GizmoElementHitTargets.h
@@ -0,0 +1,48 @@
+// Copyright Epic Games, Inc. All Rights Reserved.
+
+#pragma once
+
+#include "CoreMinimal.h"
+#include "BaseGizmos/GizmoInterfaces.h"
+#include "BaseGizmos/HitTargets.h"
+#include "GizmoElementHitTargets.generated.h"
+
+class UGizmoElementBase;
+
+/**
+ * UGizmoElementHitTarget is an IGizmoClickTarget implementation that
+ * hit-tests any object derived from UGizmoElementBase
+ */
+UCLASS()
+class INTERACTIVETOOLSFRAMEWORK_API UGizmoElementHitTarget : public UObject, public IGizmoClickTarget
+{
+	GENERATED_BODY()
+public:
+
+	/**
+	 * Gizmo element.
+	 */
+	UPROPERTY()
+	TObjectPtr<UGizmoElementBase> GizmoElement;
+
+	/**
+	 * If set, this condition is checked before performing the hit test. This gives a way
+	 * to disable the hit test without hiding the component. This is useful, for instance,
+	 * in a repositionable transform gizmo in world-coordinate mode, where the rotation
+	 * components need to be hittable for movement, but not for repositioning.
+	 */
+	TFunction<bool(const FInputDeviceRay&)> Condition = nullptr;
+
+public:
+	virtual FInputRayHit IsHit(const FInputDeviceRay& ClickPos) const;
+
+	virtual void UpdateHoverState(bool bHovering);
+
+	virtual void UpdateInteractingState(bool bInteracting);
+
+public:
+	static UGizmoElementHitTarget* Construct(
+		UGizmoElementBase* InGizmoElement,
+		UObject* Outer = (UObject*)GetTransientPackage());
+};
+
diff --git a/Engine/Source/Runtime/InteractiveToolsFramework/Public/BaseGizmos/GizmoElementRectangle.h b/Engine/Source/Runtime/InteractiveToolsFramework/Public/BaseGizmos/GizmoElementRectangle.h
new file mode 100644
--- /dev/null
+++ b/Engine/Source/Runtime/InteractiveToolsFramework/Public/BaseGizmos/GizmoElementRectangle.h
@@ -0,0 +1,116 @@
+// Copyright Epic Games, Inc. All Rights Reserved.
+
+#pragma once
+
+#include "CoreMinimal.h"
+#include "BaseGizmos/GizmoElementBase.h"
+#include "InputState.h"
+#include "UObject/ObjectMacros.h"
+#include "GizmoElementRectangle.generated.h"
+
+class FPrimitiveDrawInterface;
+class FMaterialRenderProxy;
+
+/**
+ * Simple object intended to be used as part of 3D Gizmos.
+ * Draws a rectangle based on parameters.
+ */
+UCLASS(Transient)
+class INTERACTIVETOOLSFRAMEWORK_API UGizmoElementRectangle : public UGizmoElementBase
+{
+	GENERATED_BODY()
+
+public:
+
+	//~ Begin UGizmoElementBase Interface.
+	virtual void Render(IToolsContextRenderAPI* RenderAPI, const FRenderTraversalState& RenderState) override;
+	virtual FInputRayHit LineTrace(const FVector Start, const FVector Direction) override;
+	virtual FBoxSphereBounds CalcBounds(const FTransform& LocalToWorld) const override;
+	//~ End UGizmoElementBase Interface.
+
+	// Location of rectangle center
+	virtual void SetCenter(FVector InCenter);
+	virtual FVector GetCenter() const;
+
+	// Width
+	virtual void SetWidth(float InWidth);
+	virtual float GetWidth() const;
+
+	// Height
+	virtual void SetHeight(float InHeight);
+	virtual float GetHeight() const;
+
+	// Up direction
+	virtual void SetUpDirection(const FVector& InUpDirection);
+	virtual FVector GetUpDirection() const;
+
+	// Side direction
+	virtual void SetSideDirection(const FVector& InSideDirection);
+	virtual FVector GetSideDirection() const;
+
+	// Screen space, when true orients the rectangle to the screen up and side vectors
+	virtual void SetScreenSpace(bool InScreenSpace);
+	virtual bool GetScreenSpace();
+
+	// Line color
+	virtual void SetLineColor(const FColor& InLineColor);
+	virtual FColor GetLineColor();
+
+	// Draw mesh
+	virtual void SetDrawMesh(bool InDrawMesh);
+	virtual bool GetDrawMesh() const;
+
+	// Draw line
+	virtual void SetDrawLine(bool InDrawLine);
+	virtual bool GetDrawLine() const;
+
+	// Hit mesh
+	virtual void SetHitMesh(bool InHitMesh);
+	virtual bool GetHitMesh() const;
+
+	// Hit line
+	virtual void SetHitLine(bool InHitLine);
+	virtual bool GetHitLine() const;
+
+protected:
+
+	// Location of rectangle center
+	UPROPERTY()
+	FVector Center = FVector::ZeroVector;
+
+	// Width
+	UPROPERTY()
+	float Width = 1.0f;
+
+	// Height
+	UPROPERTY()
+	float Height = 1.0f;
+
+	// Up direction
+	UPROPERTY()
+	FVector UpDirection = FVector(0.0f, 0.0f, 1.0f);
+
+	// Side direction
+	UPROPERTY()
+	FVector SideDirection = FVector(0.0f, 1.0f, 0.0f);
+
+	// Screen space flag
+	UPROPERTY()
+	bool bScreenSpace = false;
+
+	// Line color when bShowOutline is false.
+	UPROPERTY()
+	FColor LineColor = FColor::White;
+
+	UPROPERTY()
+	bool bDrawMesh = true;
+
+	UPROPERTY()
+	bool bDrawLine = false;
+
+	UPROPERTY()
+	bool bHitMesh = true;
+
+	UPROPERTY()
+	bool bHitLine = false;
+};
\ No newline at end of file
diff --git a/Engine/Source/Runtime/InteractiveToolsFramework/Public/BaseGizmos/GizmoElementShapes.h b/Engine/Source/Runtime/InteractiveToolsFramework/Public/BaseGizmos/GizmoElementShapes.h
new file mode 100644
--- /dev/null
+++ b/Engine/Source/Runtime/InteractiveToolsFramework/Public/BaseGizmos/GizmoElementShapes.h
@@ -0,0 +1,13 @@
+// Copyright Epic Games, Inc. All Rights Reserved.
+
+#pragma once
+
+#include "BaseGizmos/GizmoElementArrow.h"
+#include "BaseGizmos/GizmoElementBox.h"
+#include "BaseGizmos/GizmoElementCone.h"
+#include "BaseGizmos/GizmoElementCircle.h"
+#include "BaseGizmos/GizmoElementCylinder.h"
+#include "BaseGizmos/GizmoElementGroup.h"
+#include "BaseGizmos/GizmoElementRectangle.h"
+#include "BaseGizmos/GizmoElementTorus.h"
+
diff --git a/Engine/Source/Runtime/InteractiveToolsFramework/Public/BaseGizmos/GizmoElementStateTargets.h b/Engine/Source/Runtime/InteractiveToolsFramework/Public/BaseGizmos/GizmoElementStateTargets.h
new file mode 100644
--- /dev/null
+++ b/Engine/Source/Runtime/InteractiveToolsFramework/Public/BaseGizmos/GizmoElementStateTargets.h
@@ -0,0 +1,128 @@
+// Copyright Epic Games, Inc. All Rights Reserved.
+
+#pragma once
+
+#include "CoreMinimal.h"
+#include "InteractiveGizmoManager.h"
+#include "InteractiveToolManager.h"
+#include "BaseGizmos/GizmoInterfaces.h"
+#include "BaseGizmos/GizmoElementBase.h"
+#include "Changes/TransformChange.h"
+#include "GizmoElementStateTargets.generated.h"
+
+/**
+ * UGizmoDependentTransformChangeStateTarget is an implementation of IGizmoStateTarget that
+ * opens/closes an undo transaction via GizmoManager.
+ *
+ * The DependentChangeSources and ExternalDependentChangeSources lists allow additional
+ * FChange objects to be inserted into the transaction, provided by IToolCommandChangeSource implementations.
+ */
+UCLASS()
+class INTERACTIVETOOLSFRAMEWORK_API UGizmoDependentTransformChangeStateTarget : public UObject, public IGizmoStateTarget
+{
+	GENERATED_BODY()
+public:
+	virtual void BeginUpdate()
+	{
+		if (TransactionManager)
+		{
+			TransactionManager->BeginUndoTransaction(ChangeDescription);
+		}
+
+		for (TUniquePtr<IToolCommandChangeSource>& Source : DependentChangeSources)
+		{
+			if (Source.IsValid())
+			{
+				Source->BeginChange();
+			}
+		}
+
+		for (IToolCommandChangeSource* Source : ExternalDependentChangeSources)
+		{
+			if (Source)
+			{
+				Source->BeginChange();
+			}
+		}
+	}
+
+	virtual void EndUpdate()
+	{
+		if (TransactionManager)
+		{
+			for (TUniquePtr<IToolCommandChangeSource>& Source : DependentChangeSources)
+			{
+				if (Source.IsValid())
+				{
+					TUniquePtr<FToolCommandChange> Change = Source->EndChange();
+					if (Change)
+					{
+						UObject* Target = Source->GetChangeTarget();
+						TransactionManager->EmitObjectChange(Target, MoveTemp(Change), Source->GetChangeDescription());
+					}
+				}
+			}
+
+			for (IToolCommandChangeSource* Source : ExternalDependentChangeSources)
+			{
+				if (Source)
+				{
+					TUniquePtr<FToolCommandChange> Change = Source->EndChange();
+					if (Change)
+					{
+						UObject* Target = Source->GetChangeTarget();
+						TransactionManager->EmitObjectChange(Target, MoveTemp(Change), Source->GetChangeDescription());
+					}
+				}
+			}
+
+			TransactionManager->EndUndoTransaction();
+		}
+	}
+
+	/**
+	 * Localized text description of the transaction (will be visible in Editor on undo/redo)
+	 */
+	FText ChangeDescription;
+
+	/**
+	 * Pointer to the GizmoManager or ToolManager that is used to open/close the transaction
+	 */
+	UPROPERTY()
+	TScriptInterface<IToolContextTransactionProvider> TransactionManager;
+
+	/**
+	 * Dependent-change generators. This will be told about update start/end, and any generated changes will also be emitted.
+	 * This allows (eg) TransformProxy change events to be collected at the same time as changes to a gizmo target component.
+	 */
+	TArray<TUniquePtr<IToolCommandChangeSource>> DependentChangeSources;
+
+	/**
+	 * Dependent-change generators that are not owned by this class, otherwise handled identically to DependentChangeSources
+	 */
+	TArray<IToolCommandChangeSource*> ExternalDependentChangeSources;
+
+public:
+
+	/**
+	 * Create and initialize an standard instance of UGizmoObjectTransformChangeStateTarget
+	 * @param DescriptionIn Localized text description of the transaction
+	 * @param TransactionManagerIn pointer to the GizmoManager or ToolManager used to manage transactions
+	 */
+	static UGizmoDependentTransformChangeStateTarget* Construct(
+		const FText& DescriptionIn,
+		IToolContextTransactionProvider* TransactionManagerIn,
+		UObject* Outer = (UObject*)GetTransientPackage())
+	{
+		UGizmoDependentTransformChangeStateTarget* NewTarget = NewObject<UGizmoDependentTransformChangeStateTarget>(Outer);
+		NewTarget->ChangeDescription = DescriptionIn;
+
+		// have to explicitly configure this because we only have IToolContextTransactionProvider pointer
+		NewTarget->TransactionManager.SetInterface(TransactionManagerIn);
+		NewTarget->TransactionManager.SetObject(CastChecked<UObject>(TransactionManagerIn));
+
+		return NewTarget;
+	}
+
+
+};
diff --git a/Engine/Source/Runtime/InteractiveToolsFramework/Public/BaseGizmos/GizmoElementTorus.h b/Engine/Source/Runtime/InteractiveToolsFramework/Public/BaseGizmos/GizmoElementTorus.h
new file mode 100644
--- /dev/null
+++ b/Engine/Source/Runtime/InteractiveToolsFramework/Public/BaseGizmos/GizmoElementTorus.h
@@ -0,0 +1,123 @@
+// Copyright Epic Games, Inc. All Rights Reserved.
+
+#pragma once
+
+#include "CoreMinimal.h"
+#include "BaseGizmos/GizmoElementBase.h"
+#include "InputState.h"
+#include "UObject/ObjectMacros.h"
+#include "GizmoElementTorus.generated.h"
+
+class FPrimitiveDrawInterface;
+class FMaterialRenderProxy;
+
+/**
+ * Simple object intended to be used as part of 3D Gizmos.
+ * Draws a torus based on parameters.
+ * Note: Ray-torus hit testing is not supported! LineTrace() will never return a hit nor will CalcBounds() calculate box-sphere bounds.
+ *       Use a hidden GizmoElementCircle with circle primitive object with DrawLines set to true to approximate 
+ *       torus intersection for hit-testing purposes.
+ */
+UCLASS(Transient)
+class INTERACTIVETOOLSFRAMEWORK_API UGizmoElementTorus : public UGizmoElementBase
+{
+	GENERATED_BODY()
+
+public:
+
+	//~ Begin UGizmoElementBase Interface.
+	virtual void Render(IToolsContextRenderAPI* RenderAPI, const FRenderTraversalState& RenderState) override;
+	virtual FInputRayHit LineTrace(const FVector Start, const FVector Direction) override;
+	virtual FBoxSphereBounds CalcBounds(const FTransform& LocalToWorld) const override;
+	//~ End UGizmoElementBase Interface.
+
+	// Torus center.
+	virtual void SetCenter(const FVector& InCenter);
+	virtual FVector GetCenter() const;
+
+	// Normal to torus.
+	virtual void SetNormal(const FVector& InNormal);
+	virtual FVector GetNormal() const;
+
+	// Begin axis, should be perpendicular to normal, this where partial torus will begin.
+	virtual void SetBeginAxis(const FVector& InBeginAxis);
+	virtual FVector GetBeginAxis() const;
+
+	// Outer circle radius.
+	virtual void SetOuterRadius(float InOuterRadius);
+	virtual float GetOuterRadius() const;
+
+	// Inner circles radius.
+	virtual void SetInnerRadius(float InInnerRadius);
+	virtual float GetInnerRadius() const;
+
+	// Number of outer segments for rendering torus.
+	virtual void SetOuterSegments(int32 InOuterSegments);
+	virtual int32 GetOuterSegments() const;
+
+	// Number of inner slices for rendering torus.
+	virtual void SetInnerSlices(int32 InInnerSlices);
+	virtual int32 GetInnerSlices() const;
+
+	// True when the torus is a partial arc.
+	virtual void SetPartial(bool InPartial);
+	virtual bool GetPartial() const;
+
+	// True when the partial torus arc should be screen aligned.
+	virtual void SetScreenAlignPartial(bool InScreenAlignPartial);
+	virtual bool GetScreenAlignPartial() const;
+
+	// If partial, arc angle of partial torus in radians.
+	virtual void SetAngle(float InAngle);
+	virtual float GetAngle() const;
+
+	// If partial, renders end caps when true.
+	virtual void SetEndCaps(bool InEndCaps);
+	virtual bool GetEndCaps() const;
+
+protected:
+
+	// Torus center.
+	UPROPERTY()
+	FVector Center = FVector::ZeroVector;
+
+	// Normal to plane in which torus lies.
+	UPROPERTY()
+	FVector Normal = FVector::UpVector;
+
+	// Plane axis, should be perpendicular to normal. Indicates where partial torus should begin
+	UPROPERTY()
+	FVector BeginAxis = FVector::ForwardVector;
+
+	// Torus outer radius.
+	UPROPERTY()
+	float OuterRadius = 100.0f;
+
+	// Torus inner radius.
+	UPROPERTY()
+	float InnerRadius = 100.0f;
+
+	// Number of segments for rendering torus.
+	UPROPERTY()
+	int32 OuterSegments = 64;
+
+	// Number of slices to render in each torus segment.
+	UPROPERTY()
+	int32 InnerSlices = 8;
+
+	// True when the torus is not full.
+	UPROPERTY()
+	bool bPartial = false;
+
+	// True when partial torus should be aligned to the screen
+	UPROPERTY(EditAnywhere, Category = Options, meta = (EditCondition = "bPartial"))
+	bool bScreenAlignPartial = false;
+
+	// Angle to render for partial torus
+	UPROPERTY(EditAnywhere, Category = Options, meta = (EditCondition = "bPartial"))
+	float Angle = 2.0f * PI;
+
+	// Whether to render end caps on a partial torus.
+	UPROPERTY(EditAnywhere, Category = Options, meta = (EditCondition = "bPartial"))
+	bool bEndCaps = false;
+};
\ No newline at end of file
