New TRS gizmo: add support for hovering.

diff --git a/Engine/Source/Editor/Experimental/EditorInteractiveToolsFramework/Private/EditorGizmos/EditorTransformGizmoBuilder.cpp b/Engine/Source/Editor/Experimental/EditorInteractiveToolsFramework/Private/EditorGizmos/EditorTransformGizmoBuilder.cpp
--- a/Engine/Source/Editor/Experimental/EditorInteractiveToolsFramework/Private/EditorGizmos/EditorTransformGizmoBuilder.cpp
+++ b/Engine/Source/Editor/Experimental/EditorInteractiveToolsFramework/Private/EditorGizmos/EditorTransformGizmoBuilder.cpp
@@ -4,12 +4,19 @@
 #include "EditorGizmos/EditorTransformGizmoSource.h"
 #include "EditorGizmos/EditorTransformProxy.h"
 #include "EditorGizmos/TransformGizmo.h"
+#include "BaseGizmos/GizmoElementGroup.h"
 
 UInteractiveGizmo* UEditorTransformGizmoBuilder::BuildGizmo(const FToolBuilderState& SceneState) const
 {
 	UTransformGizmo* TransformGizmo = NewObject<UTransformGizmo>(SceneState.GizmoManager);
 	TransformGizmo->Setup();
 	TransformGizmo->TransformSource = UEditorTransformGizmoSource::Construct(TransformGizmo);
+
+	// @todo: Gizmo element construction to be moved here from UTransformGizmo.
+	// A UGizmoElementRenderMultiTarget will be constructed and both the
+	// render and hit target's Construct methods will take the gizmo element root as input.
+	TransformGizmo->HitTarget = UGizmoElementHitMultiTarget::Construct(TransformGizmo->GizmoElementRoot);
+
 	return TransformGizmo;
 }
 
diff --git a/Engine/Source/Editor/Experimental/EditorInteractiveToolsFramework/Private/EditorGizmos/TransformGizmo.cpp b/Engine/Source/Editor/Experimental/EditorInteractiveToolsFramework/Private/EditorGizmos/TransformGizmo.cpp
--- a/Engine/Source/Editor/Experimental/EditorInteractiveToolsFramework/Private/EditorGizmos/TransformGizmo.cpp
+++ b/Engine/Source/Editor/Experimental/EditorInteractiveToolsFramework/Private/EditorGizmos/TransformGizmo.cpp
@@ -1,6 +1,7 @@
 // Copyright Epic Games, Inc. All Rights Reserved.
 
 #include "EditorGizmos/TransformGizmo.h"
+#include "BaseBehaviors/MouseHoverBehavior.h"
 #include "BaseGizmos/AxisSources.h"
 #include "BaseGizmos/GizmoElementGroup.h"
 #include "BaseGizmos/GizmoElementShapes.h"
@@ -21,6 +22,8 @@
 
 #define LOCTEXT_NAMESPACE "UTransformGizmo"
 
+DEFINE_LOG_CATEGORY_STATIC(LogTransformGizmo, Log, All);
+
 void UTransformGizmo::SetDisallowNegativeScaling(bool bDisallow)
 {
 	bDisallowNegativeScaling = bDisallow;
@@ -30,6 +33,14 @@ void UTransformGizmo::Setup()
 {
 	UInteractiveGizmo::Setup();
 
+	// Add default mouse hover behavior
+	UMouseHoverBehavior* HoverBehavior = NewObject<UMouseHoverBehavior>();
+	HoverBehavior->Initialize(this);
+	HoverBehavior->SetDefaultPriority(FInputCapturePriority(FInputCapturePriority::DEFAULT_GIZMO_PRIORITY));
+	AddInputBehavior(HoverBehavior);
+	
+	// @todo: Gizmo element construction will be moved to the UEditorTransformGizmoBuilder to decouple
+	// the rendered elements from the transform gizmo.
 	GizmoElementRoot = NewObject<UGizmoElementGroup>();
 	GizmoElementRoot->SetConstantScale(true);
 
@@ -106,6 +117,79 @@ void UTransformGizmo::Render(IToolsContextRenderAPI* RenderAPI)
 	}
 }
 
+FInputRayHit UTransformGizmo::BeginHoverSequenceHitTest(const FInputDeviceRay& DevicePos)
+{
+	return UpdateHoverHitSubElement(DevicePos);
+}
+
+void UTransformGizmo::OnBeginHover(const FInputDeviceRay& DevicePos)
+{
+}
+
+bool UTransformGizmo::OnUpdateHover(const FInputDeviceRay& DevicePos)
+{
+	FInputRayHit RayHit = UpdateHoverHitSubElement(DevicePos);
+	return RayHit.bHit;
+}
+
+void UTransformGizmo::OnEndHover()
+{
+	if (HitTarget && LastHitPart != ETransformGizmoPartIdentifier::Default)
+	{
+		HitTarget->UpdateHoverState(false, (static_cast<uint64>(1) << static_cast<uint8>(LastHitPart)));
+		LastHitPart = ETransformGizmoPartIdentifier::Default;
+	}
+}
+
+FInputRayHit UTransformGizmo::UpdateHoverHitSubElement(const FInputDeviceRay& PressPos)
+{
+	if (!HitTarget)
+	{
+		return FInputRayHit();
+	}
+
+	FInputRayHit RayHit = HitTarget->IsHit(PressPos);
+
+	ETransformGizmoPartIdentifier HitPart;
+	if (RayHit.bHit && VerifyPartIdentifier(RayHit.HitIdentifier))
+	{
+		HitPart = static_cast<ETransformGizmoPartIdentifier>(RayHit.HitIdentifier);
+	}
+	else
+	{
+		HitPart = ETransformGizmoPartIdentifier::Default;
+	}
+
+	if (HitPart != LastHitPart)
+	{
+		if (LastHitPart != ETransformGizmoPartIdentifier::Default)
+		{
+			HitTarget->UpdateHoverState(false, static_cast<uint32>(LastHitPart));
+		}
+
+		if (HitPart != ETransformGizmoPartIdentifier::Default)
+		{
+			HitTarget->UpdateHoverState(true, static_cast<uint32>(HitPart));
+		}
+
+		LastHitPart = HitPart;
+	}
+
+	return RayHit;
+}
+
+bool UTransformGizmo::VerifyPartIdentifier(uint32 InPartIdentifier)
+{
+	if (InPartIdentifier >= static_cast<uint32>(ETransformGizmoPartIdentifier::Max))
+	{
+		UE_LOG(LogTransformGizmo, Warning, TEXT("Unrecognized transform gizmo part identifier %d, valid identifiers are between 0-%d."), 
+			InPartIdentifier, static_cast<uint32>(ETransformGizmoPartIdentifier::Max));
+		return false;
+	}
+
+	return true;
+}
+
 void UTransformGizmo::UpdateMode()
 {
 	if (TransformSource && TransformSource->GetVisible())
@@ -161,19 +245,19 @@ void UTransformGizmo::EnableTranslate(EAxisList::Type InAxisListToDraw)
 
 	if (bEnableX && TranslateXAxisElement == nullptr)
 	{
-		TranslateXAxisElement = MakeTranslateAxis(FVector(1.0f, 0.0f, 0.0f), FVector(0.0f, 1.0f, 0.0f), AxisMaterialX);
+		TranslateXAxisElement = MakeTranslateAxis(ETransformGizmoPartIdentifier::TranslateXAxis, FVector(1.0f, 0.0f, 0.0f), FVector(0.0f, 1.0f, 0.0f), AxisMaterialX);
 		GizmoElementRoot->Add(TranslateXAxisElement);
 	}
 
 	if (bEnableY && TranslateYAxisElement == nullptr)
 	{
-		TranslateYAxisElement = MakeTranslateAxis(FVector(0.0f, 1.0f, 0.0f), FVector(0.0f, 0.0f, 1.0f), AxisMaterialY);
+		TranslateYAxisElement = MakeTranslateAxis(ETransformGizmoPartIdentifier::TranslateYAxis, FVector(0.0f, 1.0f, 0.0f), FVector(0.0f, 0.0f, 1.0f), AxisMaterialY);
 		GizmoElementRoot->Add(TranslateYAxisElement);
 	}
 
 	if (bEnableZ && TranslateZAxisElement == nullptr)
 	{
-		TranslateZAxisElement = MakeTranslateAxis(FVector(0.0f, 0.0f, 1.0f), FVector(1.0f, 0.0f, 0.0f), AxisMaterialZ);
+		TranslateZAxisElement = MakeTranslateAxis(ETransformGizmoPartIdentifier::TranslateZAxis, FVector(0.0f, 0.0f, 1.0f), FVector(1.0f, 0.0f, 0.0f), AxisMaterialZ);
 		GizmoElementRoot->Add(TranslateZAxisElement);
 	}
 
@@ -203,10 +287,10 @@ void UTransformGizmo::EnableTranslate(EAxisList::Type InAxisListToDraw)
 		TranslateScreenSpaceElement->SetEnabled(bEnableAny);
 	}
 
-	EnablePlanarObjects(bEnableX, bEnableY, bEnableZ);
+	EnablePlanarObjects(true, bEnableX, bEnableY, bEnableZ);
 }
 
-void UTransformGizmo::EnablePlanarObjects(bool bEnableX, bool bEnableY, bool bEnableZ)
+void UTransformGizmo::EnablePlanarObjects(bool bTranslate, bool bEnableX, bool bEnableY, bool bEnableZ)
 {
 	check(GizmoElementRoot);
 
@@ -220,19 +304,22 @@ void UTransformGizmo::EnablePlanarObjects(bool bEnableX, bool bEnableY, bool bEn
 
 	if (bEnableXY && PlanarXYElement == nullptr)
 	{
-		PlanarXYElement = MakePlanarHandle(XAxis, YAxis, ZAxis, TransparentVertexColorMaterial, AxisColorZ);
+		PlanarXYElement = MakePlanarHandle(bTranslate ? ETransformGizmoPartIdentifier::TranslateXYPlanar : ETransformGizmoPartIdentifier::ScaleXYPlanar, 
+											XAxis, YAxis, ZAxis, TransparentVertexColorMaterial, AxisColorZ);
 		GizmoElementRoot->Add(PlanarXYElement);
 	}
 
 	if (bEnableYZ && PlanarYZElement == nullptr)
 	{
-		PlanarYZElement = MakePlanarHandle(YAxis, ZAxis, XAxis, TransparentVertexColorMaterial, AxisColorX);
+		PlanarYZElement = MakePlanarHandle(bTranslate ? ETransformGizmoPartIdentifier::TranslateYZPlanar : ETransformGizmoPartIdentifier::ScaleYZPlanar, 
+											YAxis, ZAxis, XAxis, TransparentVertexColorMaterial, AxisColorX);
 		GizmoElementRoot->Add(PlanarYZElement);
 	}
 
 	if (bEnableXZ && PlanarXZElement == nullptr)
 	{
-		PlanarXZElement = MakePlanarHandle(ZAxis, XAxis, YAxis, TransparentVertexColorMaterial, AxisColorY);
+		PlanarXZElement = MakePlanarHandle(bTranslate ? ETransformGizmoPartIdentifier::TranslateXZPlanar : ETransformGizmoPartIdentifier::ScaleXZPlanar, 
+											ZAxis, XAxis, YAxis, TransparentVertexColorMaterial, AxisColorY);
 		GizmoElementRoot->Add(PlanarXZElement);
 	}
 
@@ -265,19 +352,19 @@ void UTransformGizmo::EnableRotate(EAxisList::Type InAxisListToDraw)
 
 	if (bEnableX && RotateXAxisElement == nullptr)
 	{
-		RotateXAxisElement = MakeRotateAxis(XAxis, YAxis, ZAxis, AxisMaterialX, CurrentAxisMaterial);
+		RotateXAxisElement = MakeRotateAxis(ETransformGizmoPartIdentifier::RotateXAxis, XAxis, YAxis, ZAxis, AxisMaterialX, CurrentAxisMaterial);
 		GizmoElementRoot->Add(RotateXAxisElement);
 	}
 
 	if (bEnableY && RotateYAxisElement == nullptr)
 	{
-		RotateYAxisElement = MakeRotateAxis(YAxis, ZAxis, XAxis, AxisMaterialY, CurrentAxisMaterial);
+		RotateYAxisElement = MakeRotateAxis(ETransformGizmoPartIdentifier::RotateYAxis, YAxis, ZAxis, XAxis, AxisMaterialY, CurrentAxisMaterial);
 		GizmoElementRoot->Add(RotateYAxisElement);
 	}
 
 	if (bEnableZ && RotateZAxisElement == nullptr)
 	{
-		RotateZAxisElement = MakeRotateAxis(ZAxis, XAxis, YAxis, AxisMaterialZ, CurrentAxisMaterial);
+		RotateZAxisElement = MakeRotateAxis(ETransformGizmoPartIdentifier::RotateZAxis, ZAxis, XAxis, YAxis, AxisMaterialZ, CurrentAxisMaterial);
 		GizmoElementRoot->Add(RotateZAxisElement);
 	}
 
@@ -285,25 +372,25 @@ void UTransformGizmo::EnableRotate(EAxisList::Type InAxisListToDraw)
 	{
 		if (RotateScreenSpaceElement == nullptr)
 		{
-			RotateScreenSpaceElement = MakeRotateCircleHandle(RotateScreenSpaceRadius, RotateScreenSpaceCircleColor, false);
+			RotateScreenSpaceElement = MakeRotateCircleHandle(ETransformGizmoPartIdentifier::RotateScreenSpace, RotateScreenSpaceRadius, RotateScreenSpaceCircleColor, false);
 			GizmoElementRoot->Add(RotateScreenSpaceElement);
 		}
 
 		if (RotateOuterCircleElement == nullptr)
 		{
-			RotateOuterCircleElement = MakeRotateCircleHandle(RotateOuterCircleRadius, RotateOuterCircleColor, false);
+			RotateOuterCircleElement = MakeRotateCircleHandle(ETransformGizmoPartIdentifier::Default, RotateOuterCircleRadius, RotateOuterCircleColor, false);
 			GizmoElementRoot->Add(RotateOuterCircleElement);
 		}
 
 		if (RotateArcballOuterElement == nullptr)
 		{
-			RotateArcballOuterElement = MakeRotateCircleHandle(RotateArcballOuterRadius, RotateArcballCircleColor, false);
+			RotateArcballOuterElement = MakeRotateCircleHandle(ETransformGizmoPartIdentifier::RotateArcball, RotateArcballOuterRadius, RotateArcballCircleColor, false);
 			GizmoElementRoot->Add(RotateArcballOuterElement);
 		}
 
 		if (RotateArcballInnerElement == nullptr)
 		{
-			RotateArcballInnerElement = MakeRotateCircleHandle(RotateArcballInnerRadius, RotateArcballCircleColor, true);
+			RotateArcballInnerElement = MakeRotateCircleHandle(ETransformGizmoPartIdentifier::RotateArcballInnerCircle, RotateArcballInnerRadius, RotateArcballCircleColor, true);
 			GizmoElementRoot->Add(RotateArcballInnerElement);
 		}
 	}
@@ -354,19 +441,19 @@ void UTransformGizmo::EnableScale(EAxisList::Type InAxisListToDraw)
 	
 	if (bEnableX && ScaleXAxisElement == nullptr)
 	{
-		ScaleXAxisElement = MakeScaleAxis(FVector(1.0f, 0.0f, 0.0f), FVector(0.0f, 1.0f, 0.0f), AxisMaterialX);
+		ScaleXAxisElement = MakeScaleAxis(ETransformGizmoPartIdentifier::ScaleXAxis, FVector(1.0f, 0.0f, 0.0f), FVector(0.0f, 1.0f, 0.0f), AxisMaterialX);
 		GizmoElementRoot->Add(ScaleXAxisElement);
 	}
 
 	if (bEnableY && ScaleYAxisElement == nullptr)
 	{
-		ScaleYAxisElement = MakeScaleAxis(FVector(0.0f, 1.0f, 0.0f), FVector(0.0f, 0.0f, 1.0f), AxisMaterialY);
+		ScaleYAxisElement = MakeScaleAxis(ETransformGizmoPartIdentifier::ScaleYAxis, FVector(0.0f, 1.0f, 0.0f), FVector(0.0f, 0.0f, 1.0f), AxisMaterialY);
 		GizmoElementRoot->Add(ScaleYAxisElement);
 	}
 
 	if (bEnableZ && ScaleZAxisElement == nullptr)
 	{
-		ScaleZAxisElement = MakeScaleAxis(FVector(0.0f, 0.0f, 1.0f), FVector(1.0f, 0.0f, 0.0f), AxisMaterialZ);
+		ScaleZAxisElement = MakeScaleAxis(ETransformGizmoPartIdentifier::ScaleZAxis, FVector(0.0f, 0.0f, 1.0f), FVector(1.0f, 0.0f, 0.0f), AxisMaterialZ);
 		GizmoElementRoot->Add(ScaleZAxisElement);
 	}
 
@@ -396,7 +483,7 @@ void UTransformGizmo::EnableScale(EAxisList::Type InAxisListToDraw)
 		ScaleUniformElement->SetEnabled(bEnableX || bEnableY || bEnableZ);
 	}
 
-	EnablePlanarObjects(bEnableX, bEnableY, bEnableZ);
+	EnablePlanarObjects(false, bEnableX, bEnableY, bEnableZ);
 }
 
 void UTransformGizmo::UpdateCameraAxisSource()
@@ -486,9 +573,10 @@ void UTransformGizmo::SetVisibility(bool bVisibleIn)
 	bVisible = bVisibleIn;
 }
 
-UGizmoElementArrow* UTransformGizmo::MakeTranslateAxis(const FVector& InAxisDir, const FVector& InSideDir, UMaterialInterface* InMaterial)
+UGizmoElementArrow* UTransformGizmo::MakeTranslateAxis(ETransformGizmoPartIdentifier InPartId, const FVector& InAxisDir, const FVector& InSideDir, UMaterialInterface* InMaterial)
 {
 	UGizmoElementArrow* ArrowElement = NewObject<UGizmoElementArrow>();
+	ArrowElement->SetPartIdentifier(static_cast<uint32>(InPartId));
 	ArrowElement->SetHeadType(EGizmoElementArrowHeadType::Cone);
 	ArrowElement->SetBase(InAxisDir * AxisLengthOffset);
 	ArrowElement->SetDirection(InAxisDir);
@@ -504,14 +592,15 @@ UGizmoElementArrow* UTransformGizmo::MakeTranslateAxis(const FVector& InAxisDir,
 	return ArrowElement;
 }
 
-UGizmoElementArrow* UTransformGizmo::MakeScaleAxis(const FVector& InAxisDir, const FVector& InSideDir, UMaterialInterface* InMaterial)
+UGizmoElementArrow* UTransformGizmo::MakeScaleAxis(ETransformGizmoPartIdentifier InPartId, const FVector& InAxisDir, const FVector& InSideDir, UMaterialInterface* InMaterial)
 {
 	UGizmoElementArrow* ArrowElement = NewObject<UGizmoElementArrow>();
+	ArrowElement->SetPartIdentifier(static_cast<uint32>(InPartId));
 	ArrowElement->SetHeadType(EGizmoElementArrowHeadType::Cube);
 	ArrowElement->SetBase(InAxisDir * AxisLengthOffset);
 	ArrowElement->SetDirection(InAxisDir);
 	ArrowElement->SetSideDirection(InSideDir);
-	ArrowElement->SetBodyLength(TranslateAxisLength);
+	ArrowElement->SetBodyLength(ScaleAxisLength);
 	ArrowElement->SetBodyRadius(AxisRadius);
 	ArrowElement->SetHeadLength(ScaleAxisCubeDim);
 	ArrowElement->SetNumSides(32);
@@ -524,6 +613,7 @@ UGizmoElementArrow* UTransformGizmo::MakeScaleAxis(const FVector& InAxisDir, con
 UGizmoElementBox* UTransformGizmo::MakeUniformScaleHandle()
 {
 	UGizmoElementBox* BoxElement = NewObject<UGizmoElementBox>();
+	BoxElement->SetPartIdentifier(static_cast<uint32>(ETransformGizmoPartIdentifier::ScaleUniform));
 	BoxElement->SetCenter(FVector::ZeroVector);
 	BoxElement->SetUpDirection(FVector::UpVector);
 	BoxElement->SetSideDirection(FVector::RightVector);
@@ -532,7 +622,7 @@ UGizmoElementBox* UTransformGizmo::MakeUniformScaleHandle()
 	return BoxElement;
 }
 
-UGizmoElementRectangle* UTransformGizmo::MakePlanarHandle(const FVector& InUpDirection, const FVector& InSideDirection, const FVector& InPlaneNormal,
+UGizmoElementRectangle* UTransformGizmo::MakePlanarHandle(ETransformGizmoPartIdentifier InPartId, const FVector& InUpDirection, const FVector& InSideDirection, const FVector& InPlaneNormal,
 	UMaterialInterface* InMaterial, const FLinearColor& InVertexColor)
 {
 	FVector PlanarHandleCenter = (InUpDirection + InSideDirection) * PlanarHandleOffset;
@@ -542,6 +632,7 @@ UGizmoElementRectangle* UTransformGizmo::MakePlanarHandle(const FVector& InUpDir
 	VertexColor.A = LargeOuterAlpha;
 
 	UGizmoElementRectangle* RectangleElement = NewObject<UGizmoElementRectangle>();
+	RectangleElement->SetPartIdentifier(static_cast<uint32>(InPartId));
 	RectangleElement->SetUpDirection(InUpDirection);
 	RectangleElement->SetSideDirection(InSideDirection);
 	RectangleElement->SetCenter(PlanarHandleCenter);
@@ -561,6 +652,7 @@ UGizmoElementRectangle* UTransformGizmo::MakePlanarHandle(const FVector& InUpDir
 UGizmoElementRectangle* UTransformGizmo::MakeTranslateScreenSpaceHandle()
 {
 	UGizmoElementRectangle* RectangleElement = NewObject<UGizmoElementRectangle>();
+	RectangleElement->SetPartIdentifier(static_cast<uint32>(ETransformGizmoPartIdentifier::TranslateScreenSpace));
 	RectangleElement->SetUpDirection(FVector::UpVector);
 	RectangleElement->SetSideDirection(FVector::RightVector);
 	RectangleElement->SetCenter(FVector::ZeroVector);
@@ -575,10 +667,11 @@ UGizmoElementRectangle* UTransformGizmo::MakeTranslateScreenSpaceHandle()
 	return RectangleElement;
 }
 
-UGizmoElementTorus* UTransformGizmo::MakeRotateAxis(const FVector& Normal, const FVector& TorusAxis0, const FVector& TorusAxis1,
+UGizmoElementTorus* UTransformGizmo::MakeRotateAxis(ETransformGizmoPartIdentifier InPartId, const FVector& Normal, const FVector& TorusAxis0, const FVector& TorusAxis1,
 	UMaterialInterface* InMaterial, UMaterialInterface* InCurrentMaterial)
 {
 	UGizmoElementTorus* RotateAxisElement = NewObject<UGizmoElementTorus>();
+	RotateAxisElement->SetPartIdentifier(static_cast<uint32>(InPartId));
 	RotateAxisElement->SetCenter(FVector::ZeroVector);
 	RotateAxisElement->SetOuterRadius(UTransformGizmo::RotateAxisOuterRadius);
 	RotateAxisElement->SetOuterSegments(UTransformGizmo::RotateAxisOuterSegments);
@@ -597,9 +690,10 @@ UGizmoElementTorus* UTransformGizmo::MakeRotateAxis(const FVector& Normal, const
 	return RotateAxisElement;
 }
 
-UGizmoElementCircle* UTransformGizmo::MakeRotateCircleHandle(float InRadius, const FLinearColor& InColor, float bFill)
+UGizmoElementCircle* UTransformGizmo::MakeRotateCircleHandle(ETransformGizmoPartIdentifier InPartId, float InRadius, const FLinearColor& InColor, float bFill)
 {
 	UGizmoElementCircle* CircleElement = NewObject<UGizmoElementCircle>();
+	CircleElement->SetPartIdentifier(static_cast<uint32>(InPartId));
 	CircleElement->SetCenter(FVector::ZeroVector);
 	CircleElement->SetRadius(InRadius);
 	CircleElement->SetNormal(-FVector::ForwardVector);
diff --git a/Engine/Source/Editor/Experimental/EditorInteractiveToolsFramework/Public/EditorGizmos/TransformGizmo.h b/Engine/Source/Editor/Experimental/EditorInteractiveToolsFramework/Public/EditorGizmos/TransformGizmo.h
--- a/Engine/Source/Editor/Experimental/EditorInteractiveToolsFramework/Public/EditorGizmos/TransformGizmo.h
+++ b/Engine/Source/Editor/Experimental/EditorInteractiveToolsFramework/Public/EditorGizmos/TransformGizmo.h
@@ -3,6 +3,7 @@
 #pragma once
 
 #include "CoreMinimal.h"
+#include "BaseBehaviors/BehaviorTargetInterfaces.h"
 #include "BaseGizmos/GizmoElementHitTargets.h"
 #include "BaseGizmos/GizmoElementStateTargets.h"
 #include "BaseGizmos/TransformProxy.h"
@@ -33,13 +34,54 @@ class UGizmoElementRectangle;
 class UGizmoElementRoot;
 class UGizmoElementTorus;
 
+//
+// Interface for the Transform gizmo.
+//
+
+
+//
+// Part identifiers are used to associate transform gizmo parts with their corresponding representation 
+// in the render and hit target. The render and hit target should use the default identifier for 
+// any of their internal elements that do not correspond to transform gizmo parts, for example non-hittable
+// visual guide elements.
+//
+UENUM()
+enum class ETransformGizmoPartIdentifier
+{
+	Default,
+	TranslateAll,
+	TranslateXAxis,
+	TranslateYAxis,
+	TranslateZAxis,
+	TranslateXYPlanar,
+	TranslateYZPlanar,
+	TranslateXZPlanar,
+	TranslateScreenSpace,
+	RotateAll,
+	RotateXAxis,
+	RotateYAxis,
+	RotateZAxis,
+	RotateScreenSpace,
+	RotateArcball,
+	RotateArcballInnerCircle,
+	ScaleAll,
+	ScaleXAxis,
+	ScaleYAxis,
+	ScaleZAxis,
+	ScaleXYPlanar,
+	ScaleYZPlanar,
+	ScaleXZPlanar,
+	ScaleUniform, 
+	Max
+};
+
 /**
  * UTransformGizmo provides standard Transformation Gizmo interactions,
  * applied to a UTransformProxy target object. By default the Gizmo will be
  * a standard XYZ translate/rotate Gizmo (axis and plane translation).
  */
 UCLASS()
-class EDITORINTERACTIVETOOLSFRAMEWORK_API UTransformGizmo : public UInteractiveGizmo
+class EDITORINTERACTIVETOOLSFRAMEWORK_API UTransformGizmo : public UInteractiveGizmo, public IHoverBehaviorTarget
 {
 	GENERATED_BODY()
 
@@ -65,8 +107,7 @@ public:
 	static constexpr float RotateOuterCircleRadius = 73.0f;
 	static constexpr float RotateScreenSpaceRadius = 83.0f;
 
-	static constexpr float ScaleAxisLength = 35.0f;
-	static constexpr float ScaleAxisCubeSize = 3.0f;
+	static constexpr float ScaleAxisLength = 70.0f;
 	static constexpr float ScaleAxisCubeDim = 12.0f;
 
 	static constexpr float PlanarHandleOffset = 55.0f;
@@ -111,6 +152,11 @@ public:
 	virtual void Render(IToolsContextRenderAPI* RenderAPI);
 	virtual void Tick(float DeltaTime) override;
 
+	// IHoverBehaviorTarget implementation
+	virtual FInputRayHit BeginHoverSequenceHitTest(const FInputDeviceRay& DevicePos) override;
+	virtual void OnBeginHover(const FInputDeviceRay& DevicePos) override;
+	virtual bool OnUpdateHover(const FInputDeviceRay& DevicePos) override;
+	virtual void OnEndHover() override;
 
 	/**
 	 * Set the active target object for the Gizmo
@@ -128,6 +174,10 @@ public:
 	UPROPERTY()
 	TObjectPtr<UTransformProxy> ActiveTarget;
 
+	/** The hit target object */
+	UPROPERTY()
+	TObjectPtr< UGizmoElementHitMultiTarget> HitTarget;
+
 	/**
 	 * Repositions the gizmo without issuing undo/redo changes, triggering callbacks, 
 	 * or moving any components. Useful for resetting the gizmo to a new location without
@@ -192,16 +242,17 @@ public:
 	UPROPERTY()
 	TScriptInterface<ITransformGizmoSource> TransformSource;
 
+
+	/** Root of renderable gizmo elements */
+	UPROPERTY()
+	TObjectPtr<UGizmoElementGroup> GizmoElementRoot;
+
 protected:
 
 	//
 	// Gizmo Objects, used for rendering and hit testing
 	//
 
-	/** Root of renderable gizmo elements */
-	UPROPERTY()
-	TObjectPtr<UGizmoElementGroup> GizmoElementRoot;
-
 	/** Translate X Axis */
 	UPROPERTY()
 	TObjectPtr<UGizmoElementArrow> TranslateXAxisElement;
@@ -318,42 +369,53 @@ protected:
 	void EnableScale(EAxisList::Type InAxisListToDraw);
 
 	/** Enable planar handles used by translate and scale */
-	void EnablePlanarObjects(bool bEnableX, bool bEnableY, bool bEnableZ);
-
-
-	UPROPERTY()
-	EGizmoTransformMode CurrentMode = EGizmoTransformMode::None;
-
-	UPROPERTY()
-	TEnumAsByte<EAxisList::Type> CurrentAxisToDraw = EAxisList::None;
+	void EnablePlanarObjects(bool bTranslate, bool bEnableX, bool bEnableY, bool bEnableZ);
 
 	/** Construct translate axis handle */
-	virtual UGizmoElementArrow* MakeTranslateAxis(const FVector& InAxisDir, const FVector& InSideDir, UMaterialInterface* InMaterial);
+	virtual UGizmoElementArrow* MakeTranslateAxis(ETransformGizmoPartIdentifier InPartId, const FVector& InAxisDir, const FVector& InSideDir, UMaterialInterface* InMaterial);
 
 	/** Construct scale axis handle */
-	virtual UGizmoElementArrow* MakeScaleAxis(const FVector& InAxisDir, const FVector& InSideDir, UMaterialInterface* InMaterial);
+	virtual UGizmoElementArrow* MakeScaleAxis(ETransformGizmoPartIdentifier InPartId, const FVector& InAxisDir, const FVector& InSideDir, UMaterialInterface* InMaterial);
 
 	/** Construct rotate axis handle */
-	virtual UGizmoElementTorus* MakeRotateAxis(const FVector& Normal, const FVector& TorusAxis0, const FVector& TorusAxis1, 
+	virtual UGizmoElementTorus* MakeRotateAxis(ETransformGizmoPartIdentifier InPartId, const FVector& Normal, const FVector& TorusAxis0, const FVector& TorusAxis1,
 		UMaterialInterface* InMaterial, UMaterialInterface* InCurrentMaterial);
 
 	/** Construct uniform scale handle */
 	virtual UGizmoElementBox* MakeUniformScaleHandle();
 
 	/** Construct planar axis handle */
-	virtual UGizmoElementRectangle* MakePlanarHandle(const FVector& InUpDirection, const FVector& InSideDirection, const FVector& InPlaneNormal,
+	virtual UGizmoElementRectangle* MakePlanarHandle(ETransformGizmoPartIdentifier InPartId, const FVector& InUpDirection, const FVector& InSideDirection, const FVector& InPlaneNormal,
 		UMaterialInterface* InMaterial, const FLinearColor& InVertexColor);
 
 	/** Construct translate screen space handle */
 	virtual UGizmoElementRectangle* MakeTranslateScreenSpaceHandle();
 
 	/** Construct rotate screen space handle */
-	virtual UGizmoElementCircle* MakeRotateCircleHandle(float InRadius, const FLinearColor& InColor, float bFill);
+	virtual UGizmoElementCircle* MakeRotateCircleHandle(ETransformGizmoPartIdentifier InPartId, float InRadius, const FLinearColor& InColor, float bFill);
+
+	// Update hover for hit part based on input device pos
+	FInputRayHit UpdateHoverHitSubElement(const FInputDeviceRay& DevicePos);
 
 	// Axis and Plane TransformSources use this function to execute worldgrid snap queries
 	bool PositionSnapFunction(const FVector& WorldPosition, FVector& SnappedPositionOut) const;
 	FQuat RotationSnapFunction(const FQuat& DeltaRotation) const;
 
+	// Verify part identifier is within recognized range of transform gizmo part ids
+	bool VerifyPartIdentifier(uint32 InPartIdentifier);
+
+	// Currently rendered transform mode
+	UPROPERTY()
+	EGizmoTransformMode CurrentMode = EGizmoTransformMode::None;
+
+	// Currently rendered axis list
+	UPROPERTY()
+	TEnumAsByte<EAxisList::Type> CurrentAxisToDraw = EAxisList::None;
+
+	// Last hit part
+	UPROPERTY()
+	ETransformGizmoPartIdentifier LastHitPart = ETransformGizmoPartIdentifier::Default;
+
 	/** Materials and colors to be used when drawing the items for each axis */
 	UPROPERTY()
 	TObjectPtr<UMaterialInterface> TransparentVertexColorMaterial;
diff --git a/Engine/Source/Runtime/InteractiveToolsFramework/Private/BaseGizmos/GizmoElementHitTargets.cpp b/Engine/Source/Runtime/InteractiveToolsFramework/Private/BaseGizmos/GizmoElementHitTargets.cpp
--- a/Engine/Source/Runtime/InteractiveToolsFramework/Private/BaseGizmos/GizmoElementHitTargets.cpp
+++ b/Engine/Source/Runtime/InteractiveToolsFramework/Private/BaseGizmos/GizmoElementHitTargets.cpp
@@ -55,5 +55,65 @@ UGizmoElementHitTarget* UGizmoElementHitTarget::Construct(UGizmoElementBase* InG
 	return NewHitTarget;
 }
 
+FInputRayHit UGizmoElementHitMultiTarget::IsHit(const FInputDeviceRay& ClickPos) const
+{
+	if (GizmoElement && (!Condition || Condition(ClickPos)))
+	{
+		return GizmoElement->LineTrace(ClickPos.WorldRay.Origin, ClickPos.WorldRay.Direction);
+	}
+	return FInputRayHit();
+}
+
+void UGizmoElementHitMultiTarget::UpdateHoverState(bool bHovering, uint32 PartIdentifier)
+{
+	if (!GizmoElement)
+	{
+		return;
+	}
+
+	if (bHovering)
+	{
+		GizmoElement->UpdatePartInteractionState(EGizmoElementInteractionState::Hovering, PartIdentifier);
+	}
+	else
+	{
+		GizmoElement->UpdatePartInteractionState(EGizmoElementInteractionState::None, PartIdentifier);
+	}
+}
+
+void UGizmoElementHitMultiTarget::UpdateInteractingState(bool bInteracting, uint32 PartIdentifier)
+{
+	if (!GizmoElement)
+	{
+		return;
+	}
+
+	if (bInteracting)
+	{
+		GizmoElement->UpdatePartInteractionState(EGizmoElementInteractionState::Interacting, PartIdentifier);
+	}
+	else
+	{
+		GizmoElement->UpdatePartInteractionState(EGizmoElementInteractionState::None, PartIdentifier);
+	}
+}
+
+void UGizmoElementHitMultiTarget::UpdateHittableState(bool bHittable, uint32 PartIdentifier)
+{
+	if (!GizmoElement)
+	{
+		return;
+	}
+
+	GizmoElement->UpdatePartHittableState(bHittable, PartIdentifier);
+}
+
+UGizmoElementHitMultiTarget* UGizmoElementHitMultiTarget::Construct(UGizmoElementBase* InGizmoElement, UObject* Outer)
+{
+	UGizmoElementHitMultiTarget* NewHitTarget = NewObject<UGizmoElementHitMultiTarget>(Outer);
+	NewHitTarget->GizmoElement = InGizmoElement;
+	return NewHitTarget;
+}
+
 
 
diff --git a/Engine/Source/Runtime/InteractiveToolsFramework/Public/BaseGizmos/GizmoElementHitTargets.h b/Engine/Source/Runtime/InteractiveToolsFramework/Public/BaseGizmos/GizmoElementHitTargets.h
--- a/Engine/Source/Runtime/InteractiveToolsFramework/Public/BaseGizmos/GizmoElementHitTargets.h
+++ b/Engine/Source/Runtime/InteractiveToolsFramework/Public/BaseGizmos/GizmoElementHitTargets.h
@@ -12,6 +12,8 @@ class UGizmoElementBase;
 /**
  * UGizmoElementHitTarget is an IGizmoClickTarget implementation that
  * hit-tests any object derived from UGizmoElementBase
+ * This hit target should be used for hitting a whole gizmo element hierarchy.
+ * Use UGizmoElementHitMultiTarget, for hit targets that support hitting parts within a gizmo element hierarchy.
  */
 UCLASS()
 class INTERACTIVETOOLSFRAMEWORK_API UGizmoElementHitTarget : public UObject, public IGizmoClickTarget
@@ -46,3 +48,50 @@ public:
 		UObject* Outer = (UObject*)GetTransientPackage());
 };
 
+/**
+ * UGizmoElementHitMultiTarget is an IGizmoClickMultiTarget implementation that
+ * hit-tests any object derived from UGizmoElementBase. This implementation is used for 
+ * HitTargets which support hitting multiple parts within a gizmo element hierarchy. 
+ *
+ * For a gizmo with multiple parts, the part identifier establishes a correspondence between a gizmo part 
+ * and the elements representing that part within the hit target. The valid part identifiers should 
+ * be defined in the gizmo. Identifier 0 is reserved for the default ID which should be assigned to 
+ * elements that do not correspond to any gizmo part, such as non-hittable decorative elements.
+ */
+UCLASS()
+class INTERACTIVETOOLSFRAMEWORK_API UGizmoElementHitMultiTarget : public UObject, public IGizmoClickMultiTarget
+{
+	GENERATED_BODY()
+public:
+
+	/**
+	 * Gizmo element.
+	 */
+	UPROPERTY()
+	TObjectPtr<UGizmoElementBase> GizmoElement;
+
+	/**
+	 * If set, this condition is checked before performing the hit test. This gives a way
+	 * to disable the hit test without hiding the component. This is useful, for instance,
+	 * in a repositionable transform gizmo in world-coordinate mode, where the rotation
+	 * components need to be hittable for movement, but not for repositioning.
+	 */
+	TFunction<bool(const FInputDeviceRay&)> Condition = nullptr;
+
+public:
+	virtual FInputRayHit IsHit(const FInputDeviceRay& ClickPos) const;
+
+	virtual void UpdateHoverState(bool bHovering, uint32 PartIdentifier);
+
+	virtual void UpdateInteractingState(bool bInteracting, uint32 PartIdentifier);
+
+	virtual void UpdateHittableState(bool bHittable, uint32 PartIdentifier);
+
+public:
+	static UGizmoElementHitMultiTarget* Construct(
+		UGizmoElementBase* InGizmoElement,
+		UObject* Outer = (UObject*)GetTransientPackage());
+
+};
+
+
