Added translation interaction to the new TRS gizmo.

diff --git a/Engine/Source/Editor/Experimental/EditorInteractiveToolsFramework/Private/EditorGizmos/EditorAxisSources.cpp b/Engine/Source/Editor/Experimental/EditorInteractiveToolsFramework/Private/EditorGizmos/EditorAxisSources.cpp
deleted file mode 100644
diff --git a/Engine/Source/Editor/Experimental/EditorInteractiveToolsFramework/Private/EditorGizmos/EditorTransformGizmo.cpp b/Engine/Source/Editor/Experimental/EditorInteractiveToolsFramework/Private/EditorGizmos/EditorTransformGizmo.cpp
new file mode 100644
--- /dev/null
+++ b/Engine/Source/Editor/Experimental/EditorInteractiveToolsFramework/Private/EditorGizmos/EditorTransformGizmo.cpp
@@ -0,0 +1,27 @@
+// Copyright Epic Games, Inc. All Rights Reserved.
+
+#include "EditorGizmos/EditorTransformGizmo.h"
+#include "EditorGizmos/EditorTransformProxy.h"
+
+#define LOCTEXT_NAMESPACE "UEditorTransformGizmo"
+
+DEFINE_LOG_CATEGORY_STATIC(LogEditorTransformGizmo, Log, All);
+
+void UEditorTransformGizmo::Translate(const FVector& InTranslateDelta)
+{
+	check(ActiveTarget);
+
+	if (UEditorTransformProxy* EditorTransformProxy = Cast<UEditorTransformProxy>(ActiveTarget))
+	{
+		EditorTransformProxy->InputTranslateDelta(InTranslateDelta, InteractionAxisType);
+
+		// Update the cached current transform
+		CurrentTransform.AddToTranslation(InTranslateDelta);
+	}
+	else
+	{
+		Super::Translate(InTranslateDelta);
+	}
+}
+
+#undef LOCTEXT_NAMESPACE
diff --git a/Engine/Source/Editor/Experimental/EditorInteractiveToolsFramework/Private/EditorGizmos/EditorTransformGizmoBuilder.cpp b/Engine/Source/Editor/Experimental/EditorInteractiveToolsFramework/Private/EditorGizmos/EditorTransformGizmoBuilder.cpp
--- a/Engine/Source/Editor/Experimental/EditorInteractiveToolsFramework/Private/EditorGizmos/EditorTransformGizmoBuilder.cpp
+++ b/Engine/Source/Editor/Experimental/EditorInteractiveToolsFramework/Private/EditorGizmos/EditorTransformGizmoBuilder.cpp
@@ -1,16 +1,16 @@
 // Copyright Epic Games, Inc. All Rights Reserved.
 
 #include "EditorGizmos/EditorTransformGizmoBuilder.h"
+#include "EditorGizmos/EditorTransformGizmo.h"
 #include "EditorGizmos/EditorTransformGizmoSource.h"
 #include "EditorGizmos/EditorTransformProxy.h"
-#include "EditorGizmos/TransformGizmo.h"
 #include "BaseGizmos/GizmoElementGroup.h"
 
 UInteractiveGizmo* UEditorTransformGizmoBuilder::BuildGizmo(const FToolBuilderState& SceneState) const
 {
-	UTransformGizmo* TransformGizmo = NewObject<UTransformGizmo>(SceneState.GizmoManager);
+	UEditorTransformGizmo* TransformGizmo = NewObject<UEditorTransformGizmo>(SceneState.GizmoManager);
 	TransformGizmo->Setup();
-	TransformGizmo->TransformSource = UEditorTransformGizmoSource::Construct(TransformGizmo);
+	TransformGizmo->TransformGizmoSource = UEditorTransformGizmoSource::Construct(TransformGizmo);
 
 	// @todo: Gizmo element construction to be moved here from UTransformGizmo.
 	// A UGizmoElementRenderMultiTarget will be constructed and both the
diff --git a/Engine/Source/Editor/Experimental/EditorInteractiveToolsFramework/Private/EditorGizmos/EditorTransformGizmoSource.cpp b/Engine/Source/Editor/Experimental/EditorInteractiveToolsFramework/Private/EditorGizmos/EditorTransformGizmoSource.cpp
--- a/Engine/Source/Editor/Experimental/EditorInteractiveToolsFramework/Private/EditorGizmos/EditorTransformGizmoSource.cpp
+++ b/Engine/Source/Editor/Experimental/EditorInteractiveToolsFramework/Private/EditorGizmos/EditorTransformGizmoSource.cpp
@@ -4,12 +4,34 @@
 #include "EditorModeManager.h"
 #include "EditorViewportClient.h"
 
+EGizmoTransformMode FEditorTransformGizmoUtil::GetGizmoMode(UE::Widget::EWidgetMode InWidgetMode)
+{
+	switch (InWidgetMode)
+	{
+		case UE::Widget::EWidgetMode::WM_Translate: return EGizmoTransformMode::Translate;
+		case UE::Widget::EWidgetMode::WM_Rotate: return EGizmoTransformMode::Rotate;
+		case UE::Widget::EWidgetMode::WM_Scale: return EGizmoTransformMode::Scale;
+	}
+	return EGizmoTransformMode::None;
+}
+
+UE::Widget::EWidgetMode FEditorTransformGizmoUtil::GetWidgetMode(EGizmoTransformMode InGizmoMode)
+{
+	switch (InGizmoMode)
+	{
+		case EGizmoTransformMode::Translate: return UE::Widget::EWidgetMode::WM_Translate;
+		case EGizmoTransformMode::Rotate: return UE::Widget::EWidgetMode::WM_Rotate;
+		case EGizmoTransformMode::Scale: return UE::Widget::EWidgetMode::WM_Scale;
+	}
+	return UE::Widget::EWidgetMode::WM_None;
+}
+
 EGizmoTransformMode UEditorTransformGizmoSource::GetGizmoMode() const
 {
 	if (FEditorViewportClient* ViewportClient = GetViewportClient())
 	{
 		UE::Widget::EWidgetMode WidgetMode = ViewportClient->GetWidgetMode();
-		return FTransformGizmoUtil::GetGizmoMode(WidgetMode);
+		return FEditorTransformGizmoUtil::GetGizmoMode(WidgetMode);
 	}
 	return EGizmoTransformMode::None;
 }
diff --git a/Engine/Source/Editor/Experimental/EditorInteractiveToolsFramework/Private/EditorGizmos/EditorTransformProxy.cpp b/Engine/Source/Editor/Experimental/EditorInteractiveToolsFramework/Private/EditorGizmos/EditorTransformProxy.cpp
--- a/Engine/Source/Editor/Experimental/EditorInteractiveToolsFramework/Private/EditorGizmos/EditorTransformProxy.cpp
+++ b/Engine/Source/Editor/Experimental/EditorInteractiveToolsFramework/Private/EditorGizmos/EditorTransformProxy.cpp
@@ -15,14 +15,46 @@ FTransform UEditorTransformProxy::GetTransform() const
 		FMatrix RotMatrix = ViewportClient->GetWidgetCoordSystem();
 		return FTransform(FQuat(RotMatrix), Location, FVector::OneVector);
 	}
-	
-	return FTransform::Identity;
+	else
+	{
+		return FTransform::Identity;
+	}
+}
+
+void UEditorTransformProxy::InputTranslateDelta(const FVector& InDeltaTranslate, EAxisList::Type InAxisList)
+{
+	if (FEditorViewportClient* ViewportClient = GLevelEditorModeTools().GetFocusedViewportClient())
+	{
+		FVector Translate = InDeltaTranslate;
+		FRotator Rot = FRotator::ZeroRotator;
+		FVector Scale = FVector::ZeroVector;
+
+		ViewportClient->InputWidgetDelta(ViewportClient->Viewport, InAxisList, Translate, Rot, Scale);
+	}
+}
+
+void UEditorTransformProxy::InputScaleDelta(const FVector& InDeltaScale, EAxisList::Type InAxisList)
+{
+	if (FEditorViewportClient* ViewportClient = GLevelEditorModeTools().GetFocusedViewportClient())
+	{
+		FVector Translate = FVector::ZeroVector;
+		FRotator Rot = FRotator::ZeroRotator;
+		FVector Scale = InDeltaScale;
+
+		ViewportClient->InputWidgetDelta(ViewportClient->Viewport, InAxisList, Translate, Rot, Scale);
+	}
 }
 
-void UEditorTransformProxy::SetTransform(const FTransform& Transform)
+void UEditorTransformProxy::InputRotateDelta(const FRotator& InDeltaRotate, EAxisList::Type InAxisList)
 {
-	// @todo: use ViewportClient->HandleWidgetDelta();
-	// @todo: handle Pivots, need to look at how this is currently handled for component visualizers, for example.
+	if (FEditorViewportClient* ViewportClient = GLevelEditorModeTools().GetFocusedViewportClient())
+	{
+		FVector Translate = FVector::ZeroVector;
+		FRotator Rot = InDeltaRotate;
+		FVector Scale = FVector::ZeroVector;
+
+		ViewportClient->InputWidgetDelta(ViewportClient->Viewport, InAxisList, Translate, Rot, Scale);
+	}
 }
 
 #undef LOCTEXT_NAMESPACE
diff --git a/Engine/Source/Editor/Experimental/EditorInteractiveToolsFramework/Private/EditorGizmos/TransformGizmo.cpp b/Engine/Source/Editor/Experimental/EditorInteractiveToolsFramework/Private/EditorGizmos/TransformGizmo.cpp
--- a/Engine/Source/Editor/Experimental/EditorInteractiveToolsFramework/Private/EditorGizmos/TransformGizmo.cpp
+++ b/Engine/Source/Editor/Experimental/EditorInteractiveToolsFramework/Private/EditorGizmos/TransformGizmo.cpp
@@ -1,16 +1,15 @@
 // Copyright Epic Games, Inc. All Rights Reserved.
 
 #include "EditorGizmos/TransformGizmo.h"
+#include "BaseBehaviors/ClickDragBehavior.h"
 #include "BaseBehaviors/MouseHoverBehavior.h"
 #include "BaseGizmos/AxisSources.h"
 #include "BaseGizmos/GizmoElementGroup.h"
 #include "BaseGizmos/GizmoElementShapes.h"
+#include "BaseGizmos/GizmoMath.h"
 #include "BaseGizmos/GizmoRenderingUtil.h"
-#include "BaseGizmos/TransformSources.h"
-#include "EditorGizmos/EditorAxisSources.h"
-#include "EditorGizmos/EditorTransformGizmoSource.h"
-#include "EditorGizmos/EditorTransformProxy.h"
-#include "EditorGizmos/EditorParameterToTransformAdapters.h"
+#include "BaseGizmos/ParameterSourcesFloat.h"
+#include "BaseGizmos/StateTargets.h"
 #include "Elements/Interfaces/TypedElementObjectInterface.h"
 #include "Elements/Interfaces/TypedElementWorldInterface.h"
 #include "Engine/CollisionProfile.h"
@@ -33,17 +32,36 @@ void UTransformGizmo::Setup()
 {
 	UInteractiveGizmo::Setup();
 
+	SetupBehaviors();
+	SetupMaterials();
+
+	// @todo: Gizmo element construction will be moved to the UEditorTransformGizmoBuilder to decouple
+	// the rendered elements from the transform gizmo.
+	GizmoElementRoot = NewObject<UGizmoElementGroup>();
+	GizmoElementRoot->SetConstantScale(true);
+	GizmoElementRoot->SetHoverMaterial(CurrentAxisMaterial);
+	GizmoElementRoot->SetInteractMaterial(CurrentAxisMaterial);
+
+	bInInteraction = false;
+}
+
+void UTransformGizmo::SetupBehaviors()
+{
 	// Add default mouse hover behavior
 	UMouseHoverBehavior* HoverBehavior = NewObject<UMouseHoverBehavior>();
 	HoverBehavior->Initialize(this);
 	HoverBehavior->SetDefaultPriority(FInputCapturePriority(FInputCapturePriority::DEFAULT_GIZMO_PRIORITY));
 	AddInputBehavior(HoverBehavior);
-	
-	// @todo: Gizmo element construction will be moved to the UEditorTransformGizmoBuilder to decouple
-	// the rendered elements from the transform gizmo.
-	GizmoElementRoot = NewObject<UGizmoElementGroup>();
-	GizmoElementRoot->SetConstantScale(true);
 
+	// Add default mouse input behavior
+	MouseBehavior = NewObject<UClickDragInputBehavior>();
+	MouseBehavior->Initialize(this);
+	MouseBehavior->SetDefaultPriority(FInputCapturePriority(FInputCapturePriority::DEFAULT_GIZMO_PRIORITY));
+	AddInputBehavior(MouseBehavior);
+}
+
+void UTransformGizmo::SetupMaterials()
+{
 	UMaterial* AxisMaterialBase = GEngine->ArrowMaterial;
 
 	AxisMaterialX = UMaterialInstanceDynamic::Create(AxisMaterialBase, NULL);
@@ -79,39 +97,36 @@ void UTransformGizmo::Setup()
 	{
 		GridMaterial = TransparentVertexColorMaterial;
 	}
-
-	GizmoElementRoot->SetHoverMaterial(CurrentAxisMaterial);
-	GizmoElementRoot->SetInteractMaterial(CurrentAxisMaterial);
 }
 
-
 void UTransformGizmo::Shutdown()
 {
 	ClearActiveTarget();
 }
 
-void UTransformGizmo::Render(IToolsContextRenderAPI* RenderAPI)
+FTransform UTransformGizmo::GetGizmoTransform() const
 {
-	if (bVisible && GizmoElementRoot && RenderAPI)
+	float Scale = 1.0f;
+
+	if (TransformGizmoSource)
 	{
-		EToolContextCoordinateSystem Space = EToolContextCoordinateSystem::World;
-		float Scale = 1.0f;
+		Scale = TransformGizmoSource->GetGizmoScale();
+	}
 
-		if (TransformSource)
-		{
-			Space = TransformSource->GetGizmoCoordSystemSpace();
-			Scale = TransformSource->GetGizmoScale();
-		}
+	FTransform GizmoLocalToWorldTransform = CurrentTransform;
+	GizmoLocalToWorldTransform.SetScale3D(FVector(Scale, Scale, Scale));
 
-		FTransform LocalToWorldTransform = ActiveTarget->GetTransform();
-		if (Space == EToolContextCoordinateSystem::World)
-		{
-			LocalToWorldTransform.SetRotation(FQuat::Identity);
-		}
-		LocalToWorldTransform.SetScale3D(FVector(Scale, Scale, Scale));
+	return GizmoLocalToWorldTransform;
+}
+
+void UTransformGizmo::Render(IToolsContextRenderAPI* RenderAPI)
+{
+	if (bVisible && GizmoElementRoot && RenderAPI)
+	{
+		CurrentTransform = ActiveTarget->GetTransform();
 
 		UGizmoElementBase::FRenderTraversalState RenderState;
-		RenderState.LocalToWorldTransform = LocalToWorldTransform;
+		RenderState.LocalToWorldTransform = GetGizmoTransform();
 
 		GizmoElementRoot->Render(RenderAPI, RenderState);
 	}
@@ -119,7 +134,7 @@ void UTransformGizmo::Render(IToolsContextRenderAPI* RenderAPI)
 
 FInputRayHit UTransformGizmo::BeginHoverSequenceHitTest(const FInputDeviceRay& DevicePos)
 {
-	return UpdateHoverHitSubElement(DevicePos);
+	return UpdateHoveredPart(DevicePos);
 }
 
 void UTransformGizmo::OnBeginHover(const FInputDeviceRay& DevicePos)
@@ -128,7 +143,7 @@ void UTransformGizmo::OnBeginHover(const FInputDeviceRay& DevicePos)
 
 bool UTransformGizmo::OnUpdateHover(const FInputDeviceRay& DevicePos)
 {
-	FInputRayHit RayHit = UpdateHoverHitSubElement(DevicePos);
+	FInputRayHit RayHit = UpdateHoveredPart(DevicePos);
 	return RayHit.bHit;
 }
 
@@ -136,12 +151,11 @@ void UTransformGizmo::OnEndHover()
 {
 	if (HitTarget && LastHitPart != ETransformGizmoPartIdentifier::Default)
 	{
-		HitTarget->UpdateHoverState(false, (static_cast<uint64>(1) << static_cast<uint8>(LastHitPart)));
-		LastHitPart = ETransformGizmoPartIdentifier::Default;
+		HitTarget->UpdateHoverState(false, static_cast<uint32>(LastHitPart));
 	}
 }
 
-FInputRayHit UTransformGizmo::UpdateHoverHitSubElement(const FInputDeviceRay& PressPos)
+FInputRayHit UTransformGizmo::UpdateHoveredPart(const FInputDeviceRay& PressPos)
 {
 	if (!HitTarget)
 	{
@@ -178,24 +192,157 @@ FInputRayHit UTransformGizmo::UpdateHoverHitSubElement(const FInputDeviceRay& Pr
 	return RayHit;
 }
 
-bool UTransformGizmo::VerifyPartIdentifier(uint32 InPartIdentifier)
+uint32 UTransformGizmo::GetMaxPartIdentifier() const
 {
-	if (InPartIdentifier >= static_cast<uint32>(ETransformGizmoPartIdentifier::Max))
+	return static_cast<uint32>(ETransformGizmoPartIdentifier::Max);
+}
+
+bool UTransformGizmo::VerifyPartIdentifier(uint32 InPartIdentifier) const
+{
+	if (InPartIdentifier >= GetMaxPartIdentifier())
 	{
 		UE_LOG(LogTransformGizmo, Warning, TEXT("Unrecognized transform gizmo part identifier %d, valid identifiers are between 0-%d."), 
-			InPartIdentifier, static_cast<uint32>(ETransformGizmoPartIdentifier::Max));
+			InPartIdentifier, GetMaxPartIdentifier());
 		return false;
 	}
 
 	return true;
 }
 
+
+
+FInputRayHit UTransformGizmo::CanBeginClickDragSequence(const FInputDeviceRay& PressPos)
+{
+	FInputRayHit RayHit;
+
+	if (HitTarget)
+	{
+		RayHit = HitTarget->IsHit(PressPos);
+		ETransformGizmoPartIdentifier HitPart;
+		if (RayHit.bHit && VerifyPartIdentifier(RayHit.HitIdentifier))
+		{
+			HitPart = static_cast<ETransformGizmoPartIdentifier>(RayHit.HitIdentifier);
+		}
+		else
+		{
+			HitPart = ETransformGizmoPartIdentifier::Default;
+		}
+
+		if (HitPart != ETransformGizmoPartIdentifier::Default)
+		{
+			LastHitPart = static_cast<ETransformGizmoPartIdentifier>(RayHit.HitIdentifier);
+		}
+	}
+
+	return RayHit;
+}
+
+void UTransformGizmo::OnClickPress(const FInputDeviceRay& PressPos)
+{
+	if (LastHitPart == ETransformGizmoPartIdentifier::TranslateXAxis)
+	{
+		InteractionAxisOrigin = CurrentTransform.GetLocation();
+		InteractionAxis = GetWorldAxis(FVector::XAxisVector);
+		InteractionAxisType = EAxisList::X;
+		OnClickPressTranslate(PressPos);
+	}
+	else if (LastHitPart == ETransformGizmoPartIdentifier::TranslateYAxis)
+	{
+		InteractionAxisOrigin = CurrentTransform.GetLocation();
+		InteractionAxis = GetWorldAxis(FVector::YAxisVector);
+		InteractionAxisType = EAxisList::Y;
+		OnClickPressTranslate(PressPos);
+	}
+	else if (LastHitPart == ETransformGizmoPartIdentifier::TranslateZAxis)
+	{
+		InteractionAxisOrigin = CurrentTransform.GetLocation();
+		InteractionAxis = GetWorldAxis(FVector::ZAxisVector);
+		InteractionAxisType = EAxisList::Z;
+		OnClickPressTranslate(PressPos);
+	}
+
+	if (bInInteraction)
+	{
+		if (HitTarget&& LastHitPart != ETransformGizmoPartIdentifier::Default)
+		{
+			HitTarget->UpdateInteractingState(true, static_cast<uint32>(LastHitPart));
+		}
+
+		if (StateTarget)
+		{
+			StateTarget->BeginUpdate();
+		}
+	}
+}
+
+void UTransformGizmo::OnClickDrag(const FInputDeviceRay& DragPos)
+{
+	if (!bInInteraction)
+	{
+		return;
+	}
+
+	if (LastHitPart == ETransformGizmoPartIdentifier::TranslateXAxis ||
+		LastHitPart == ETransformGizmoPartIdentifier::TranslateYAxis ||
+		LastHitPart == ETransformGizmoPartIdentifier::TranslateZAxis)
+	{
+		OnClickDragTranslate(DragPos);
+	}
+}
+
+void UTransformGizmo::OnClickRelease(const FInputDeviceRay& ReleasePos)
+{
+	if (!bInInteraction)
+	{
+		return;
+	}
+
+	if (LastHitPart == ETransformGizmoPartIdentifier::TranslateXAxis ||
+		LastHitPart == ETransformGizmoPartIdentifier::TranslateYAxis ||
+		LastHitPart == ETransformGizmoPartIdentifier::TranslateZAxis)
+	{
+		OnClickReleaseTranslate(ReleasePos);
+	}
+
+	if (StateTarget)
+	{
+		StateTarget->EndUpdate();
+	}
+
+	bInInteraction = false;
+
+	if (HitTarget && LastHitPart != ETransformGizmoPartIdentifier::Default)
+	{
+		HitTarget->UpdateInteractingState(false, static_cast<uint32>(LastHitPart));
+	}
+}
+
+
+void UTransformGizmo::OnTerminateDragSequence()
+{
+	if (!bInInteraction)
+	{
+		return;
+	}
+
+	if (StateTarget)
+	{
+		StateTarget->EndUpdate();
+	}
+	bInInteraction = false;
+
+	if (HitTarget && LastHitPart != ETransformGizmoPartIdentifier::Default)
+	{
+		HitTarget->UpdateInteractingState(false, static_cast<uint32>(LastHitPart));
+	}
+}
+
 void UTransformGizmo::UpdateMode()
 {
-	if (TransformSource && TransformSource->GetVisible())
+	if (TransformGizmoSource && TransformGizmoSource->GetVisible())
 	{
-		EGizmoTransformMode NewMode = TransformSource->GetGizmoMode();
-		EAxisList::Type NewAxisToDraw = TransformSource->GetGizmoAxisToDraw(NewMode);
+		EGizmoTransformMode NewMode = TransformGizmoSource->GetGizmoMode();
+		EAxisList::Type NewAxisToDraw = TransformGizmoSource->GetGizmoAxisToDraw(NewMode);
 
 		if (NewMode != CurrentMode)
 		{
@@ -529,35 +676,23 @@ void UTransformGizmo::SetActiveTarget(UTransformProxy* Target, IToolContextTrans
 	{
 		TransactionProvider = GetGizmoManager();
 	}
+	
+	/* @todo - Add state target 
+	
+	// This state target emits an explicit FChange that moves the GizmoActor root component during undo/redo.
+	// It also opens/closes the Transaction that saves/restores the target object locations.
+	if (TransactionProvider == nullptr)
+	{
+		TransactionProvider = GetGizmoManager();
+	}
+	StateTarget = UGizmoTransformChangeStateTarget::Construct(GizmoElementRoot,
+		LOCTEXT("UCombinedTransformGizmoTransaction", "Transform"), TransactionProvider, this);
+	StateTarget->DependentChangeSources.Add(MakeUnique<FTransformProxyChangeSource>(Target));
+	*/
 
 	CameraAxisSource = NewObject<UGizmoConstantFrameAxisSource>(this);
 }
 
-
-void UTransformGizmo::ReinitializeGizmoTransform(const FTransform& NewTransform)
-{
-	// @todo update gizmo objects here?
-
-	// The underlying proxy has an existing way to reinitialize its transform without callbacks.
-	TGuardValue<bool>(ActiveTarget->bSetPivotMode, true);
-	ActiveTarget->SetTransform(NewTransform);
-}
-
-
-void UTransformGizmo::SetNewGizmoTransform(const FTransform& NewTransform)
-{
-	// @todo update gizmo objects here?
-
-	check(ActiveTarget != nullptr);
-
-	StateTarget->BeginUpdate();
-
-	ActiveTarget->SetTransform(NewTransform);
-
-	StateTarget->EndUpdate();
-}
-
-
 // @todo: This should either be named to "SetScale" or removed, since it can be done with ReinitializeGizmoTransform
 void UTransformGizmo::SetNewChildScale(const FVector& NewChildScale)
 {
@@ -785,4 +920,55 @@ FQuat UTransformGizmo::RotationSnapFunction(const FQuat& DeltaRotation) const
 	return SnappedDeltaRotation;
 }
 
+FVector UTransformGizmo::GetWorldAxis(const FVector& InAxis)
+{
+	if (TransformGizmoSource->GetGizmoCoordSystemSpace() == EToolContextCoordinateSystem::Local)
+	{
+		return CurrentTransform.GetRotation().RotateVector(InAxis);
+	}
+	
+	return InAxis;
+}
+
+void UTransformGizmo::OnClickPressTranslate(const FInputDeviceRay& InPressPos)
+{
+	// Find interaction start point and parameter.
+	FVector NearestPt; 
+	float RayNearestParam, LineNearestParam;
+	GizmoMath::NearestPointOnLineToRay(InteractionAxisOrigin, InteractionAxis,
+		InPressPos.WorldRay.Origin, InPressPos.WorldRay.Direction,
+		InteractionStartPoint, LineNearestParam,
+		NearestPt, RayNearestParam);
+
+	InteractionCurrPoint = InteractionStartPoint;
+	bInInteraction = true;
+}
+
+void UTransformGizmo::OnClickDragTranslate(const FInputDeviceRay& InDragPos)
+{
+	float RayNearestParam, AxisNearestParam;
+	FVector RayNearestPt, AxisNearestPoint;
+	GizmoMath::NearestPointOnLineToRay(InteractionAxisOrigin, InteractionAxis,
+		InDragPos.WorldRay.Origin, InDragPos.WorldRay.Direction,
+		AxisNearestPoint, AxisNearestParam,
+		RayNearestPt, RayNearestParam);
+
+	FVector TranslateDelta = AxisNearestPoint - InteractionCurrPoint;
+	InteractionCurrPoint = AxisNearestPoint;
+
+	Translate(TranslateDelta);
+}
+
+void UTransformGizmo::OnClickReleaseTranslate(const FInputDeviceRay& InReleasePos)
+{
+	bInInteraction = false;
+}
+
+void UTransformGizmo::Translate(const FVector& InTranslateDelta)
+{
+	CurrentTransform.AddToTranslation(InTranslateDelta);
+	ActiveTarget->SetTransform(CurrentTransform);
+}
+
+
 #undef LOCTEXT_NAMESPACE
diff --git a/Engine/Source/Editor/Experimental/EditorInteractiveToolsFramework/Private/EditorGizmos/TransformGizmoInterfaces.cpp b/Engine/Source/Editor/Experimental/EditorInteractiveToolsFramework/Private/EditorGizmos/TransformGizmoInterfaces.cpp
deleted file mode 100644
diff --git a/Engine/Source/Editor/Experimental/EditorInteractiveToolsFramework/Public/EditorGizmos/EditorAxisSources.h b/Engine/Source/Editor/Experimental/EditorInteractiveToolsFramework/Public/EditorGizmos/EditorAxisSources.h
deleted file mode 100644
diff --git a/Engine/Source/Editor/Experimental/EditorInteractiveToolsFramework/Public/EditorGizmos/EditorParameterToTransformAdapters.h b/Engine/Source/Editor/Experimental/EditorInteractiveToolsFramework/Public/EditorGizmos/EditorParameterToTransformAdapters.h
deleted file mode 100644
diff --git a/Engine/Source/Editor/Experimental/EditorInteractiveToolsFramework/Public/EditorGizmos/EditorTransformGizmo.h b/Engine/Source/Editor/Experimental/EditorInteractiveToolsFramework/Public/EditorGizmos/EditorTransformGizmo.h
new file mode 100644
--- /dev/null
+++ b/Engine/Source/Editor/Experimental/EditorInteractiveToolsFramework/Public/EditorGizmos/EditorTransformGizmo.h
@@ -0,0 +1,22 @@
+// Copyright Epic Games, Inc. All Rights Reserved.
+
+#pragma once
+
+#include "CoreMinimal.h"
+#include "EditorGizmos/TransformGizmo.h"
+#include "EditorTransformGizmo.generated.h"
+
+/**
+ * UEditorTransformGizmo handles Editor-specific functionality for the TransformGizmo,
+ * applied to a UEditorTransformProxy target object.
+ */
+UCLASS()
+class EDITORINTERACTIVETOOLSFRAMEWORK_API UEditorTransformGizmo : public UTransformGizmo
+{
+	GENERATED_BODY()
+
+protected:
+
+	/** Update current gizmo mode based on transform source */
+	virtual void Translate(const FVector& InTranslateDelta) override;
+};
\ No newline at end of file
diff --git a/Engine/Source/Editor/Experimental/EditorInteractiveToolsFramework/Public/EditorGizmos/EditorTransformGizmoSource.h b/Engine/Source/Editor/Experimental/EditorInteractiveToolsFramework/Public/EditorGizmos/EditorTransformGizmoSource.h
--- a/Engine/Source/Editor/Experimental/EditorInteractiveToolsFramework/Public/EditorGizmos/EditorTransformGizmoSource.h
+++ b/Engine/Source/Editor/Experimental/EditorInteractiveToolsFramework/Public/EditorGizmos/EditorTransformGizmoSource.h
@@ -10,6 +10,15 @@ class FEditorViewportClient;
 class FSceneView;
 class FEditorModeTools;
 
+namespace FEditorTransformGizmoUtil
+{
+	/** Convert UE::Widget::EWidgetMode to ETransformGizmoMode*/
+	EGizmoTransformMode GetGizmoMode(UE::Widget::EWidgetMode InWidgetMode);
+
+	/** Convert EEditorGizmoMode to UE::Widget::EWidgetMode*/
+	UE::Widget::EWidgetMode GetWidgetMode(EGizmoTransformMode InGizmoMode);
+};
+
 /**
  * UEditorTransformGizmoSource is an ITransformGizmoSource implementation that provides
  * current state information used to configure the Editor transform gizmo.
diff --git a/Engine/Source/Editor/Experimental/EditorInteractiveToolsFramework/Public/EditorGizmos/EditorTransformProxy.h b/Engine/Source/Editor/Experimental/EditorInteractiveToolsFramework/Public/EditorGizmos/EditorTransformProxy.h
--- a/Engine/Source/Editor/Experimental/EditorInteractiveToolsFramework/Public/EditorGizmos/EditorTransformProxy.h
+++ b/Engine/Source/Editor/Experimental/EditorInteractiveToolsFramework/Public/EditorGizmos/EditorTransformProxy.h
@@ -23,14 +23,25 @@ class EDITORINTERACTIVETOOLSFRAMEWORK_API UEditorTransformProxy : public UTransf
 public:
 
 	/**
-	 * @return the shared transform for all the sub-objects
+	 * @return the stored transform for currently selected objects.
 	 */
 	virtual FTransform GetTransform() const override;
 
 	/**
-	 * Update the main transform and then update the sub-objects based on their relative transformations
+	 * Unimplemented - all updates to the Editor transform proxy MUST be made by calling the Input delta methods.
 	 */
-	virtual void SetTransform(const FTransform& Transform) override;
+	virtual void SetTransform(const FTransform& Transform) override 
+	{
+		check(false);
+	}
 
+	/** Input translate delta to be applied in world space of the current transform. */
+	virtual void InputTranslateDelta(const FVector& InDeltaTranslate, EAxisList::Type InAxisList);
+
+	/** Input rotation delta to be applied in local space of the current transform. */
+	virtual void InputRotateDelta(const FRotator& InDeltaRotate, EAxisList::Type InAxisList);
+
+	/** Input scale delta to be applied in local space of the current transform. */
+	virtual void InputScaleDelta(const FVector& InDeltaScale, EAxisList::Type InAxisList);
 };
 
diff --git a/Engine/Source/Editor/Experimental/EditorInteractiveToolsFramework/Public/EditorGizmos/TransformGizmo.h b/Engine/Source/Editor/Experimental/EditorInteractiveToolsFramework/Public/EditorGizmos/TransformGizmo.h
--- a/Engine/Source/Editor/Experimental/EditorInteractiveToolsFramework/Public/EditorGizmos/TransformGizmo.h
+++ b/Engine/Source/Editor/Experimental/EditorInteractiveToolsFramework/Public/EditorGizmos/TransformGizmo.h
@@ -7,7 +7,6 @@
 #include "BaseGizmos/GizmoElementHitTargets.h"
 #include "BaseGizmos/GizmoElementStateTargets.h"
 #include "BaseGizmos/TransformProxy.h"
-#include "EditorGizmos/EditorAxisSources.h"
 #include "EditorGizmos/TransformGizmoInterfaces.h"
 #include "InteractiveGizmo.h"
 #include "InteractiveToolObjects.h"
@@ -20,9 +19,8 @@ class UInteractiveGizmoManager;
 class IGizmoAxisSource;
 class IGizmoTransformSource;
 class IGizmoStateTarget;
-class FEditorTransformGizmoTransformChange;
+class UClickDragInputBehavior;
 class UGizmoConstantFrameAxisSource;
-class UGizmoTransformChangeStateTarget;
 class UGizmoElementArrow;
 class UGizmoElementBase;
 class UGizmoElementBox;
@@ -81,7 +79,7 @@ enum class ETransformGizmoPartIdentifier
  * a standard XYZ translate/rotate Gizmo (axis and plane translation).
  */
 UCLASS()
-class EDITORINTERACTIVETOOLSFRAMEWORK_API UTransformGizmo : public UInteractiveGizmo, public IHoverBehaviorTarget
+class EDITORINTERACTIVETOOLSFRAMEWORK_API UTransformGizmo : public UInteractiveGizmo, public IHoverBehaviorTarget, public IClickDragBehaviorTarget
 {
 	GENERATED_BODY()
 
@@ -158,6 +156,13 @@ public:
 	virtual bool OnUpdateHover(const FInputDeviceRay& DevicePos) override;
 	virtual void OnEndHover() override;
 
+	// IClickDragBehaviorTarget implementation
+	virtual FInputRayHit CanBeginClickDragSequence(const FInputDeviceRay& PressPos) override;
+	virtual void OnClickPress(const FInputDeviceRay& PressPos) override;
+	virtual void OnClickDrag(const FInputDeviceRay& DragPos) override;
+	virtual void OnClickRelease(const FInputDeviceRay& ReleasePos) override;
+	virtual void OnTerminateDragSequence() override;
+
 	/**
 	 * Set the active target object for the Gizmo
 	 * @param Target active target
@@ -170,27 +175,6 @@ public:
 	 */
 	virtual void ClearActiveTarget();
 
-	/** The active target object for the Gizmo */
-	UPROPERTY()
-	TObjectPtr<UTransformProxy> ActiveTarget;
-
-	/** The hit target object */
-	UPROPERTY()
-	TObjectPtr< UGizmoElementHitMultiTarget> HitTarget;
-
-	/**
-	 * Repositions the gizmo without issuing undo/redo changes, triggering callbacks, 
-	 * or moving any components. Useful for resetting the gizmo to a new location without
-	 * it being viewed as a gizmo manipulation.
-	 */
-	void ReinitializeGizmoTransform(const FTransform& NewTransform);
-
-	/**
-	 * Set a new position for the Gizmo. This is done via the same mechanisms as the sub-gizmos,
-	 * so it generates the same Change/Modify() events, and hence works with Undo/Redo
-	 */
-	virtual void SetNewGizmoTransform(const FTransform& NewTransform);
-
 	/**
 	 * Explicitly set the child scale. Mainly useful to "reset" the child scale to (1,1,1) when re-using Gizmo across multiple transform actions.
 	 * @warning does not generate change/modify events!!
@@ -202,15 +186,37 @@ public:
 	 */
 	virtual void SetVisibility(bool bVisible);
 
-	/**
-	 * Whether gizmo is visible.
-	 */
+public:
+
+	/** The active target object for the Gizmo */
+	UPROPERTY()
+	TObjectPtr<UTransformProxy> ActiveTarget;
+
+	/** The hit target object */
+	UPROPERTY()
+	TObjectPtr<UGizmoElementHitMultiTarget> HitTarget;
+
+	/** The mouse click behavior of the gizmo is accessible so that it can be modified to use different mouse keys. */
+	UPROPERTY()
+	TObjectPtr<UClickDragInputBehavior> MouseBehavior;
+
+	/** Transform Gizmo Source */
+	UPROPERTY()
+	TScriptInterface<ITransformGizmoSource> TransformGizmoSource;
+
+	/** Root of renderable gizmo elements */
+	UPROPERTY()
+	TObjectPtr<UGizmoElementGroup> GizmoElementRoot;
+
+	/** Whether gizmo is visible. */
 	UPROPERTY()
 	bool bVisible = false;
 
-	/**
-	 * If true, then when using world frame, Axis and Plane translation snap to the world grid via the ContextQueriesAPI (in PositionSnapFunction)
-	 */
+	/** Whether gizmo is interacting. */
+	UPROPERTY()
+	bool bInInteraction = false;
+
+	/** If true, then when using world frame, Axis and Plane translation snap to the world grid via the ContextQueriesAPI (in PositionSnapFunction) */
 	UPROPERTY()
 	bool bSnapToWorldGrid = false;
 
@@ -236,17 +242,6 @@ public:
 	UPROPERTY()
 	bool bSnapToWorldRotGrid = false;
 
-	//
-	// Transform Source
-	//
-	UPROPERTY()
-	TScriptInterface<ITransformGizmoSource> TransformSource;
-
-
-	/** Root of renderable gizmo elements */
-	UPROPERTY()
-	TObjectPtr<UGizmoElementGroup> GizmoElementRoot;
-
 protected:
 
 	//
@@ -325,9 +320,6 @@ protected:
 	UPROPERTY()
 	TObjectPtr<UGizmoElementBox> ScaleUniformElement;
 
-	//
-	// Axis Sources
-	//
 	/** Axis that points towards camera, X/Y plane tangents aligned to right/up. Shared across Gizmos, and created internally during SetActiveTarget() */
 	UPROPERTY()
 	TObjectPtr<UGizmoConstantFrameAxisSource> CameraAxisSource;
@@ -335,13 +327,9 @@ protected:
 	// internal function that updates CameraAxisSource by getting current view state from GizmoManager
 	void UpdateCameraAxisSource();
 
-	/** 
-	 * State target is shared across gizmos, and created internally during SetActiveTarget(). 
-	 * Several FChange providers are registered with this StateTarget, including the UTransformGizmo
-	 * itself (IToolCommandChangeSource implementation above is called)
-	 */
+	/** The state target is created internally during SetActiveTarget() */
 	UPROPERTY()
-	TObjectPtr<UGizmoDependentTransformChangeStateTarget> StateTarget;
+	TObjectPtr<IGizmoStateTarget> StateTarget;
 
 	/**
 	 * These are used to let the translation subgizmos use raycasts into the scene to align the gizmo with scene geometry.
@@ -351,25 +339,32 @@ protected:
 	TUniqueFunction<bool(const FRay&, FVector&)> DestinationAlignmentRayCaster = [](const FRay&, FVector&) {return false; };
 
 	bool bDisallowNegativeScaling = false;
+
 protected:
 
+	/** Setup behaviors */
+	virtual void SetupBehaviors();
+
+	/** Setup materials */
+	virtual void SetupMaterials();
+
 	/** Update current gizmo mode based on transform source */
-	void UpdateMode();
+	virtual void UpdateMode();
 
 	/** Enable the given mode with the specified axes, EAxisList::Type::None will hide objects associated with mode */
-	void EnableMode(EGizmoTransformMode InGizmoMode, EAxisList::Type InAxisListToDraw);
+	virtual void EnableMode(EGizmoTransformMode InGizmoMode, EAxisList::Type InAxisListToDraw);
 
 	/** Enable translate using specified axis list */
-	void EnableTranslate(EAxisList::Type InAxisListToDraw);
+	virtual void EnableTranslate(EAxisList::Type InAxisListToDraw);
 
 	/** Enable rotate using specified axis list */
-	void EnableRotate(EAxisList::Type InAxisListToDraw);
+	virtual void EnableRotate(EAxisList::Type InAxisListToDraw);
 
 	/** Enable scale using specified axis list */
-	void EnableScale(EAxisList::Type InAxisListToDraw);
+	virtual void EnableScale(EAxisList::Type InAxisListToDraw);
 
 	/** Enable planar handles used by translate and scale */
-	void EnablePlanarObjects(bool bTranslate, bool bEnableX, bool bEnableY, bool bEnableZ);
+	virtual void EnablePlanarObjects(bool bTranslate, bool bEnableX, bool bEnableY, bool bEnableZ);
 
 	/** Construct translate axis handle */
 	virtual UGizmoElementArrow* MakeTranslateAxis(ETransformGizmoPartIdentifier InPartId, const FVector& InAxisDir, const FVector& InSideDir, UMaterialInterface* InMaterial);
@@ -394,27 +389,39 @@ protected:
 	/** Construct rotate screen space handle */
 	virtual UGizmoElementCircle* MakeRotateCircleHandle(ETransformGizmoPartIdentifier InPartId, float InRadius, const FLinearColor& InColor, float bFill);
 
-	// Update hover for hit part based on input device pos
-	FInputRayHit UpdateHoverHitSubElement(const FInputDeviceRay& DevicePos);
+	/** Get gizmo transform based on cached current transform. */
+	virtual FTransform GetGizmoTransform() const;
+
+	/** Determine hit part and update hover state based on current input ray */
+	virtual FInputRayHit UpdateHoveredPart(const FInputDeviceRay& DevicePos);
+
+	/** Get current interaction axis */
+	virtual FVector GetWorldAxis(const FVector& InAxis);
+
+	/** Update current gizmo mode based on transform source */
+	virtual void OnClickPressTranslate(const FInputDeviceRay& InPressPos);
+
+	/** Update current gizmo mode based on transform source */
+	virtual void OnClickDragTranslate(const FInputDeviceRay& PressPos);
+
+	/** Update current gizmo mode based on transform source */
+	virtual void OnClickReleaseTranslate(const FInputDeviceRay& PressPos);
+
+	/** Update current gizmo mode based on transform source */
+	virtual void Translate(const FVector& InTranslateDelta);
+
 
 	// Axis and Plane TransformSources use this function to execute worldgrid snap queries
 	bool PositionSnapFunction(const FVector& WorldPosition, FVector& SnappedPositionOut) const;
 	FQuat RotationSnapFunction(const FQuat& DeltaRotation) const;
 
-	// Verify part identifier is within recognized range of transform gizmo part ids
-	bool VerifyPartIdentifier(uint32 InPartIdentifier);
-
-	// Currently rendered transform mode
-	UPROPERTY()
-	EGizmoTransformMode CurrentMode = EGizmoTransformMode::None;
+	// Get max part identifier.
+	virtual uint32 GetMaxPartIdentifier() const;
 
-	// Currently rendered axis list
-	UPROPERTY()
-	TEnumAsByte<EAxisList::Type> CurrentAxisToDraw = EAxisList::None;
+	// Verify part identifier is within recognized range of transform gizmo part ids
+	virtual bool VerifyPartIdentifier(uint32 InPartIdentifier) const;
 
-	// Last hit part
-	UPROPERTY()
-	ETransformGizmoPartIdentifier LastHitPart = ETransformGizmoPartIdentifier::Default;
+protected:
 
 	/** Materials and colors to be used when drawing the items for each axis */
 	UPROPERTY()
@@ -435,4 +442,45 @@ protected:
 	TObjectPtr<UMaterialInstanceDynamic> WhiteMaterial;
 	UPROPERTY()
 	TObjectPtr<UMaterialInstanceDynamic> OpaquePlaneMaterialXY;
+
+	/** Current transform */
+	UPROPERTY()
+	FTransform CurrentTransform = FTransform::Identity;
+
+	/** Currently rendered transform mode */
+	UPROPERTY()
+	EGizmoTransformMode CurrentMode = EGizmoTransformMode::None;
+
+	/** Currently rendered axis list */
+	UPROPERTY()
+	TEnumAsByte<EAxisList::Type> CurrentAxisToDraw = EAxisList::None;
+
+	/** Last hit part */
+	UPROPERTY()
+	ETransformGizmoPartIdentifier LastHitPart = ETransformGizmoPartIdentifier::Default;
+
+	//
+	// The values below are used in the context of a single click-drag interaction, ie if bInInteraction = true
+	// They otherwise should be considered uninitialized
+	//
+
+	/** Active world space axis origin (only valid between state target BeginModify/EndModify) */
+	UPROPERTY()
+	FVector InteractionAxisOrigin;
+
+	/** Active world space axis (only valid between state target BeginModify/EndModify) */
+	UPROPERTY()
+	FVector InteractionAxis;
+
+	/** Active axis type (only valid between state target BeginModify/EndModify) */
+	UPROPERTY()
+	TEnumAsByte<EAxisList::Type> InteractionAxisType;
+
+	/** Active interaction start point (only valid between state target BeginModify/EndModify) */
+	UPROPERTY()
+	FVector InteractionStartPoint;
+
+	/** Active interaction current point (only valid between state target BeginModify/EndModify) */
+	UPROPERTY()
+	FVector InteractionCurrPoint;
 };
diff --git a/Engine/Source/Editor/Experimental/EditorInteractiveToolsFramework/Public/EditorGizmos/TransformGizmoInterfaces.h b/Engine/Source/Editor/Experimental/EditorInteractiveToolsFramework/Public/EditorGizmos/TransformGizmoInterfaces.h
--- a/Engine/Source/Editor/Experimental/EditorInteractiveToolsFramework/Public/EditorGizmos/TransformGizmoInterfaces.h
+++ b/Engine/Source/Editor/Experimental/EditorInteractiveToolsFramework/Public/EditorGizmos/TransformGizmoInterfaces.h
@@ -20,16 +20,6 @@ enum class EGizmoTransformMode : uint8
 	Max
 };
 
-
-namespace FTransformGizmoUtil
-{
-	/** Convert UE::Widget::EWidgetMode to ETransformGizmoMode*/
-	EGizmoTransformMode GetGizmoMode(UE::Widget::EWidgetMode InWidgetMode);
-
-	/** Convert EEditorGizmoMode to UE::Widget::EWidgetMode*/
-	UE::Widget::EWidgetMode GetWidgetMode(EGizmoTransformMode InGizmoMode);
-};
-
 UINTERFACE()
 class EDITORINTERACTIVETOOLSFRAMEWORK_API UTransformGizmoSource : public UInterface
 {
diff --git a/Engine/Source/Runtime/InteractiveToolsFramework/Private/BaseGizmos/GizmoElementHitTargets.cpp b/Engine/Source/Runtime/InteractiveToolsFramework/Private/BaseGizmos/GizmoElementHitTargets.cpp
--- a/Engine/Source/Runtime/InteractiveToolsFramework/Private/BaseGizmos/GizmoElementHitTargets.cpp
+++ b/Engine/Source/Runtime/InteractiveToolsFramework/Private/BaseGizmos/GizmoElementHitTargets.cpp
@@ -64,37 +64,43 @@ FInputRayHit UGizmoElementHitMultiTarget::IsHit(const FInputDeviceRay& ClickPos)
 	return FInputRayHit();
 }
 
-void UGizmoElementHitMultiTarget::UpdateHoverState(bool bHovering, uint32 PartIdentifier)
+void UGizmoElementHitMultiTarget::UpdateHoverState(bool bInHovering, uint32 PartIdentifier)
 {
 	if (!GizmoElement)
 	{
 		return;
 	}
 
-	if (bHovering)
+	if (bInHovering && !bHovering)
 	{
 		GizmoElement->UpdatePartInteractionState(EGizmoElementInteractionState::Hovering, PartIdentifier);
+		bHovering = true;
+		bInteracting = false;
 	}
-	else
+	else if (!bInHovering && bHovering)
 	{
 		GizmoElement->UpdatePartInteractionState(EGizmoElementInteractionState::None, PartIdentifier);
+		bHovering = false;
 	}
 }
 
-void UGizmoElementHitMultiTarget::UpdateInteractingState(bool bInteracting, uint32 PartIdentifier)
+void UGizmoElementHitMultiTarget::UpdateInteractingState(bool bInInteracting, uint32 PartIdentifier)
 {
 	if (!GizmoElement)
 	{
 		return;
 	}
 
-	if (bInteracting)
+	if (bInInteracting && !bInteracting)
 	{
 		GizmoElement->UpdatePartInteractionState(EGizmoElementInteractionState::Interacting, PartIdentifier);
+		bInteracting = true;
+		bHovering = false;
 	}
-	else
+	else if (!bInInteracting && bInteracting)
 	{
 		GizmoElement->UpdatePartInteractionState(EGizmoElementInteractionState::None, PartIdentifier);
+		bInteracting = false;
 	}
 }
 
diff --git a/Engine/Source/Runtime/InteractiveToolsFramework/Public/BaseGizmos/GizmoElementHitTargets.h b/Engine/Source/Runtime/InteractiveToolsFramework/Public/BaseGizmos/GizmoElementHitTargets.h
--- a/Engine/Source/Runtime/InteractiveToolsFramework/Public/BaseGizmos/GizmoElementHitTargets.h
+++ b/Engine/Source/Runtime/InteractiveToolsFramework/Public/BaseGizmos/GizmoElementHitTargets.h
@@ -87,6 +87,14 @@ public:
 
 	virtual void UpdateHittableState(bool bHittable, uint32 PartIdentifier);
 
+protected:
+
+	UPROPERTY()
+	bool bHovering = false;
+
+	UPROPERTY()
+	bool bInteracting = false;
+
 public:
 	static UGizmoElementHitMultiTarget* Construct(
 		UGizmoElementBase* InGizmoElement,
