Add scale, planar scale and planar translate interaction to the new TRS gizmo.

diff --git a/Engine/Source/Editor/Experimental/EditorInteractiveToolsFramework/Private/EditorGizmos/EditorTransformGizmo.cpp b/Engine/Source/Editor/Experimental/EditorInteractiveToolsFramework/Private/EditorGizmos/EditorTransformGizmo.cpp
--- a/Engine/Source/Editor/Experimental/EditorInteractiveToolsFramework/Private/EditorGizmos/EditorTransformGizmo.cpp
+++ b/Engine/Source/Editor/Experimental/EditorInteractiveToolsFramework/Private/EditorGizmos/EditorTransformGizmo.cpp
@@ -7,7 +7,7 @@
 
 DEFINE_LOG_CATEGORY_STATIC(LogEditorTransformGizmo, Log, All);
 
-void UEditorTransformGizmo::Translate(const FVector& InTranslateDelta)
+void UEditorTransformGizmo::ApplyTranslateDelta(const FVector& InTranslateDelta)
 {
 	check(ActiveTarget);
 
@@ -20,7 +20,27 @@ void UEditorTransformGizmo::Translate(const FVector& InTranslateDelta)
 	}
 	else
 	{
-		Super::Translate(InTranslateDelta);
+		Super::ApplyTranslateDelta(InTranslateDelta);
+	}
+}
+
+void UEditorTransformGizmo::ApplyScaleDelta(const FVector& InScaleDelta)
+{
+	check(ActiveTarget);
+
+	if (UEditorTransformProxy* EditorTransformProxy = Cast<UEditorTransformProxy>(ActiveTarget))
+	{
+		FVector StartScale = CurrentTransform.GetScale3D();
+
+		EditorTransformProxy->InputScaleDelta(InScaleDelta, InteractionAxisType);
+
+		// Update the cached current transform
+		FVector NewScale = StartScale + InScaleDelta * InteractionAxis;
+		CurrentTransform.SetScale3D(NewScale);
+	}
+	else
+	{
+		Super::ApplyScaleDelta(InScaleDelta);
 	}
 }
 
diff --git a/Engine/Source/Editor/Experimental/EditorInteractiveToolsFramework/Private/EditorGizmos/TransformGizmo.cpp b/Engine/Source/Editor/Experimental/EditorInteractiveToolsFramework/Private/EditorGizmos/TransformGizmo.cpp
--- a/Engine/Source/Editor/Experimental/EditorInteractiveToolsFramework/Private/EditorGizmos/TransformGizmo.cpp
+++ b/Engine/Source/Editor/Experimental/EditorInteractiveToolsFramework/Private/EditorGizmos/TransformGizmo.cpp
@@ -244,21 +244,96 @@ void UTransformGizmo::OnClickPress(const FInputDeviceRay& PressPos)
 		InteractionAxisOrigin = CurrentTransform.GetLocation();
 		InteractionAxis = GetWorldAxis(FVector::XAxisVector);
 		InteractionAxisType = EAxisList::X;
-		OnClickPressTranslate(PressPos);
+		OnClickPressAxis(PressPos);
 	}
 	else if (LastHitPart == ETransformGizmoPartIdentifier::TranslateYAxis)
 	{
 		InteractionAxisOrigin = CurrentTransform.GetLocation();
 		InteractionAxis = GetWorldAxis(FVector::YAxisVector);
 		InteractionAxisType = EAxisList::Y;
-		OnClickPressTranslate(PressPos);
+		OnClickPressAxis(PressPos);
 	}
 	else if (LastHitPart == ETransformGizmoPartIdentifier::TranslateZAxis)
 	{
 		InteractionAxisOrigin = CurrentTransform.GetLocation();
 		InteractionAxis = GetWorldAxis(FVector::ZAxisVector);
 		InteractionAxisType = EAxisList::Z;
-		OnClickPressTranslate(PressPos);
+		OnClickPressAxis(PressPos);
+	}
+	else if (LastHitPart == ETransformGizmoPartIdentifier::ScaleXAxis)
+	{
+		InteractionAxisOrigin = CurrentTransform.GetLocation();
+		InteractionAxis = GetWorldAxis(FVector::XAxisVector);
+		InteractionAxisType = EAxisList::X;
+		OnClickPressAxis(PressPos);
+	}
+	else if (LastHitPart == ETransformGizmoPartIdentifier::ScaleYAxis)
+	{
+		InteractionAxisOrigin = CurrentTransform.GetLocation();
+		InteractionAxis = GetWorldAxis(FVector::YAxisVector);
+		InteractionAxisType = EAxisList::Y;
+		OnClickPressAxis(PressPos);
+	}
+	else if (LastHitPart == ETransformGizmoPartIdentifier::ScaleZAxis)
+	{
+		InteractionAxisOrigin = CurrentTransform.GetLocation();
+		InteractionAxis = GetWorldAxis(FVector::ZAxisVector);
+		InteractionAxisType = EAxisList::Z;
+		OnClickPressAxis(PressPos);
+	}
+	if (LastHitPart == ETransformGizmoPartIdentifier::TranslateXYPlanar)
+	{
+		InteractionAxisOrigin = CurrentTransform.GetLocation();
+		InteractionNormal = GetWorldAxis(FVector::ZAxisVector);
+		InteractionAxisX = GetWorldAxis(FVector::XAxisVector);
+		InteractionAxisY = GetWorldAxis(FVector::YAxisVector);
+		InteractionAxisType = EAxisList::XY;
+		OnClickPressPlanar(PressPos);
+	}
+	else if (LastHitPart == ETransformGizmoPartIdentifier::TranslateYZPlanar)
+	{
+		InteractionAxisOrigin = CurrentTransform.GetLocation();
+		InteractionNormal = GetWorldAxis(FVector::XAxisVector);
+		InteractionAxisX = GetWorldAxis(FVector::YAxisVector);
+		InteractionAxisY = GetWorldAxis(FVector::ZAxisVector);
+		InteractionAxisType = EAxisList::YZ;
+		OnClickPressPlanar(PressPos);
+	}
+	else if (LastHitPart == ETransformGizmoPartIdentifier::TranslateXZPlanar)
+	{
+		InteractionAxisOrigin = CurrentTransform.GetLocation();
+		InteractionNormal = GetWorldAxis(FVector::YAxisVector);
+		InteractionAxisX = GetWorldAxis(FVector::ZAxisVector);
+		InteractionAxisY = GetWorldAxis(FVector::XAxisVector);
+		InteractionAxisType = EAxisList::XZ;
+		OnClickPressPlanar(PressPos);
+	}
+	else if (LastHitPart == ETransformGizmoPartIdentifier::ScaleXYPlanar)
+	{
+		InteractionAxisOrigin = CurrentTransform.GetLocation();
+		InteractionNormal = GetWorldAxis(FVector::ZAxisVector);
+		InteractionAxisX = GetWorldAxis(FVector::XAxisVector);
+		InteractionAxisY = GetWorldAxis(FVector::YAxisVector);
+		InteractionAxisType = EAxisList::XY;
+		OnClickPressPlanar(PressPos);
+	}
+	else if (LastHitPart == ETransformGizmoPartIdentifier::ScaleYZPlanar)
+	{
+		InteractionAxisOrigin = CurrentTransform.GetLocation();
+		InteractionNormal = GetWorldAxis(FVector::XAxisVector);
+		InteractionAxisX = GetWorldAxis(FVector::YAxisVector);
+		InteractionAxisY = GetWorldAxis(FVector::ZAxisVector);
+		InteractionAxisType = EAxisList::YZ;
+		OnClickPressPlanar(PressPos);
+	}
+	else if (LastHitPart == ETransformGizmoPartIdentifier::ScaleXZPlanar)
+	{
+		InteractionAxisOrigin = CurrentTransform.GetLocation();
+		InteractionNormal = GetWorldAxis(FVector::YAxisVector);
+		InteractionAxisX = GetWorldAxis(FVector::ZAxisVector);
+		InteractionAxisY = GetWorldAxis(FVector::XAxisVector);
+		InteractionAxisType = EAxisList::XZ;
+		OnClickPressPlanar(PressPos);
 	}
 
 	if (bInInteraction)
@@ -284,9 +359,21 @@ void UTransformGizmo::OnClickDrag(const FInputDeviceRay& DragPos)
 
 	if (LastHitPart == ETransformGizmoPartIdentifier::TranslateXAxis ||
 		LastHitPart == ETransformGizmoPartIdentifier::TranslateYAxis ||
-		LastHitPart == ETransformGizmoPartIdentifier::TranslateZAxis)
+		LastHitPart == ETransformGizmoPartIdentifier::TranslateZAxis ||
+		LastHitPart == ETransformGizmoPartIdentifier::ScaleXAxis ||
+		LastHitPart == ETransformGizmoPartIdentifier::ScaleYAxis ||
+		LastHitPart == ETransformGizmoPartIdentifier::ScaleZAxis)
+	{
+		OnClickDragAxis(DragPos);
+	}
+	else if (LastHitPart == ETransformGizmoPartIdentifier::TranslateXYPlanar ||
+		LastHitPart == ETransformGizmoPartIdentifier::TranslateYZPlanar ||
+		LastHitPart == ETransformGizmoPartIdentifier::TranslateXZPlanar ||
+		LastHitPart == ETransformGizmoPartIdentifier::ScaleXYPlanar ||
+		LastHitPart == ETransformGizmoPartIdentifier::ScaleYZPlanar ||
+		LastHitPart == ETransformGizmoPartIdentifier::ScaleXZPlanar)
 	{
-		OnClickDragTranslate(DragPos);
+		OnClickDragPlanar(DragPos);
 	}
 }
 
@@ -299,9 +386,21 @@ void UTransformGizmo::OnClickRelease(const FInputDeviceRay& ReleasePos)
 
 	if (LastHitPart == ETransformGizmoPartIdentifier::TranslateXAxis ||
 		LastHitPart == ETransformGizmoPartIdentifier::TranslateYAxis ||
-		LastHitPart == ETransformGizmoPartIdentifier::TranslateZAxis)
+		LastHitPart == ETransformGizmoPartIdentifier::TranslateZAxis ||
+		LastHitPart == ETransformGizmoPartIdentifier::ScaleXAxis ||
+		LastHitPart == ETransformGizmoPartIdentifier::ScaleYAxis ||
+		LastHitPart == ETransformGizmoPartIdentifier::ScaleZAxis)
+	{
+		OnClickReleaseAxis(ReleasePos);
+	}
+	else if (LastHitPart == ETransformGizmoPartIdentifier::TranslateXYPlanar ||
+		LastHitPart == ETransformGizmoPartIdentifier::TranslateYZPlanar ||
+		LastHitPart == ETransformGizmoPartIdentifier::TranslateXZPlanar ||
+		LastHitPart == ETransformGizmoPartIdentifier::ScaleXYPlanar ||
+		LastHitPart == ETransformGizmoPartIdentifier::ScaleYZPlanar ||
+		LastHitPart == ETransformGizmoPartIdentifier::ScaleXZPlanar)
 	{
-		OnClickReleaseTranslate(ReleasePos);
+		OnClickReleaseAxis(ReleasePos);
 	}
 
 	if (StateTarget)
@@ -441,6 +540,27 @@ void UTransformGizmo::EnablePlanarObjects(bool bTranslate, bool bEnableX, bool b
 {
 	check(GizmoElementRoot);
 
+	auto EnablePlanarElement = [this](
+		TObjectPtr<UGizmoElementRectangle>& PlanarElement,
+		ETransformGizmoPartIdentifier PartId,
+		const FVector& Axis0,
+		const FVector& Axis1,
+		const FVector& Axis2,
+		const FLinearColor& AxisColor,
+		bool bEnable)
+	{
+		if (bEnable && PlanarElement == nullptr)
+		{
+			PlanarElement = MakePlanarHandle(PartId, Axis0, Axis1, Axis2, TransparentVertexColorMaterial, AxisColor);
+			GizmoElementRoot->Add(PlanarElement);
+		}
+
+		if (PlanarElement)
+		{
+			PlanarElement->SetEnabled(bEnable);
+		}
+	};
+
 	const bool bEnableXY = bEnableX && bEnableY;
 	const bool bEnableYZ = bEnableY && bEnableZ;
 	const bool bEnableXZ = bEnableX && bEnableZ;
@@ -449,40 +569,17 @@ void UTransformGizmo::EnablePlanarObjects(bool bTranslate, bool bEnableX, bool b
 	const FVector YAxis(0.0f, 1.0f, 0.0f);
 	const FVector ZAxis(0.0f, 0.0f, 1.0f);
 
-	if (bEnableXY && PlanarXYElement == nullptr)
-	{
-		PlanarXYElement = MakePlanarHandle(bTranslate ? ETransformGizmoPartIdentifier::TranslateXYPlanar : ETransformGizmoPartIdentifier::ScaleXYPlanar, 
-											XAxis, YAxis, ZAxis, TransparentVertexColorMaterial, AxisColorZ);
-		GizmoElementRoot->Add(PlanarXYElement);
-	}
-
-	if (bEnableYZ && PlanarYZElement == nullptr)
-	{
-		PlanarYZElement = MakePlanarHandle(bTranslate ? ETransformGizmoPartIdentifier::TranslateYZPlanar : ETransformGizmoPartIdentifier::ScaleYZPlanar, 
-											YAxis, ZAxis, XAxis, TransparentVertexColorMaterial, AxisColorX);
-		GizmoElementRoot->Add(PlanarYZElement);
-	}
-
-	if (bEnableXZ && PlanarXZElement == nullptr)
-	{
-		PlanarXZElement = MakePlanarHandle(bTranslate ? ETransformGizmoPartIdentifier::TranslateXZPlanar : ETransformGizmoPartIdentifier::ScaleXZPlanar, 
-											ZAxis, XAxis, YAxis, TransparentVertexColorMaterial, AxisColorY);
-		GizmoElementRoot->Add(PlanarXZElement);
-	}
-
-	if (PlanarXYElement)
-	{
-		PlanarXYElement->SetEnabled(bEnableXY);
-	}
-
-	if (PlanarYZElement)
+	if (bTranslate)
 	{
-		PlanarYZElement->SetEnabled(bEnableYZ);
+		EnablePlanarElement(TranslatePlanarXYElement, ETransformGizmoPartIdentifier::TranslateXYPlanar, XAxis, YAxis, ZAxis, AxisColorZ, bEnableXY);
+		EnablePlanarElement(TranslatePlanarYZElement, ETransformGizmoPartIdentifier::TranslateYZPlanar, YAxis, ZAxis, XAxis, AxisColorX, bEnableYZ);
+		EnablePlanarElement(TranslatePlanarXZElement, ETransformGizmoPartIdentifier::TranslateXZPlanar, ZAxis, XAxis, YAxis, AxisColorY, bEnableXZ);
 	}
-
-	if (PlanarXZElement)
+	else
 	{
-		PlanarXZElement->SetEnabled(bEnableXZ);
+		EnablePlanarElement(ScalePlanarXYElement, ETransformGizmoPartIdentifier::ScaleXYPlanar, XAxis, YAxis, ZAxis, AxisColorZ, bEnableXY);
+		EnablePlanarElement(ScalePlanarYZElement, ETransformGizmoPartIdentifier::ScaleYZPlanar, YAxis, ZAxis, XAxis, AxisColorX, bEnableYZ);
+		EnablePlanarElement(ScalePlanarXZElement, ETransformGizmoPartIdentifier::ScaleXZPlanar, ZAxis, XAxis, YAxis, AxisColorY, bEnableXZ);
 	}
 }
 
@@ -930,7 +1027,7 @@ FVector UTransformGizmo::GetWorldAxis(const FVector& InAxis)
 	return InAxis;
 }
 
-void UTransformGizmo::OnClickPressTranslate(const FInputDeviceRay& InPressPos)
+void UTransformGizmo::OnClickPressAxis(const FInputDeviceRay& InPressPos)
 {
 	// Find interaction start point and parameter.
 	FVector NearestPt; 
@@ -944,7 +1041,7 @@ void UTransformGizmo::OnClickPressTranslate(const FInputDeviceRay& InPressPos)
 	bInInteraction = true;
 }
 
-void UTransformGizmo::OnClickDragTranslate(const FInputDeviceRay& InDragPos)
+void UTransformGizmo::OnClickDragAxis(const FInputDeviceRay& InDragPos)
 {
 	float RayNearestParam, AxisNearestParam;
 	FVector RayNearestPt, AxisNearestPoint;
@@ -953,22 +1050,106 @@ void UTransformGizmo::OnClickDragTranslate(const FInputDeviceRay& InDragPos)
 		AxisNearestPoint, AxisNearestParam,
 		RayNearestPt, RayNearestParam);
 
-	FVector TranslateDelta = AxisNearestPoint - InteractionCurrPoint;
+	FVector Delta = AxisNearestPoint - InteractionCurrPoint;
 	InteractionCurrPoint = AxisNearestPoint;
 
-	Translate(TranslateDelta);
+	if (LastHitPart == ETransformGizmoPartIdentifier::TranslateXAxis ||
+		LastHitPart == ETransformGizmoPartIdentifier::TranslateYAxis ||
+		LastHitPart == ETransformGizmoPartIdentifier::TranslateZAxis)
+	{
+		ApplyTranslateDelta(Delta);
+	}
+	else
+	{
+		Delta *= ScaleMultiplier;
+		ApplyScaleDelta(Delta);
+	}
 }
 
-void UTransformGizmo::OnClickReleaseTranslate(const FInputDeviceRay& InReleasePos)
+void UTransformGizmo::OnClickReleaseAxis(const FInputDeviceRay& InReleasePos)
 {
 	bInInteraction = false;
 }
 
-void UTransformGizmo::Translate(const FVector& InTranslateDelta)
+void UTransformGizmo::OnClickPressPlanar(const FInputDeviceRay& PressPos)
+{
+	// if ray is parallel to plane, nothing has been hit
+	if (FMath::IsNearlyZero(FVector::DotProduct(InteractionNormal, PressPos.WorldRay.Direction)))
+	{
+		return;
+	}
+
+	FPlane Plane(InteractionAxisOrigin, InteractionNormal);
+	double HitDepth = FMath::RayPlaneIntersectionParam(PressPos.WorldRay.Origin, PressPos.WorldRay.Direction, Plane);
+	if (HitDepth < 0)
+	{
+		return;
+	}
+
+	FVector HitPoint = PressPos.WorldRay.Origin + PressPos.WorldRay.Direction * HitDepth;
+	InteractionStartPoint = HitPoint;
+	InteractionCurrPoint = InteractionStartPoint;
+
+	bInInteraction = true;
+}
+
+void UTransformGizmo::OnClickDragPlanar(const FInputDeviceRay& DragPos)
+{
+	// if ray is parallel to plane, nothing has been hit
+	if (FMath::IsNearlyZero(FVector::DotProduct(InteractionNormal, DragPos.WorldRay.Direction)))
+	{
+		return;
+	}
+
+	FPlane Plane(InteractionAxisOrigin, InteractionNormal);
+	double HitDepth = FMath::RayPlaneIntersectionParam(DragPos.WorldRay.Origin, DragPos.WorldRay.Direction, Plane);
+	if (HitDepth < 0)
+	{
+		return;
+	}
+
+	FVector HitPoint = DragPos.WorldRay.Origin + DragPos.WorldRay.Direction * HitDepth;
+	FVector Delta = HitPoint - InteractionCurrPoint;
+
+	InteractionCurrPoint = HitPoint;
+
+	if (LastHitPart == ETransformGizmoPartIdentifier::TranslateXYPlanar ||
+		LastHitPart == ETransformGizmoPartIdentifier::TranslateYZPlanar ||
+		LastHitPart == ETransformGizmoPartIdentifier::TranslateXZPlanar)
+	{
+		ApplyTranslateDelta(Delta);
+	}
+	else
+	{
+		const float ScaleMax = Delta.GetMax();
+		const float ScaleMin = Delta.GetMin();
+		const float ScaleApplied = (ScaleMax > -ScaleMin) ? ScaleMax * ScaleMultiplier : ScaleMin * ScaleMultiplier;
+
+		FVector ScaleDelta(InteractionAxisType & EAxisList::X ? ScaleApplied : 0.0,
+						   InteractionAxisType & EAxisList::Y ? ScaleApplied : 0.0,
+						   InteractionAxisType & EAxisList::Z ? ScaleApplied : 0.0);
+
+		ApplyScaleDelta(ScaleDelta);
+	}
+}
+
+void UTransformGizmo::OnClickReleasePlanar(const FInputDeviceRay& InReleasePos)
+{
+	bInInteraction = false;
+}
+
+void UTransformGizmo::ApplyTranslateDelta(const FVector& InTranslateDelta)
 {
 	CurrentTransform.AddToTranslation(InTranslateDelta);
 	ActiveTarget->SetTransform(CurrentTransform);
 }
 
+void UTransformGizmo::ApplyScaleDelta(const FVector& InScaleDelta)
+{
+	FVector StartScale = CurrentTransform.GetScale3D();
+	FVector NewScale = StartScale + InScaleDelta * InteractionAxis;
+	CurrentTransform.SetScale3D(NewScale);
+	ActiveTarget->SetTransform(CurrentTransform);
+}
 
 #undef LOCTEXT_NAMESPACE
diff --git a/Engine/Source/Editor/Experimental/EditorInteractiveToolsFramework/Public/EditorGizmos/EditorTransformGizmo.h b/Engine/Source/Editor/Experimental/EditorInteractiveToolsFramework/Public/EditorGizmos/EditorTransformGizmo.h
--- a/Engine/Source/Editor/Experimental/EditorInteractiveToolsFramework/Public/EditorGizmos/EditorTransformGizmo.h
+++ b/Engine/Source/Editor/Experimental/EditorInteractiveToolsFramework/Public/EditorGizmos/EditorTransformGizmo.h
@@ -17,6 +17,9 @@ class EDITORINTERACTIVETOOLSFRAMEWORK_API UEditorTransformGizmo : public UTransf
 
 protected:
 
-	/** Update current gizmo mode based on transform source */
-	virtual void Translate(const FVector& InTranslateDelta) override;
+	/** Apply translate delta to transform proxy */
+	virtual void ApplyTranslateDelta(const FVector& InTranslateDelta);
+
+	/** Apply scale delta to transform proxy */
+	virtual void ApplyScaleDelta(const FVector& InScaleDelta);
 };
\ No newline at end of file
diff --git a/Engine/Source/Editor/Experimental/EditorInteractiveToolsFramework/Public/EditorGizmos/TransformGizmo.h b/Engine/Source/Editor/Experimental/EditorInteractiveToolsFramework/Public/EditorGizmos/TransformGizmo.h
--- a/Engine/Source/Editor/Experimental/EditorInteractiveToolsFramework/Public/EditorGizmos/TransformGizmo.h
+++ b/Engine/Source/Editor/Experimental/EditorInteractiveToolsFramework/Public/EditorGizmos/TransformGizmo.h
@@ -264,17 +264,17 @@ protected:
 	UPROPERTY()
 	TObjectPtr<UGizmoElementRectangle> TranslateScreenSpaceElement;
 
-	/** Planar XY handle */
+	/** Translate planar XY handle */
 	UPROPERTY()
-	TObjectPtr<UGizmoElementRectangle> PlanarXYElement;
+	TObjectPtr<UGizmoElementRectangle> TranslatePlanarXYElement;
 
-	/** Planar YZ handle */
+	/** Translate planar YZ handle */
 	UPROPERTY()
-	TObjectPtr<UGizmoElementRectangle> PlanarYZElement;
+	TObjectPtr<UGizmoElementRectangle> TranslatePlanarYZElement;
 
-	/** Planar XZ handle */
+	/** Translate planar XZ handle */
 	UPROPERTY()
-	TObjectPtr<UGizmoElementRectangle> PlanarXZElement;
+	TObjectPtr<UGizmoElementRectangle> TranslatePlanarXZElement;
 
 	/** Rotate X Axis */
 	UPROPERTY()
@@ -316,6 +316,18 @@ protected:
 	UPROPERTY()
 	TObjectPtr<UGizmoElementArrow> ScaleZAxisElement;
 
+	/** Scale planar XY handle */
+	UPROPERTY()
+	TObjectPtr<UGizmoElementRectangle> ScalePlanarXYElement;
+
+	/** Scale planar YZ handle */
+	UPROPERTY()
+	TObjectPtr<UGizmoElementRectangle> ScalePlanarYZElement;
+
+	/** Scale planar XZ handle */
+	UPROPERTY()
+	TObjectPtr<UGizmoElementRectangle> ScalePlanarXZElement;
+
 	/** Uniform scale object */
 	UPROPERTY()
 	TObjectPtr<UGizmoElementBox> ScaleUniformElement;
@@ -398,18 +410,29 @@ protected:
 	/** Get current interaction axis */
 	virtual FVector GetWorldAxis(const FVector& InAxis);
 
-	/** Update current gizmo mode based on transform source */
-	virtual void OnClickPressTranslate(const FInputDeviceRay& InPressPos);
+	/** Handle click press for translate and scale axes */
+	virtual void OnClickPressAxis(const FInputDeviceRay& InPressPos);
 
-	/** Update current gizmo mode based on transform source */
-	virtual void OnClickDragTranslate(const FInputDeviceRay& PressPos);
+	/** Handle click drag for translate and scale axes */
+	virtual void OnClickDragAxis(const FInputDeviceRay& PressPos);
 
-	/** Update current gizmo mode based on transform source */
-	virtual void OnClickReleaseTranslate(const FInputDeviceRay& PressPos);
+	/** Handle click release for translate and scale axes */
+	virtual void OnClickReleaseAxis(const FInputDeviceRay& PressPos);
 
-	/** Update current gizmo mode based on transform source */
-	virtual void Translate(const FVector& InTranslateDelta);
+	/** Handle click press for translate and scale planar */
+	virtual void OnClickPressPlanar(const FInputDeviceRay& InPressPos);
+
+	/** Handle click drag for translate and scale planar */
+	virtual void OnClickDragPlanar(const FInputDeviceRay& PressPos);
+
+	/** Handle click release for translate and scale planar */
+	virtual void OnClickReleasePlanar(const FInputDeviceRay& PressPos);
+
+	/** Apply translate delta to transform proxy */
+	virtual void ApplyTranslateDelta(const FVector& InTranslateDelta);
 
+	/** Apply scale delta to transform proxy */
+	virtual void ApplyScaleDelta(const FVector& InScaleDelta);
 
 	// Axis and Plane TransformSources use this function to execute worldgrid snap queries
 	bool PositionSnapFunction(const FVector& WorldPosition, FVector& SnappedPositionOut) const;
@@ -443,6 +466,10 @@ protected:
 	UPROPERTY()
 	TObjectPtr<UMaterialInstanceDynamic> OpaquePlaneMaterialXY;
 
+	/** Scale delta is multiplied by this amount */
+	UPROPERTY()
+	double ScaleMultiplier = 0.05;
+
 	/** Current transform */
 	UPROPERTY()
 	FTransform CurrentTransform = FTransform::Identity;
@@ -472,6 +499,18 @@ protected:
 	UPROPERTY()
 	FVector InteractionAxis;
 
+	/** Active world space normal used for planar (only valid between state target BeginModify/EndModify) */
+	UPROPERTY()
+	FVector InteractionNormal;
+
+	/** Active world space axis X used for planar (only valid between state target BeginModify/EndModify) */
+	UPROPERTY()
+	FVector InteractionAxisX;
+
+	/** Active world space axis Y used for planar (only valid between state target BeginModify/EndModify) */
+	UPROPERTY()
+	FVector InteractionAxisY;
+
 	/** Active axis type (only valid between state target BeginModify/EndModify) */
 	UPROPERTY()
 	TEnumAsByte<EAxisList::Type> InteractionAxisType;
@@ -483,4 +522,12 @@ protected:
 	/** Active interaction current point (only valid between state target BeginModify/EndModify) */
 	UPROPERTY()
 	FVector InteractionCurrPoint;
+
+	/** Active interaction start point planar (only valid between state target BeginModify/EndModify) */
+	UPROPERTY()
+	FVector2D InteractionStartPoint2D;
+
+	/** Active interaction current point planar (only valid between state target BeginModify/EndModify) */
+	UPROPERTY()
+	FVector2D InteractionCurrPoint2D;
 };
