Fixed TRS gizmo scale delta calculation. Added new methods for converting mouse moves to translate/scale deltas.

diff --git a/Engine/Source/Editor/Experimental/EditorInteractiveToolsFramework/Private/EditorGizmos/EditorTransformGizmo.cpp b/Engine/Source/Editor/Experimental/EditorInteractiveToolsFramework/Private/EditorGizmos/EditorTransformGizmo.cpp
--- a/Engine/Source/Editor/Experimental/EditorInteractiveToolsFramework/Private/EditorGizmos/EditorTransformGizmo.cpp
+++ b/Engine/Source/Editor/Experimental/EditorInteractiveToolsFramework/Private/EditorGizmos/EditorTransformGizmo.cpp
@@ -13,7 +13,7 @@ void UEditorTransformGizmo::ApplyTranslateDelta(const FVector& InTranslateDelta)
 
 	if (UEditorTransformProxy* EditorTransformProxy = Cast<UEditorTransformProxy>(ActiveTarget))
 	{
-		EditorTransformProxy->InputTranslateDelta(InTranslateDelta, InteractionAxisType);
+		EditorTransformProxy->InputTranslateDelta(InTranslateDelta, InteractionAxisList);
 
 		// Update the cached current transform
 		CurrentTransform.AddToTranslation(InTranslateDelta);
@@ -32,10 +32,10 @@ void UEditorTransformGizmo::ApplyScaleDelta(const FVector& InScaleDelta)
 	{
 		FVector StartScale = CurrentTransform.GetScale3D();
 
-		EditorTransformProxy->InputScaleDelta(InScaleDelta, InteractionAxisType);
+		EditorTransformProxy->InputScaleDelta(InScaleDelta, InteractionAxisList);
 
 		// Update the cached current transform
-		FVector NewScale = StartScale + InScaleDelta * InteractionAxis;
+		FVector NewScale = StartScale + InScaleDelta;
 		CurrentTransform.SetScale3D(NewScale);
 	}
 	else
diff --git a/Engine/Source/Editor/Experimental/EditorInteractiveToolsFramework/Private/EditorGizmos/TransformGizmo.cpp b/Engine/Source/Editor/Experimental/EditorInteractiveToolsFramework/Private/EditorGizmos/TransformGizmo.cpp
--- a/Engine/Source/Editor/Experimental/EditorInteractiveToolsFramework/Private/EditorGizmos/TransformGizmo.cpp
+++ b/Engine/Source/Editor/Experimental/EditorInteractiveToolsFramework/Private/EditorGizmos/TransformGizmo.cpp
@@ -241,97 +241,97 @@ void UTransformGizmo::OnClickPress(const FInputDeviceRay& PressPos)
 	if (LastHitPart == ETransformGizmoPartIdentifier::TranslateXAxis)
 	{
 		InteractionAxisOrigin = CurrentTransform.GetLocation();
-		InteractionAxis = GetWorldAxis(FVector::XAxisVector);
-		InteractionAxisType = EAxisList::X;
+		InteractionAxisDirection = GetWorldAxis(FVector::XAxisVector);
+		InteractionAxisList = EAxisList::X;
 		OnClickPressAxis(PressPos);
 	}
 	else if (LastHitPart == ETransformGizmoPartIdentifier::TranslateYAxis)
 	{
 		InteractionAxisOrigin = CurrentTransform.GetLocation();
-		InteractionAxis = GetWorldAxis(FVector::YAxisVector);
-		InteractionAxisType = EAxisList::Y;
+		InteractionAxisDirection = GetWorldAxis(FVector::YAxisVector);
+		InteractionAxisList = EAxisList::Y;
 		OnClickPressAxis(PressPos);
 	}
 	else if (LastHitPart == ETransformGizmoPartIdentifier::TranslateZAxis)
 	{
 		InteractionAxisOrigin = CurrentTransform.GetLocation();
-		InteractionAxis = GetWorldAxis(FVector::ZAxisVector);
-		InteractionAxisType = EAxisList::Z;
+		InteractionAxisDirection = GetWorldAxis(FVector::ZAxisVector);
+		InteractionAxisList = EAxisList::Z;
 		OnClickPressAxis(PressPos);
 	}
 	else if (LastHitPart == ETransformGizmoPartIdentifier::ScaleXAxis)
 	{
 		InteractionAxisOrigin = CurrentTransform.GetLocation();
-		InteractionAxis = GetWorldAxis(FVector::XAxisVector);
-		InteractionAxisType = EAxisList::X;
+		InteractionAxisDirection = GetWorldAxis(FVector::XAxisVector);
+		InteractionAxisList = EAxisList::X;
 		OnClickPressAxis(PressPos);
 	}
 	else if (LastHitPart == ETransformGizmoPartIdentifier::ScaleYAxis)
 	{
 		InteractionAxisOrigin = CurrentTransform.GetLocation();
-		InteractionAxis = GetWorldAxis(FVector::YAxisVector);
-		InteractionAxisType = EAxisList::Y;
+		InteractionAxisDirection = GetWorldAxis(FVector::YAxisVector);
+		InteractionAxisList = EAxisList::Y;
 		OnClickPressAxis(PressPos);
 	}
 	else if (LastHitPart == ETransformGizmoPartIdentifier::ScaleZAxis)
 	{
 		InteractionAxisOrigin = CurrentTransform.GetLocation();
-		InteractionAxis = GetWorldAxis(FVector::ZAxisVector);
-		InteractionAxisType = EAxisList::Z;
+		InteractionAxisDirection = GetWorldAxis(FVector::ZAxisVector);
+		InteractionAxisList = EAxisList::Z;
 		OnClickPressAxis(PressPos);
 	}
 	if (LastHitPart == ETransformGizmoPartIdentifier::TranslateXYPlanar)
 	{
-		InteractionAxisOrigin = CurrentTransform.GetLocation();
-		InteractionNormal = GetWorldAxis(FVector::ZAxisVector);
-		InteractionAxisX = GetWorldAxis(FVector::XAxisVector);
-		InteractionAxisY = GetWorldAxis(FVector::YAxisVector);
-		InteractionAxisType = EAxisList::XY;
+		InteractionPlanarOrigin = CurrentTransform.GetLocation();
+		InteractionPlanarNormal = GetWorldAxis(FVector::ZAxisVector);
+		InteractionPlanarAxisX = GetWorldAxis(FVector::XAxisVector);
+		InteractionPlanarAxisY = GetWorldAxis(FVector::YAxisVector);
+		InteractionAxisList = EAxisList::XY;
 		OnClickPressPlanar(PressPos);
 	}
 	else if (LastHitPart == ETransformGizmoPartIdentifier::TranslateYZPlanar)
 	{
-		InteractionAxisOrigin = CurrentTransform.GetLocation();
-		InteractionNormal = GetWorldAxis(FVector::XAxisVector);
-		InteractionAxisX = GetWorldAxis(FVector::YAxisVector);
-		InteractionAxisY = GetWorldAxis(FVector::ZAxisVector);
-		InteractionAxisType = EAxisList::YZ;
+		InteractionPlanarOrigin = CurrentTransform.GetLocation();
+		InteractionPlanarNormal = GetWorldAxis(FVector::XAxisVector);
+		InteractionPlanarAxisX = GetWorldAxis(FVector::YAxisVector);
+		InteractionPlanarAxisY = GetWorldAxis(FVector::ZAxisVector);
+		InteractionAxisList = EAxisList::YZ;
 		OnClickPressPlanar(PressPos);
 	}
 	else if (LastHitPart == ETransformGizmoPartIdentifier::TranslateXZPlanar)
 	{
-		InteractionAxisOrigin = CurrentTransform.GetLocation();
-		InteractionNormal = GetWorldAxis(FVector::YAxisVector);
-		InteractionAxisX = GetWorldAxis(FVector::ZAxisVector);
-		InteractionAxisY = GetWorldAxis(FVector::XAxisVector);
-		InteractionAxisType = EAxisList::XZ;
+		InteractionPlanarOrigin = CurrentTransform.GetLocation();
+		InteractionPlanarNormal = GetWorldAxis(FVector::YAxisVector);
+		InteractionPlanarAxisX = GetWorldAxis(FVector::ZAxisVector);
+		InteractionPlanarAxisY = GetWorldAxis(FVector::XAxisVector);
+		InteractionAxisList = EAxisList::XZ;
 		OnClickPressPlanar(PressPos);
 	}
 	else if (LastHitPart == ETransformGizmoPartIdentifier::ScaleXYPlanar)
 	{
-		InteractionAxisOrigin = CurrentTransform.GetLocation();
-		InteractionNormal = GetWorldAxis(FVector::ZAxisVector);
-		InteractionAxisX = GetWorldAxis(FVector::XAxisVector);
-		InteractionAxisY = GetWorldAxis(FVector::YAxisVector);
-		InteractionAxisType = EAxisList::XY;
+		InteractionPlanarOrigin = CurrentTransform.GetLocation();
+		InteractionPlanarNormal = GetWorldAxis(FVector::ZAxisVector);
+		InteractionPlanarAxisX = GetWorldAxis(FVector::XAxisVector);
+		InteractionPlanarAxisY = GetWorldAxis(FVector::YAxisVector);
+		InteractionAxisList = EAxisList::XY;
 		OnClickPressPlanar(PressPos);
 	}
 	else if (LastHitPart == ETransformGizmoPartIdentifier::ScaleYZPlanar)
 	{
-		InteractionAxisOrigin = CurrentTransform.GetLocation();
-		InteractionNormal = GetWorldAxis(FVector::XAxisVector);
-		InteractionAxisX = GetWorldAxis(FVector::YAxisVector);
-		InteractionAxisY = GetWorldAxis(FVector::ZAxisVector);
-		InteractionAxisType = EAxisList::YZ;
+		InteractionPlanarOrigin = CurrentTransform.GetLocation();
+		InteractionPlanarNormal = GetWorldAxis(FVector::XAxisVector);
+		InteractionPlanarAxisX = GetWorldAxis(FVector::YAxisVector);
+		InteractionPlanarAxisY = GetWorldAxis(FVector::ZAxisVector);
+		InteractionAxisList = EAxisList::YZ;
 		OnClickPressPlanar(PressPos);
 	}
 	else if (LastHitPart == ETransformGizmoPartIdentifier::ScaleXZPlanar)
 	{
-		InteractionAxisOrigin = CurrentTransform.GetLocation();
-		InteractionNormal = GetWorldAxis(FVector::YAxisVector);
-		InteractionAxisX = GetWorldAxis(FVector::ZAxisVector);
-		InteractionAxisY = GetWorldAxis(FVector::XAxisVector);
-		InteractionAxisType = EAxisList::XZ;
+		InteractionPlanarOrigin = CurrentTransform.GetLocation();
+		InteractionPlanarNormal = GetWorldAxis(FVector::YAxisVector);
+		InteractionPlanarAxisX = GetWorldAxis(FVector::ZAxisVector);
+		InteractionPlanarAxisY = GetWorldAxis(FVector::XAxisVector);
+		InteractionAxisList = EAxisList::XZ;
 		OnClickPressPlanar(PressPos);
 	}
 
@@ -1030,14 +1030,14 @@ FVector UTransformGizmo::GetWorldAxis(const FVector& InAxis)
 void UTransformGizmo::OnClickPressAxis(const FInputDeviceRay& InPressPos)
 {
 	// Find interaction start point and parameter.
-	FVector NearestPt; 
-	float RayNearestParam, LineNearestParam;
-	GizmoMath::NearestPointOnLineToRay(InteractionAxisOrigin, InteractionAxis,
+	FVector NearestPt, InteractionStartPoint;
+	float RayNearestParam;
+	GizmoMath::NearestPointOnLineToRay(InteractionAxisOrigin, InteractionAxisDirection,
 		InPressPos.WorldRay.Origin, InPressPos.WorldRay.Direction,
-		InteractionStartPoint, LineNearestParam,
+		InteractionStartPoint, InteractionAxisStartParam,
 		NearestPt, RayNearestParam);
 
-	InteractionCurrPoint = InteractionStartPoint;
+	InteractionAxisCurrParam = InteractionAxisStartParam;
 	bInInteraction = true;
 }
 
@@ -1045,25 +1045,27 @@ void UTransformGizmo::OnClickDragAxis(const FInputDeviceRay& InDragPos)
 {
 	float RayNearestParam, AxisNearestParam;
 	FVector RayNearestPt, AxisNearestPoint;
-	GizmoMath::NearestPointOnLineToRay(InteractionAxisOrigin, InteractionAxis,
+	GizmoMath::NearestPointOnLineToRay(InteractionAxisOrigin, InteractionAxisDirection,
 		InDragPos.WorldRay.Origin, InDragPos.WorldRay.Direction,
 		AxisNearestPoint, AxisNearestParam,
 		RayNearestPt, RayNearestParam);
 
-	FVector Delta = AxisNearestPoint - InteractionCurrPoint;
-	InteractionCurrPoint = AxisNearestPoint;
-
 	if (LastHitPart == ETransformGizmoPartIdentifier::TranslateXAxis ||
 		LastHitPart == ETransformGizmoPartIdentifier::TranslateYAxis ||
 		LastHitPart == ETransformGizmoPartIdentifier::TranslateZAxis)
 	{
+		FVector Delta;
+		ComputeAxisTranslateDelta(InteractionAxisCurrParam, AxisNearestParam, Delta);
 		ApplyTranslateDelta(Delta);
 	}
 	else
 	{
-		Delta *= ScaleMultiplier;
+		FVector Delta;
+		ComputeAxisScaleDelta(InteractionAxisCurrParam, AxisNearestParam, Delta);
 		ApplyScaleDelta(Delta);
 	}
+
+	InteractionAxisCurrParam = AxisNearestParam;
 }
 
 void UTransformGizmo::OnClickReleaseAxis(const FInputDeviceRay& InReleasePos)
@@ -1074,21 +1076,20 @@ void UTransformGizmo::OnClickReleaseAxis(const FInputDeviceRay& InReleasePos)
 void UTransformGizmo::OnClickPressPlanar(const FInputDeviceRay& PressPos)
 {
 	// if ray is parallel to plane, nothing has been hit
-	if (FMath::IsNearlyZero(FVector::DotProduct(InteractionNormal, PressPos.WorldRay.Direction)))
+	if (FMath::IsNearlyZero(FVector::DotProduct(InteractionPlanarNormal, PressPos.WorldRay.Direction)))
 	{
 		return;
 	}
 
-	FPlane Plane(InteractionAxisOrigin, InteractionNormal);
+	FPlane Plane(InteractionPlanarOrigin, InteractionPlanarNormal);
 	double HitDepth = FMath::RayPlaneIntersectionParam(PressPos.WorldRay.Origin, PressPos.WorldRay.Direction, Plane);
 	if (HitDepth < 0)
 	{
 		return;
 	}
 
-	FVector HitPoint = PressPos.WorldRay.Origin + PressPos.WorldRay.Direction * HitDepth;
-	InteractionStartPoint = HitPoint;
-	InteractionCurrPoint = InteractionStartPoint;
+	InteractionPlanarStartPoint = PressPos.WorldRay.Origin + PressPos.WorldRay.Direction * HitDepth;
+	InteractionPlanarCurrPoint = InteractionPlanarStartPoint;
 
 	bInInteraction = true;
 }
@@ -1096,12 +1097,12 @@ void UTransformGizmo::OnClickPressPlanar(const FInputDeviceRay& PressPos)
 void UTransformGizmo::OnClickDragPlanar(const FInputDeviceRay& DragPos)
 {
 	// if ray is parallel to plane, nothing has been hit
-	if (FMath::IsNearlyZero(FVector::DotProduct(InteractionNormal, DragPos.WorldRay.Direction)))
+	if (FMath::IsNearlyZero(FVector::DotProduct(InteractionPlanarNormal, DragPos.WorldRay.Direction)))
 	{
 		return;
 	}
 
-	FPlane Plane(InteractionAxisOrigin, InteractionNormal);
+	FPlane Plane(InteractionPlanarOrigin, InteractionPlanarNormal);
 	double HitDepth = FMath::RayPlaneIntersectionParam(DragPos.WorldRay.Origin, DragPos.WorldRay.Direction, Plane);
 	if (HitDepth < 0)
 	{
@@ -1109,28 +1110,23 @@ void UTransformGizmo::OnClickDragPlanar(const FInputDeviceRay& DragPos)
 	}
 
 	FVector HitPoint = DragPos.WorldRay.Origin + DragPos.WorldRay.Direction * HitDepth;
-	FVector Delta = HitPoint - InteractionCurrPoint;
-
-	InteractionCurrPoint = HitPoint;
 
 	if (LastHitPart == ETransformGizmoPartIdentifier::TranslateXYPlanar ||
 		LastHitPart == ETransformGizmoPartIdentifier::TranslateYZPlanar ||
 		LastHitPart == ETransformGizmoPartIdentifier::TranslateXZPlanar)
 	{
+		FVector Delta;
+		ComputePlanarTranslateDelta(InteractionPlanarCurrPoint, HitPoint, Delta);
 		ApplyTranslateDelta(Delta);
 	}
 	else
 	{
-		const float ScaleMax = Delta.GetMax();
-		const float ScaleMin = Delta.GetMin();
-		const float ScaleApplied = (ScaleMax > -ScaleMin) ? ScaleMax * ScaleMultiplier : ScaleMin * ScaleMultiplier;
-
-		FVector ScaleDelta(InteractionAxisType & EAxisList::X ? ScaleApplied : 0.0,
-						   InteractionAxisType & EAxisList::Y ? ScaleApplied : 0.0,
-						   InteractionAxisType & EAxisList::Z ? ScaleApplied : 0.0);
-
-		ApplyScaleDelta(ScaleDelta);
+		FVector Delta;
+		ComputePlanarScaleDelta(InteractionPlanarCurrPoint, HitPoint, Delta);
+		ApplyScaleDelta(Delta);
 	}
+
+	InteractionPlanarCurrPoint = HitPoint;
 }
 
 void UTransformGizmo::OnClickReleasePlanar(const FInputDeviceRay& InReleasePos)
@@ -1138,6 +1134,40 @@ void UTransformGizmo::OnClickReleasePlanar(const FInputDeviceRay& InReleasePos)
 	bInInteraction = false;
 }
 
+void UTransformGizmo::ComputeAxisTranslateDelta(double InStartParam, double InEndParam, FVector& OutTranslateDelta)
+{
+	const double ParamDelta = InEndParam - InStartParam;
+	OutTranslateDelta = InteractionAxisDirection * ParamDelta;
+}
+
+void UTransformGizmo::ComputeAxisScaleDelta(double InStartParam, double InEndParam, FVector& OutScaleDelta)
+{
+	const double ParamDelta = InEndParam - InStartParam;
+	const float ScaleApplied = ParamDelta * ScaleMultiplier;
+
+	OutScaleDelta = FVector(InteractionAxisList & EAxisList::X ? ScaleApplied : 0.0,
+		InteractionAxisList & EAxisList::Y ? ScaleApplied : 0.0,
+		InteractionAxisList & EAxisList::Z ? ScaleApplied : 0.0);
+}
+
+void UTransformGizmo::ComputePlanarTranslateDelta(const FVector& InStartPoint, const FVector& InEndPoint, FVector& OutTranslateDelta)
+{
+	OutTranslateDelta = InEndPoint - InStartPoint;
+}
+
+void UTransformGizmo::ComputePlanarScaleDelta(const FVector& InStartPoint, const FVector& InEndPoint, FVector& OutScaleDelta)
+{
+	FVector Delta = InEndPoint - InStartPoint;
+	float DragUp = FVector::DotProduct(Delta, InteractionPlanarAxisX);
+	float DragSide = FVector::DotProduct(Delta, InteractionPlanarAxisY);
+	const float ScaleApplied = FMath::Abs(DragUp) > FMath::Abs(DragSide) ? DragUp * ScaleMultiplier : DragSide * ScaleMultiplier;
+
+	OutScaleDelta = FVector(
+		InteractionAxisList & EAxisList::X ? ScaleApplied : 0.0,
+		InteractionAxisList & EAxisList::Y ? ScaleApplied : 0.0,
+		InteractionAxisList & EAxisList::Z ? ScaleApplied : 0.0);
+}
+
 void UTransformGizmo::ApplyTranslateDelta(const FVector& InTranslateDelta)
 {
 	CurrentTransform.AddToTranslation(InTranslateDelta);
@@ -1147,7 +1177,7 @@ void UTransformGizmo::ApplyTranslateDelta(const FVector& InTranslateDelta)
 void UTransformGizmo::ApplyScaleDelta(const FVector& InScaleDelta)
 {
 	FVector StartScale = CurrentTransform.GetScale3D();
-	FVector NewScale = StartScale + InScaleDelta * InteractionAxis;
+	FVector NewScale = StartScale + InScaleDelta;
 	CurrentTransform.SetScale3D(NewScale);
 	ActiveTarget->SetTransform(CurrentTransform);
 }
diff --git a/Engine/Source/Editor/Experimental/EditorInteractiveToolsFramework/Public/EditorGizmos/TransformGizmo.h b/Engine/Source/Editor/Experimental/EditorInteractiveToolsFramework/Public/EditorGizmos/TransformGizmo.h
--- a/Engine/Source/Editor/Experimental/EditorInteractiveToolsFramework/Public/EditorGizmos/TransformGizmo.h
+++ b/Engine/Source/Editor/Experimental/EditorInteractiveToolsFramework/Public/EditorGizmos/TransformGizmo.h
@@ -428,6 +428,18 @@ protected:
 	/** Handle click release for translate and scale planar */
 	virtual void OnClickReleasePlanar(const FInputDeviceRay& PressPos);
 
+	/** Compute translate axis delta based on start/end params */
+	virtual void ComputeAxisTranslateDelta(double InStartParam, double InEndParam, FVector& OutDelta);
+
+	/** Compute scale axis delta based on start/end params */
+	virtual void ComputeAxisScaleDelta(double InStartParam, double InEndParam, FVector& OutDelta);
+
+	/** Compute translate planar delta based on world space start/end points */
+	virtual void ComputePlanarTranslateDelta(const FVector& InStartPoint, const FVector& InEndPoint, FVector& OutDelta);
+
+	/** Compute scale planar delta based on world space start/end points */
+	virtual void ComputePlanarScaleDelta(const FVector& InStartPoint, const FVector& InEndPoint, FVector& OutDelta);
+
 	/** Apply translate delta to transform proxy */
 	virtual void ApplyTranslateDelta(const FVector& InTranslateDelta);
 
@@ -491,43 +503,47 @@ protected:
 	// They otherwise should be considered uninitialized
 	//
 
+	/** Active axis type (only valid between state target BeginModify/EndModify) */
+	UPROPERTY()
+	TEnumAsByte<EAxisList::Type> InteractionAxisList;
+
 	/** Active world space axis origin (only valid between state target BeginModify/EndModify) */
 	UPROPERTY()
 	FVector InteractionAxisOrigin;
 
 	/** Active world space axis (only valid between state target BeginModify/EndModify) */
 	UPROPERTY()
-	FVector InteractionAxis;
+	FVector InteractionAxisDirection;
 
-	/** Active world space normal used for planar (only valid between state target BeginModify/EndModify) */
+	/** Active interaction start hit param (only valid between state target BeginModify/EndModify) */
 	UPROPERTY()
-	FVector InteractionNormal;
+	float InteractionAxisStartParam;
 
-	/** Active world space axis X used for planar (only valid between state target BeginModify/EndModify) */
+	/** Active interaction current hit param (only valid between state target BeginModify/EndModify) */
 	UPROPERTY()
-	FVector InteractionAxisX;
+	float InteractionAxisCurrParam;
 
-	/** Active world space axis Y used for planar (only valid between state target BeginModify/EndModify) */
+	/** Active world space planar origin (only valid between state target BeginModify/EndModify) */
 	UPROPERTY()
-	FVector InteractionAxisY;
+	FVector InteractionPlanarOrigin;
 
-	/** Active axis type (only valid between state target BeginModify/EndModify) */
+	/** Active world space normal used for planar (only valid between state target BeginModify/EndModify) */
 	UPROPERTY()
-	TEnumAsByte<EAxisList::Type> InteractionAxisType;
+	FVector InteractionPlanarNormal;
 
-	/** Active interaction start point (only valid between state target BeginModify/EndModify) */
+	/** Active world space axis X used for planar (only valid between state target BeginModify/EndModify) */
 	UPROPERTY()
-	FVector InteractionStartPoint;
+	FVector InteractionPlanarAxisX;
 
-	/** Active interaction current point (only valid between state target BeginModify/EndModify) */
+	/** Active world space axis Y used for planar (only valid between state target BeginModify/EndModify) */
 	UPROPERTY()
-	FVector InteractionCurrPoint;
+	FVector InteractionPlanarAxisY;
 
 	/** Active interaction start point planar (only valid between state target BeginModify/EndModify) */
 	UPROPERTY()
-	FVector2D InteractionStartPoint2D;
+	FVector InteractionPlanarStartPoint;
 
 	/** Active interaction current point planar (only valid between state target BeginModify/EndModify) */
 	UPROPERTY()
-	FVector2D InteractionCurrPoint2D;
+	FVector InteractionPlanarCurrPoint;
 };
