ToolsFramework: Add UEditorInteractiveToolsContext::OnTick delegate, similar to OnRender and OnDrawHUD, allows code w/ access to the Context to hook into the Context Tick w/o having to override it/etc.

diff --git a/Engine/Source/Editor/Experimental/EditorInteractiveToolsFramework/Private/EdModeInteractiveToolsContext.cpp b/Engine/Source/Editor/Experimental/EditorInteractiveToolsFramework/Private/EdModeInteractiveToolsContext.cpp
--- a/Engine/Source/Editor/Experimental/EditorInteractiveToolsFramework/Private/EdModeInteractiveToolsContext.cpp
+++ b/Engine/Source/Editor/Experimental/EditorInteractiveToolsFramework/Private/EdModeInteractiveToolsContext.cpp
@@ -254,10 +254,16 @@ public:
 		int NumActors = SelectionChange.Actors.Num();
 		for (int k = 0; k < NumActors; ++k)
 		{
-			GEditor->SelectActor(SelectionChange.Actors[k], bAdd, false, true, false);
+			// Calling GEditor->NoteSelectionChange(true) will not send out change notifications on the TypedElementSelectionSet.
+			// The selection change will work but any Editor stuff listening for changes will not be notified.
+			// This may be a bug, needs further investigation.
+			// In the meantime, just send out the notification on the last SelectActor() call
+			bool bNotify = (k == NumActors-1);
+			GEditor->SelectActor(SelectionChange.Actors[k], bAdd, bNotify, true, false);
 		}
 
 		GEditor->NoteSelectionChange(true);
+
 		return true;
 	}
 
@@ -440,6 +446,7 @@ void UEditorInteractiveToolsContext::Tick(FEditorViewportClient* ViewportClient,
 	// tick our stuff
 	ToolManager->Tick(DeltaTime);
 	GizmoManager->Tick(DeltaTime);
+	OnTick.Broadcast(DeltaTime);
 }
 
 
diff --git a/Engine/Source/Editor/Experimental/EditorInteractiveToolsFramework/Public/EdModeInteractiveToolsContext.h b/Engine/Source/Editor/Experimental/EditorInteractiveToolsFramework/Public/EdModeInteractiveToolsContext.h
--- a/Engine/Source/Editor/Experimental/EditorInteractiveToolsFramework/Public/EdModeInteractiveToolsContext.h
+++ b/Engine/Source/Editor/Experimental/EditorInteractiveToolsFramework/Public/EdModeInteractiveToolsContext.h
@@ -77,12 +77,14 @@ public:
 	virtual void Render(const FSceneView* View, FViewport* Viewport, FPrimitiveDrawInterface* PDI);
 	virtual void DrawHUD(FViewportClient* ViewportClient,FViewport* Viewport,const FSceneView* View, FCanvas* Canvas);
 
-	// These delegates can be used to hook in to Render() and DrawHUD() calls. In particular, non-legacy UEdMode's
+	// These delegates can be used to hook into the Render() / DrawHUD() / Tick() calls above. In particular, non-legacy UEdMode's
 	// don't normally receive Render() and DrawHUD() calls from the mode manager, but can attach to these.
 	DECLARE_MULTICAST_DELEGATE_OneParam(FOnRender, IToolsContextRenderAPI* RenderAPI);
 	FOnRender OnRender;
 	DECLARE_MULTICAST_DELEGATE_TwoParams(FOnDrawHUD, FCanvas* Canvas, IToolsContextRenderAPI* RenderAPI);
 	FOnDrawHUD OnDrawHUD;
+	DECLARE_MULTICAST_DELEGATE_OneParam(FOnTick, float DeltaTime);
+	FOnTick OnTick;
 
 	/** @return true if selected actors/components can be deleted */
 	virtual bool ProcessEditDelete();
