Update TransformGizmo ClickDrag handling methods to use function tables, in preparation for adding more interactions.

diff --git a/Engine/Source/Editor/Experimental/EditorInteractiveToolsFramework/Private/EditorGizmos/TransformGizmo.cpp b/Engine/Source/Editor/Experimental/EditorInteractiveToolsFramework/Private/EditorGizmos/TransformGizmo.cpp
--- a/Engine/Source/Editor/Experimental/EditorInteractiveToolsFramework/Private/EditorGizmos/TransformGizmo.cpp
+++ b/Engine/Source/Editor/Experimental/EditorInteractiveToolsFramework/Private/EditorGizmos/TransformGizmo.cpp
@@ -34,6 +34,7 @@ void UTransformGizmo::Setup()
 
 	SetupBehaviors();
 	SetupMaterials();
+	SetupOnClickFunctions();
 
 	// @todo: Gizmo element construction will be moved to the UEditorTransformGizmoBuilder to decouple
 	// the rendered elements from the transform gizmo.
@@ -238,205 +239,6 @@ FInputRayHit UTransformGizmo::CanBeginClickDragSequence(const FInputDeviceRay& P
 	return RayHit;
 }
 
-void UTransformGizmo::OnClickPress(const FInputDeviceRay& PressPos)
-{
-	if (LastHitPart == ETransformGizmoPartIdentifier::TranslateXAxis)
-	{
-		InteractionAxisOrigin = CurrentTransform.GetLocation();
-		InteractionAxisDirection = GetWorldAxis(FVector::XAxisVector);
-		InteractionAxisList = EAxisList::X;
-		OnClickPressAxis(PressPos);
-	}
-	else if (LastHitPart == ETransformGizmoPartIdentifier::TranslateYAxis)
-	{
-		InteractionAxisOrigin = CurrentTransform.GetLocation();
-		InteractionAxisDirection = GetWorldAxis(FVector::YAxisVector);
-		InteractionAxisList = EAxisList::Y;
-		OnClickPressAxis(PressPos);
-	}
-	else if (LastHitPart == ETransformGizmoPartIdentifier::TranslateZAxis)
-	{
-		InteractionAxisOrigin = CurrentTransform.GetLocation();
-		InteractionAxisDirection = GetWorldAxis(FVector::ZAxisVector);
-		InteractionAxisList = EAxisList::Z;
-		OnClickPressAxis(PressPos);
-	}
-	else if (LastHitPart == ETransformGizmoPartIdentifier::ScaleXAxis)
-	{
-		InteractionAxisOrigin = CurrentTransform.GetLocation();
-		InteractionAxisDirection = GetWorldAxis(FVector::XAxisVector);
-		InteractionAxisList = EAxisList::X;
-		OnClickPressAxis(PressPos);
-	}
-	else if (LastHitPart == ETransformGizmoPartIdentifier::ScaleYAxis)
-	{
-		InteractionAxisOrigin = CurrentTransform.GetLocation();
-		InteractionAxisDirection = GetWorldAxis(FVector::YAxisVector);
-		InteractionAxisList = EAxisList::Y;
-		OnClickPressAxis(PressPos);
-	}
-	else if (LastHitPart == ETransformGizmoPartIdentifier::ScaleZAxis)
-	{
-		InteractionAxisOrigin = CurrentTransform.GetLocation();
-		InteractionAxisDirection = GetWorldAxis(FVector::ZAxisVector);
-		InteractionAxisList = EAxisList::Z;
-		OnClickPressAxis(PressPos);
-	}
-	if (LastHitPart == ETransformGizmoPartIdentifier::TranslateXYPlanar)
-	{
-		InteractionPlanarOrigin = CurrentTransform.GetLocation();
-		InteractionPlanarNormal = GetWorldAxis(FVector::ZAxisVector);
-		InteractionPlanarAxisX = GetWorldAxis(FVector::XAxisVector);
-		InteractionPlanarAxisY = GetWorldAxis(FVector::YAxisVector);
-		InteractionAxisList = EAxisList::XY;
-		OnClickPressPlanar(PressPos);
-	}
-	else if (LastHitPart == ETransformGizmoPartIdentifier::TranslateYZPlanar)
-	{
-		InteractionPlanarOrigin = CurrentTransform.GetLocation();
-		InteractionPlanarNormal = GetWorldAxis(FVector::XAxisVector);
-		InteractionPlanarAxisX = GetWorldAxis(FVector::YAxisVector);
-		InteractionPlanarAxisY = GetWorldAxis(FVector::ZAxisVector);
-		InteractionAxisList = EAxisList::YZ;
-		OnClickPressPlanar(PressPos);
-	}
-	else if (LastHitPart == ETransformGizmoPartIdentifier::TranslateXZPlanar)
-	{
-		InteractionPlanarOrigin = CurrentTransform.GetLocation();
-		InteractionPlanarNormal = GetWorldAxis(FVector::YAxisVector);
-		InteractionPlanarAxisX = GetWorldAxis(FVector::ZAxisVector);
-		InteractionPlanarAxisY = GetWorldAxis(FVector::XAxisVector);
-		InteractionAxisList = EAxisList::XZ;
-		OnClickPressPlanar(PressPos);
-	}
-	else if (LastHitPart == ETransformGizmoPartIdentifier::ScaleXYPlanar)
-	{
-		InteractionPlanarOrigin = CurrentTransform.GetLocation();
-		InteractionPlanarNormal = GetWorldAxis(FVector::ZAxisVector);
-		InteractionPlanarAxisX = GetWorldAxis(FVector::XAxisVector);
-		InteractionPlanarAxisY = GetWorldAxis(FVector::YAxisVector);
-		InteractionAxisList = EAxisList::XY;
-		OnClickPressPlanar(PressPos);
-	}
-	else if (LastHitPart == ETransformGizmoPartIdentifier::ScaleYZPlanar)
-	{
-		InteractionPlanarOrigin = CurrentTransform.GetLocation();
-		InteractionPlanarNormal = GetWorldAxis(FVector::XAxisVector);
-		InteractionPlanarAxisX = GetWorldAxis(FVector::YAxisVector);
-		InteractionPlanarAxisY = GetWorldAxis(FVector::ZAxisVector);
-		InteractionAxisList = EAxisList::YZ;
-		OnClickPressPlanar(PressPos);
-	}
-	else if (LastHitPart == ETransformGizmoPartIdentifier::ScaleXZPlanar)
-	{
-		InteractionPlanarOrigin = CurrentTransform.GetLocation();
-		InteractionPlanarNormal = GetWorldAxis(FVector::YAxisVector);
-		InteractionPlanarAxisX = GetWorldAxis(FVector::ZAxisVector);
-		InteractionPlanarAxisY = GetWorldAxis(FVector::XAxisVector);
-		InteractionAxisList = EAxisList::XZ;
-		OnClickPressPlanar(PressPos);
-	}
-
-	if (bInInteraction)
-	{
-		if (HitTarget&& LastHitPart != ETransformGizmoPartIdentifier::Default)
-		{
-			HitTarget->UpdateInteractingState(true, static_cast<uint32>(LastHitPart));
-		}
-
-		if (StateTarget)
-		{
-			StateTarget->BeginUpdate();
-		}
-	}
-}
-
-void UTransformGizmo::OnClickDrag(const FInputDeviceRay& DragPos)
-{
-	if (!bInInteraction)
-	{
-		return;
-	}
-
-	if (LastHitPart == ETransformGizmoPartIdentifier::TranslateXAxis ||
-		LastHitPart == ETransformGizmoPartIdentifier::TranslateYAxis ||
-		LastHitPart == ETransformGizmoPartIdentifier::TranslateZAxis ||
-		LastHitPart == ETransformGizmoPartIdentifier::ScaleXAxis ||
-		LastHitPart == ETransformGizmoPartIdentifier::ScaleYAxis ||
-		LastHitPart == ETransformGizmoPartIdentifier::ScaleZAxis)
-	{
-		OnClickDragAxis(DragPos);
-	}
-	else if (LastHitPart == ETransformGizmoPartIdentifier::TranslateXYPlanar ||
-		LastHitPart == ETransformGizmoPartIdentifier::TranslateYZPlanar ||
-		LastHitPart == ETransformGizmoPartIdentifier::TranslateXZPlanar ||
-		LastHitPart == ETransformGizmoPartIdentifier::ScaleXYPlanar ||
-		LastHitPart == ETransformGizmoPartIdentifier::ScaleYZPlanar ||
-		LastHitPart == ETransformGizmoPartIdentifier::ScaleXZPlanar)
-	{
-		OnClickDragPlanar(DragPos);
-	}
-}
-
-void UTransformGizmo::OnClickRelease(const FInputDeviceRay& ReleasePos)
-{
-	if (!bInInteraction)
-	{
-		return;
-	}
-
-	if (LastHitPart == ETransformGizmoPartIdentifier::TranslateXAxis ||
-		LastHitPart == ETransformGizmoPartIdentifier::TranslateYAxis ||
-		LastHitPart == ETransformGizmoPartIdentifier::TranslateZAxis ||
-		LastHitPart == ETransformGizmoPartIdentifier::ScaleXAxis ||
-		LastHitPart == ETransformGizmoPartIdentifier::ScaleYAxis ||
-		LastHitPart == ETransformGizmoPartIdentifier::ScaleZAxis)
-	{
-		OnClickReleaseAxis(ReleasePos);
-	}
-	else if (LastHitPart == ETransformGizmoPartIdentifier::TranslateXYPlanar ||
-		LastHitPart == ETransformGizmoPartIdentifier::TranslateYZPlanar ||
-		LastHitPart == ETransformGizmoPartIdentifier::TranslateXZPlanar ||
-		LastHitPart == ETransformGizmoPartIdentifier::ScaleXYPlanar ||
-		LastHitPart == ETransformGizmoPartIdentifier::ScaleYZPlanar ||
-		LastHitPart == ETransformGizmoPartIdentifier::ScaleXZPlanar)
-	{
-		OnClickReleaseAxis(ReleasePos);
-	}
-
-	if (StateTarget)
-	{
-		StateTarget->EndUpdate();
-	}
-
-	bInInteraction = false;
-
-	if (HitTarget && LastHitPart != ETransformGizmoPartIdentifier::Default)
-	{
-		HitTarget->UpdateInteractingState(false, static_cast<uint32>(LastHitPart));
-	}
-}
-
-
-void UTransformGizmo::OnTerminateDragSequence()
-{
-	if (!bInInteraction)
-	{
-		return;
-	}
-
-	if (StateTarget)
-	{
-		StateTarget->EndUpdate();
-	}
-	bInInteraction = false;
-
-	if (HitTarget && LastHitPart != ETransformGizmoPartIdentifier::Default)
-	{
-		HitTarget->UpdateInteractingState(false, static_cast<uint32>(LastHitPart));
-	}
-}
-
 void UTransformGizmo::UpdateMode()
 {
 	if (TransformGizmoSource && TransformGizmoSource->GetVisible())
@@ -1032,142 +834,384 @@ FVector UTransformGizmo::GetWorldAxis(const FVector& InAxis)
 	return InAxis;
 }
 
-void UTransformGizmo::OnClickPressAxis(const FInputDeviceRay& InPressPos)
-{
-	// Find interaction start point and parameter.
-	FVector NearestPt, InteractionStartPoint;
-	float RayNearestParam;
-	GizmoMath::NearestPointOnLineToRay(InteractionAxisOrigin, InteractionAxisDirection,
-		InPressPos.WorldRay.Origin, InPressPos.WorldRay.Direction,
-		InteractionStartPoint, InteractionAxisStartParam,
-		NearestPt, RayNearestParam);
 
-	InteractionAxisCurrParam = InteractionAxisStartParam;
-	bInInteraction = true;
+void UTransformGizmo::SetupOnClickFunctions()
+{
+	int NumParts = static_cast<int>(ETransformGizmoPartIdentifier::Max);
+	OnClickPressFunctions.SetNum(NumParts);
+	OnClickDragFunctions.SetNum(NumParts);
+	OnClickReleaseFunctions.SetNum(NumParts);
+
+	OnClickPressFunctions[static_cast<int>(ETransformGizmoPartIdentifier::TranslateXAxis)] = &UTransformGizmo::OnClickPressTranslateXAxis;
+	OnClickPressFunctions[static_cast<int>(ETransformGizmoPartIdentifier::TranslateYAxis)] = &UTransformGizmo::OnClickPressTranslateYAxis;
+	OnClickPressFunctions[static_cast<int>(ETransformGizmoPartIdentifier::TranslateZAxis)] = &UTransformGizmo::OnClickPressTranslateZAxis;
+	OnClickPressFunctions[static_cast<int>(ETransformGizmoPartIdentifier::ScaleXAxis)] = &UTransformGizmo::OnClickPressScaleXAxis;
+	OnClickPressFunctions[static_cast<int>(ETransformGizmoPartIdentifier::ScaleYAxis)] = &UTransformGizmo::OnClickPressScaleYAxis;
+	OnClickPressFunctions[static_cast<int>(ETransformGizmoPartIdentifier::ScaleZAxis)] = &UTransformGizmo::OnClickPressScaleZAxis;
+	OnClickPressFunctions[static_cast<int>(ETransformGizmoPartIdentifier::TranslateXYPlanar)] = &UTransformGizmo::OnClickPressTranslateXYPlanar;
+	OnClickPressFunctions[static_cast<int>(ETransformGizmoPartIdentifier::TranslateYZPlanar)] = &UTransformGizmo::OnClickPressTranslateYZPlanar;
+	OnClickPressFunctions[static_cast<int>(ETransformGizmoPartIdentifier::TranslateXZPlanar)] = &UTransformGizmo::OnClickPressTranslateXZPlanar;
+	OnClickPressFunctions[static_cast<int>(ETransformGizmoPartIdentifier::ScaleXYPlanar)] = &UTransformGizmo::OnClickPressScaleXYPlanar;
+	OnClickPressFunctions[static_cast<int>(ETransformGizmoPartIdentifier::ScaleYZPlanar)] = &UTransformGizmo::OnClickPressScaleYZPlanar;
+	OnClickPressFunctions[static_cast<int>(ETransformGizmoPartIdentifier::ScaleXZPlanar)] = &UTransformGizmo::OnClickPressScaleXZPlanar;
+
+	OnClickDragFunctions[static_cast<int>(ETransformGizmoPartIdentifier::TranslateXAxis)] = &UTransformGizmo::OnClickDragTranslateAxis;
+	OnClickDragFunctions[static_cast<int>(ETransformGizmoPartIdentifier::TranslateYAxis)] = &UTransformGizmo::OnClickDragTranslateAxis;
+	OnClickDragFunctions[static_cast<int>(ETransformGizmoPartIdentifier::TranslateZAxis)] = &UTransformGizmo::OnClickDragTranslateAxis;
+	OnClickDragFunctions[static_cast<int>(ETransformGizmoPartIdentifier::ScaleXAxis)] = &UTransformGizmo::OnClickDragScaleAxis;
+	OnClickDragFunctions[static_cast<int>(ETransformGizmoPartIdentifier::ScaleYAxis)] = &UTransformGizmo::OnClickDragScaleAxis;
+	OnClickDragFunctions[static_cast<int>(ETransformGizmoPartIdentifier::ScaleZAxis)] = &UTransformGizmo::OnClickDragScaleAxis;
+	OnClickDragFunctions[static_cast<int>(ETransformGizmoPartIdentifier::TranslateXYPlanar)] = &UTransformGizmo::OnClickDragTranslatePlanar;
+	OnClickDragFunctions[static_cast<int>(ETransformGizmoPartIdentifier::TranslateYZPlanar)] = &UTransformGizmo::OnClickDragTranslatePlanar;
+	OnClickDragFunctions[static_cast<int>(ETransformGizmoPartIdentifier::TranslateXZPlanar)] = &UTransformGizmo::OnClickDragTranslatePlanar;
+	OnClickDragFunctions[static_cast<int>(ETransformGizmoPartIdentifier::ScaleXYPlanar)] = &UTransformGizmo::OnClickDragScalePlanar;
+	OnClickDragFunctions[static_cast<int>(ETransformGizmoPartIdentifier::ScaleYZPlanar)] = &UTransformGizmo::OnClickDragScalePlanar;
+	OnClickDragFunctions[static_cast<int>(ETransformGizmoPartIdentifier::ScaleXZPlanar)] = &UTransformGizmo::OnClickDragScalePlanar;
+
+	OnClickReleaseFunctions[static_cast<int>(ETransformGizmoPartIdentifier::TranslateXAxis)] = &UTransformGizmo::OnClickReleaseTranslateAxis;
+	OnClickReleaseFunctions[static_cast<int>(ETransformGizmoPartIdentifier::TranslateYAxis)] = &UTransformGizmo::OnClickReleaseTranslateAxis;
+	OnClickReleaseFunctions[static_cast<int>(ETransformGizmoPartIdentifier::TranslateZAxis)] = &UTransformGizmo::OnClickReleaseTranslateAxis;
+	OnClickReleaseFunctions[static_cast<int>(ETransformGizmoPartIdentifier::ScaleXAxis)] = &UTransformGizmo::OnClickReleaseScaleAxis;
+	OnClickReleaseFunctions[static_cast<int>(ETransformGizmoPartIdentifier::ScaleYAxis)] = &UTransformGizmo::OnClickReleaseScaleAxis;
+	OnClickReleaseFunctions[static_cast<int>(ETransformGizmoPartIdentifier::ScaleZAxis)] = &UTransformGizmo::OnClickReleaseScaleAxis;
+	OnClickReleaseFunctions[static_cast<int>(ETransformGizmoPartIdentifier::TranslateXYPlanar)] = &UTransformGizmo::OnClickReleaseTranslatePlanar;
+	OnClickReleaseFunctions[static_cast<int>(ETransformGizmoPartIdentifier::TranslateYZPlanar)] = &UTransformGizmo::OnClickReleaseTranslatePlanar;
+	OnClickReleaseFunctions[static_cast<int>(ETransformGizmoPartIdentifier::TranslateXZPlanar)] = &UTransformGizmo::OnClickReleaseTranslatePlanar;
+	OnClickReleaseFunctions[static_cast<int>(ETransformGizmoPartIdentifier::ScaleXYPlanar)] = &UTransformGizmo::OnClickReleaseScalePlanar;
+	OnClickReleaseFunctions[static_cast<int>(ETransformGizmoPartIdentifier::ScaleYZPlanar)] = &UTransformGizmo::OnClickReleaseScalePlanar;
+	OnClickReleaseFunctions[static_cast<int>(ETransformGizmoPartIdentifier::ScaleXZPlanar)] = &UTransformGizmo::OnClickReleaseScalePlanar;
 }
 
-void UTransformGizmo::OnClickDragAxis(const FInputDeviceRay& InDragPos)
+
+float UTransformGizmo::GetNearestRayParamToInteractionAxis(const FInputDeviceRay& InRay)
 {
 	float RayNearestParam, AxisNearestParam;
 	FVector RayNearestPt, AxisNearestPoint;
 	GizmoMath::NearestPointOnLineToRay(InteractionAxisOrigin, InteractionAxisDirection,
-		InDragPos.WorldRay.Origin, InDragPos.WorldRay.Direction,
-		AxisNearestPoint, AxisNearestParam,
+		InRay.WorldRay.Origin, InRay.WorldRay.Direction, AxisNearestPoint, AxisNearestParam,
 		RayNearestPt, RayNearestParam);
+	return AxisNearestParam;
+}
 
-	if (LastHitPart == ETransformGizmoPartIdentifier::TranslateXAxis ||
-		LastHitPart == ETransformGizmoPartIdentifier::TranslateYAxis ||
-		LastHitPart == ETransformGizmoPartIdentifier::TranslateZAxis)
+bool UTransformGizmo::GetRayParamIntersectionWithInteractionPlane(const FInputDeviceRay& InRay, float& OutHitParam)
+{
+	// if ray is parallel to plane, nothing has been hit
+	if (FMath::IsNearlyZero(FVector::DotProduct(InteractionPlanarNormal, InRay.WorldRay.Direction)))
 	{
-		FVector Delta;
-		ComputeAxisTranslateDelta(InteractionAxisCurrParam, AxisNearestParam, Delta);
-		ApplyTranslateDelta(Delta);
+		return false;
 	}
-	else
+
+	FPlane Plane(InteractionPlanarOrigin, InteractionPlanarNormal);
+	OutHitParam = FMath::RayPlaneIntersectionParam(InRay.WorldRay.Origin, InRay.WorldRay.Direction, Plane);
+	if (OutHitParam < 0)
 	{
-		FVector Delta;
-		ComputeAxisScaleDelta(InteractionAxisCurrParam, AxisNearestParam, Delta);
-		ApplyScaleDelta(Delta);
+		return false;
 	}
 
-	InteractionAxisCurrParam = AxisNearestParam;
+	return true;
 }
 
-void UTransformGizmo::OnClickReleaseAxis(const FInputDeviceRay& InReleasePos)
+void UTransformGizmo::OnClickPress(const FInputDeviceRay& PressPos)
 {
-	bInInteraction = false;
+	check(OnClickPressFunctions.Num() == static_cast<int>(ETransformGizmoPartIdentifier::Max));
+
+	if (OnClickPressFunctions[static_cast<int>(LastHitPart)])
+	{
+		OnClickPressFunctions[static_cast<int>(LastHitPart)](this, PressPos);
+	}
+
+	if (bInInteraction)
+	{
+		if (HitTarget && LastHitPart != ETransformGizmoPartIdentifier::Default)
+		{
+			HitTarget->UpdateInteractingState(true, static_cast<uint32>(LastHitPart));
+		}
+
+		if (StateTarget)
+		{
+			StateTarget->BeginUpdate();
+		}
+	}
 }
 
-void UTransformGizmo::OnClickPressPlanar(const FInputDeviceRay& PressPos)
+void UTransformGizmo::OnClickDrag(const FInputDeviceRay& DragPos)
 {
-	// if ray is parallel to plane, nothing has been hit
-	if (FMath::IsNearlyZero(FVector::DotProduct(InteractionPlanarNormal, PressPos.WorldRay.Direction)))
+	if (!bInInteraction)
 	{
 		return;
 	}
 
-	FPlane Plane(InteractionPlanarOrigin, InteractionPlanarNormal);
-	double HitDepth = FMath::RayPlaneIntersectionParam(PressPos.WorldRay.Origin, PressPos.WorldRay.Direction, Plane);
-	if (HitDepth < 0)
+	int HitPartIndex = static_cast<int>(LastHitPart);
+	check(HitPartIndex < OnClickDragFunctions.Num());
+
+	if (OnClickDragFunctions[HitPartIndex])
+	{
+		OnClickDragFunctions[HitPartIndex](this, DragPos);
+	}
+}
+
+void UTransformGizmo::OnClickRelease(const FInputDeviceRay& ReleasePos)
+{
+	if (!bInInteraction)
 	{
 		return;
 	}
 
-	InteractionPlanarStartPoint = PressPos.WorldRay.Origin + PressPos.WorldRay.Direction * HitDepth;
-	InteractionPlanarCurrPoint = InteractionPlanarStartPoint;
+	int HitPartIndex = static_cast<int>(LastHitPart);
+	check(HitPartIndex < OnClickReleaseFunctions.Num());
 
-	bInInteraction = true;
+	if (OnClickReleaseFunctions[HitPartIndex])
+	{
+		OnClickReleaseFunctions[HitPartIndex](this, ReleasePos);
+	}
+
+	if (StateTarget)
+	{
+		StateTarget->EndUpdate();
+	}
+
+	bInInteraction = false;
+
+	if (HitTarget && LastHitPart != ETransformGizmoPartIdentifier::Default)
+	{
+		HitTarget->UpdateInteractingState(false, static_cast<uint32>(LastHitPart));
+	}
 }
 
-void UTransformGizmo::OnClickDragPlanar(const FInputDeviceRay& DragPos)
+void UTransformGizmo::OnTerminateDragSequence()
 {
-	// if ray is parallel to plane, nothing has been hit
-	if (FMath::IsNearlyZero(FVector::DotProduct(InteractionPlanarNormal, DragPos.WorldRay.Direction)))
+	if (!bInInteraction)
 	{
 		return;
 	}
 
-	FPlane Plane(InteractionPlanarOrigin, InteractionPlanarNormal);
-	double HitDepth = FMath::RayPlaneIntersectionParam(DragPos.WorldRay.Origin, DragPos.WorldRay.Direction, Plane);
-	if (HitDepth < 0)
+	if (StateTarget)
 	{
-		return;
+		StateTarget->EndUpdate();
+	}
+	bInInteraction = false;
+
+	if (HitTarget && LastHitPart != ETransformGizmoPartIdentifier::Default)
+	{
+		HitTarget->UpdateInteractingState(false, static_cast<uint32>(LastHitPart));
 	}
+}
+
+void UTransformGizmo::OnClickPressTranslateXAxis(const FInputDeviceRay& PressPos)
+{
+	InteractionAxisOrigin = CurrentTransform.GetLocation();
+	InteractionAxisDirection = GetWorldAxis(FVector::XAxisVector);
+	InteractionAxisList = EAxisList::X;
+	OnClickPressAxis(PressPos);
+}
+
+void UTransformGizmo::OnClickPressTranslateYAxis(const FInputDeviceRay& PressPos)
+{
+	InteractionAxisOrigin = CurrentTransform.GetLocation();
+	InteractionAxisDirection = GetWorldAxis(FVector::YAxisVector);
+	InteractionAxisList = EAxisList::Y;
+	OnClickPressAxis(PressPos);
+}
+
+void UTransformGizmo::OnClickPressTranslateZAxis(const FInputDeviceRay& PressPos)
+{
+	InteractionAxisOrigin = CurrentTransform.GetLocation();
+	InteractionAxisDirection = GetWorldAxis(FVector::ZAxisVector);
+	InteractionAxisList = EAxisList::Z;
+	OnClickPressAxis(PressPos);
+}
 
-	FVector HitPoint = DragPos.WorldRay.Origin + DragPos.WorldRay.Direction * HitDepth;
+void UTransformGizmo::OnClickPressScaleXAxis(const FInputDeviceRay& PressPos)
+{
+	InteractionAxisOrigin = CurrentTransform.GetLocation();
+	InteractionAxisDirection = GetWorldAxis(FVector::XAxisVector);
+	InteractionAxisList = EAxisList::X;
+	OnClickPressAxis(PressPos);
+}
+
+void UTransformGizmo::OnClickPressScaleYAxis(const FInputDeviceRay& PressPos)
+{
+	InteractionAxisOrigin = CurrentTransform.GetLocation();
+	InteractionAxisDirection = GetWorldAxis(FVector::YAxisVector);
+	InteractionAxisList = EAxisList::Y;
+	OnClickPressAxis(PressPos);
+}
+
+void UTransformGizmo::OnClickPressScaleZAxis(const FInputDeviceRay& PressPos)
+{
+	InteractionAxisOrigin = CurrentTransform.GetLocation();
+	InteractionAxisDirection = GetWorldAxis(FVector::ZAxisVector);
+	InteractionAxisList = EAxisList::Z;
+	OnClickPressAxis(PressPos);
+}
 
-	if (LastHitPart == ETransformGizmoPartIdentifier::TranslateXYPlanar ||
-		LastHitPart == ETransformGizmoPartIdentifier::TranslateYZPlanar ||
-		LastHitPart == ETransformGizmoPartIdentifier::TranslateXZPlanar)
+void UTransformGizmo::OnClickPressAxis(const FInputDeviceRay& PressPos)
+{
+	InteractionAxisStartParam = GetNearestRayParamToInteractionAxis(PressPos);
+	InteractionAxisCurrParam = InteractionAxisStartParam;
+	bInInteraction = true;
+}
+
+void UTransformGizmo::OnClickDragTranslateAxis(const FInputDeviceRay& DragPos)
+{
+	float AxisNearestParam = GetNearestRayParamToInteractionAxis(DragPos);
+	FVector Delta = ComputeAxisTranslateDelta(InteractionAxisCurrParam, AxisNearestParam);
+	ApplyTranslateDelta(Delta);
+	InteractionAxisCurrParam = AxisNearestParam;
+}
+
+void UTransformGizmo::OnClickDragScaleAxis(const FInputDeviceRay& DragPos)
+{
+	float AxisNearestParam = GetNearestRayParamToInteractionAxis(DragPos);
+	FVector Delta = ComputeAxisScaleDelta(InteractionAxisCurrParam, AxisNearestParam);
+	ApplyScaleDelta(Delta);
+	InteractionAxisCurrParam = AxisNearestParam;
+}
+
+void UTransformGizmo::OnClickReleaseTranslateAxis(const FInputDeviceRay& InReleasePos)
+{
+	bInInteraction = false;
+}
+
+void UTransformGizmo::OnClickReleaseScaleAxis(const FInputDeviceRay& InReleasePos)
+{
+	bInInteraction = false;
+}
+
+void UTransformGizmo::OnClickPressTranslateXYPlanar(const FInputDeviceRay& PressPos)
+{
+	InteractionPlanarOrigin = CurrentTransform.GetLocation();
+	InteractionPlanarNormal = GetWorldAxis(FVector::ZAxisVector);
+	InteractionPlanarAxisX = GetWorldAxis(FVector::XAxisVector);
+	InteractionPlanarAxisY = GetWorldAxis(FVector::YAxisVector);
+	InteractionAxisList = EAxisList::XY;
+	OnClickPressPlanar(PressPos);
+}
+
+void UTransformGizmo::OnClickPressTranslateYZPlanar(const FInputDeviceRay& PressPos)
+{
+	InteractionPlanarOrigin = CurrentTransform.GetLocation();
+	InteractionPlanarNormal = GetWorldAxis(FVector::XAxisVector);
+	InteractionPlanarAxisX = GetWorldAxis(FVector::YAxisVector);
+	InteractionPlanarAxisY = GetWorldAxis(FVector::ZAxisVector);
+	InteractionAxisList = EAxisList::YZ;
+	OnClickPressPlanar(PressPos);
+}
+
+void UTransformGizmo::OnClickPressTranslateXZPlanar(const FInputDeviceRay& PressPos)
+{
+	InteractionPlanarOrigin = CurrentTransform.GetLocation();
+	InteractionPlanarNormal = GetWorldAxis(FVector::YAxisVector);
+	InteractionPlanarAxisX = GetWorldAxis(FVector::ZAxisVector);
+	InteractionPlanarAxisY = GetWorldAxis(FVector::XAxisVector);
+	InteractionAxisList = EAxisList::XZ;
+	OnClickPressPlanar(PressPos);
+}
+
+void UTransformGizmo::OnClickPressScaleXYPlanar(const FInputDeviceRay& PressPos)
+{
+	InteractionPlanarOrigin = CurrentTransform.GetLocation();
+	InteractionPlanarNormal = GetWorldAxis(FVector::ZAxisVector);
+	InteractionPlanarAxisX = GetWorldAxis(FVector::XAxisVector);
+	InteractionPlanarAxisY = GetWorldAxis(FVector::YAxisVector);
+	InteractionAxisList = EAxisList::XY;
+	OnClickPressPlanar(PressPos);
+}
+
+void UTransformGizmo::OnClickPressScaleYZPlanar(const FInputDeviceRay& PressPos)
+{
+	InteractionPlanarOrigin = CurrentTransform.GetLocation();
+	InteractionPlanarNormal = GetWorldAxis(FVector::XAxisVector);
+	InteractionPlanarAxisX = GetWorldAxis(FVector::YAxisVector);
+	InteractionPlanarAxisY = GetWorldAxis(FVector::ZAxisVector);
+	InteractionAxisList = EAxisList::YZ;
+	OnClickPressPlanar(PressPos);
+}
+
+void UTransformGizmo::OnClickPressScaleXZPlanar(const FInputDeviceRay& PressPos)
+{
+	InteractionPlanarOrigin = CurrentTransform.GetLocation();
+	InteractionPlanarNormal = GetWorldAxis(FVector::YAxisVector);
+	InteractionPlanarAxisX = GetWorldAxis(FVector::ZAxisVector);
+	InteractionPlanarAxisY = GetWorldAxis(FVector::XAxisVector);
+	InteractionAxisList = EAxisList::XZ;
+	OnClickPressPlanar(PressPos);
+}
+
+void UTransformGizmo::OnClickPressPlanar(const FInputDeviceRay& PressPos)
+{
+	float HitDepth;
+	if (GetRayParamIntersectionWithInteractionPlane(PressPos, HitDepth))
+	{
+		InteractionPlanarStartPoint = PressPos.WorldRay.Origin + PressPos.WorldRay.Direction * HitDepth;
+		InteractionPlanarCurrPoint = InteractionPlanarStartPoint;
+		bInInteraction = true;
+	}
+}
+
+void UTransformGizmo::OnClickDragTranslatePlanar(const FInputDeviceRay& DragPos)
+{
+	float HitDepth;
+	if (GetRayParamIntersectionWithInteractionPlane(DragPos, HitDepth))
 	{
-		FVector Delta;
-		ComputePlanarTranslateDelta(InteractionPlanarCurrPoint, HitPoint, Delta);
+		FVector HitPoint = DragPos.WorldRay.Origin + DragPos.WorldRay.Direction * HitDepth;
+		FVector Delta = ComputePlanarTranslateDelta(InteractionPlanarCurrPoint, HitPoint);
 		ApplyTranslateDelta(Delta);
+		InteractionPlanarCurrPoint = HitPoint;
 	}
-	else
+}
+
+void UTransformGizmo::OnClickDragScalePlanar(const FInputDeviceRay& DragPos)
+{
+	float HitDepth;
+	if (GetRayParamIntersectionWithInteractionPlane(DragPos, HitDepth))
 	{
-		FVector Delta;
-		ComputePlanarScaleDelta(InteractionPlanarCurrPoint, HitPoint, Delta);
+		FVector HitPoint = DragPos.WorldRay.Origin + DragPos.WorldRay.Direction * HitDepth;
+		FVector Delta = ComputePlanarScaleDelta(InteractionPlanarCurrPoint, HitPoint);
 		ApplyScaleDelta(Delta);
+		InteractionPlanarCurrPoint = HitPoint;
 	}
+}
 
-	InteractionPlanarCurrPoint = HitPoint;
+void UTransformGizmo::OnClickReleaseTranslatePlanar(const FInputDeviceRay& InReleasePos)
+{
+	bInInteraction = false;
 }
 
-void UTransformGizmo::OnClickReleasePlanar(const FInputDeviceRay& InReleasePos)
+void UTransformGizmo::OnClickReleaseScalePlanar(const FInputDeviceRay& InReleasePos)
 {
 	bInInteraction = false;
 }
 
-void UTransformGizmo::ComputeAxisTranslateDelta(double InStartParam, double InEndParam, FVector& OutTranslateDelta)
+FVector UTransformGizmo::ComputeAxisTranslateDelta(double InStartParam, double InEndParam)
 {
 	const double ParamDelta = InEndParam - InStartParam;
-	OutTranslateDelta = InteractionAxisDirection * ParamDelta;
+	return InteractionAxisDirection * ParamDelta;
 }
 
-void UTransformGizmo::ComputeAxisScaleDelta(double InStartParam, double InEndParam, FVector& OutScaleDelta)
+FVector UTransformGizmo::ComputeAxisScaleDelta(double InStartParam, double InEndParam)
 {
 	const double ParamDelta = InEndParam - InStartParam;
 	const float ScaleApplied = ParamDelta * ScaleMultiplier;
 
-	OutScaleDelta = FVector(InteractionAxisList & EAxisList::X ? ScaleApplied : 0.0,
+	return FVector(InteractionAxisList & EAxisList::X ? ScaleApplied : 0.0,
 		InteractionAxisList & EAxisList::Y ? ScaleApplied : 0.0,
 		InteractionAxisList & EAxisList::Z ? ScaleApplied : 0.0);
 }
 
-void UTransformGizmo::ComputePlanarTranslateDelta(const FVector& InStartPoint, const FVector& InEndPoint, FVector& OutTranslateDelta)
+FVector UTransformGizmo::ComputePlanarTranslateDelta(const FVector& InStartPoint, const FVector& InEndPoint)
 {
-	OutTranslateDelta = InEndPoint - InStartPoint;
+	return InEndPoint - InStartPoint;
 }
 
-void UTransformGizmo::ComputePlanarScaleDelta(const FVector& InStartPoint, const FVector& InEndPoint, FVector& OutScaleDelta)
+FVector UTransformGizmo::ComputePlanarScaleDelta(const FVector& InStartPoint, const FVector& InEndPoint)
 {
 	FVector Delta = InEndPoint - InStartPoint;
 	float DragUp = FVector::DotProduct(Delta, InteractionPlanarAxisX);
 	float DragSide = FVector::DotProduct(Delta, InteractionPlanarAxisY);
 	const float ScaleApplied = FMath::Abs(DragUp) > FMath::Abs(DragSide) ? DragUp * ScaleMultiplier : DragSide * ScaleMultiplier;
 
-	OutScaleDelta = FVector(
+	return FVector(
 		InteractionAxisList & EAxisList::X ? ScaleApplied : 0.0,
 		InteractionAxisList & EAxisList::Y ? ScaleApplied : 0.0,
 		InteractionAxisList & EAxisList::Z ? ScaleApplied : 0.0);
diff --git a/Engine/Source/Editor/Experimental/EditorInteractiveToolsFramework/Public/EditorGizmos/TransformGizmo.h b/Engine/Source/Editor/Experimental/EditorInteractiveToolsFramework/Public/EditorGizmos/TransformGizmo.h
--- a/Engine/Source/Editor/Experimental/EditorInteractiveToolsFramework/Public/EditorGizmos/TransformGizmo.h
+++ b/Engine/Source/Editor/Experimental/EditorInteractiveToolsFramework/Public/EditorGizmos/TransformGizmo.h
@@ -356,6 +356,9 @@ protected:
 	/** Setup materials */
 	virtual void SetupMaterials();
 
+	/** Setup on click functions */
+	virtual void SetupOnClickFunctions();
+
 	/** Update current gizmo mode based on transform source */
 	virtual void UpdateMode();
 
@@ -406,37 +409,103 @@ protected:
 	/** Get current interaction axis */
 	virtual FVector GetWorldAxis(const FVector& InAxis);
 
-	/** Handle click press for translate and scale axes */
-	virtual void OnClickPressAxis(const FInputDeviceRay& InPressPos);
+	/** Get nearest param along input ray to the current interaction axis */
+	virtual float GetNearestRayParamToInteractionAxis(const FInputDeviceRay& InRay);
+
+	/** Return true if input ray intersects current interaction plane and return param along ray in OutHitParam */
+	virtual bool GetRayParamIntersectionWithInteractionPlane(const FInputDeviceRay& InRay, float& OutHitParam);
+
+	/**
+	 * Translate and scale axis click-drag handling methods 
+	 */ 
+
+	/** Handle click press for translate X axis */
+	virtual void OnClickPressTranslateXAxis(const FInputDeviceRay& PressPos);
+
+	/** Handle click press for translate Y axis */
+	virtual void OnClickPressTranslateYAxis(const FInputDeviceRay& PressPos);
+
+	/** Handle click press for translate Z axis */
+	virtual void OnClickPressTranslateZAxis(const FInputDeviceRay& PressPos);
+
+	/** Handle click press for scale X axis */
+	virtual void OnClickPressScaleXAxis(const FInputDeviceRay& PressPos);
 
-	/** Handle click drag for translate and scale axes */
-	virtual void OnClickDragAxis(const FInputDeviceRay& PressPos);
+	/** Handle click press for scale Y axis */
+	virtual void OnClickPressScaleYAxis(const FInputDeviceRay& PressPos);
 
-	/** Handle click release for translate and scale axes */
-	virtual void OnClickReleaseAxis(const FInputDeviceRay& PressPos);
+	/** Handle click press for scale Z axis */
+	virtual void OnClickPressScaleZAxis(const FInputDeviceRay& PressPos);
 
-	/** Handle click press for translate and scale planar */
-	virtual void OnClickPressPlanar(const FInputDeviceRay& InPressPos);
+	/** Handle click press for translate and scale X axes */
+	virtual void OnClickPressAxis(const FInputDeviceRay& PressPos);
+
+	/** Handle click drag for translate axes */
+	virtual void OnClickDragTranslateAxis(const FInputDeviceRay& DragPos);
+
+	/** Handle click drag for scale axes */
+	virtual void OnClickDragScaleAxis(const FInputDeviceRay& DragPos);
+
+	/** Handle click release for translate axes */
+	virtual void OnClickReleaseTranslateAxis(const FInputDeviceRay& ReleasePos);
+
+	/** Handle click release for scale axes */
+	virtual void OnClickReleaseScaleAxis(const FInputDeviceRay& ReleasePos);
+
+	/**
+	 * Translate and scale planar click-drag handling methods
+	 */
 
-	/** Handle click drag for translate and scale planar */
-	virtual void OnClickDragPlanar(const FInputDeviceRay& PressPos);
+	 /** Handle click press for translate XY planar */
+	virtual void OnClickPressTranslateXYPlanar(const FInputDeviceRay& PressPos);
 
-	/** Handle click release for translate and scale planar */
-	virtual void OnClickReleasePlanar(const FInputDeviceRay& PressPos);
+	/** Handle click press for translate YZ planar */
+	virtual void OnClickPressTranslateYZPlanar(const FInputDeviceRay& PressPos);
+
+	/** Handle click press for translate XZ planar */
+	virtual void OnClickPressTranslateXZPlanar(const FInputDeviceRay& PressPos);
+
+	/** Handle click press for scale XY planar */
+	virtual void OnClickPressScaleXYPlanar(const FInputDeviceRay& PressPos);
+
+	/** Handle click press for scale YZ planar */
+	virtual void OnClickPressScaleYZPlanar(const FInputDeviceRay& PressPos);
+
+	/** Handle click press for scale XZ planar */
+	virtual void OnClickPressScaleXZPlanar(const FInputDeviceRay& PressPos);
+
+	/** Handle click press for generic planar */
+	virtual void OnClickPressPlanar(const FInputDeviceRay& PressPos);
+
+	/** Handle click drag for translate planar */
+	virtual void OnClickDragTranslatePlanar(const FInputDeviceRay& DragPos);
+
+	/** Handle click drag for scale planar */
+	virtual void OnClickDragScalePlanar(const FInputDeviceRay& DragPos);
+
+	/** Handle click release for translate planar */
+	virtual void OnClickReleaseTranslatePlanar(const FInputDeviceRay& ReleasePos);
+
+	/** Handle click release for scale planar */
+	virtual void OnClickReleaseScalePlanar(const FInputDeviceRay& ReleasePos);
 
 	/** Compute translate axis delta based on start/end params */
-	virtual void ComputeAxisTranslateDelta(double InStartParam, double InEndParam, FVector& OutDelta);
+	virtual FVector ComputeAxisTranslateDelta(double InStartParam, double InEndParam);
 
 	/** Compute scale axis delta based on start/end params */
-	virtual void ComputeAxisScaleDelta(double InStartParam, double InEndParam, FVector& OutDelta);
+	virtual FVector ComputeAxisScaleDelta(double InStartParam, double InEndParam);
 
 	/** Compute translate planar delta based on world space start/end points */
-	virtual void ComputePlanarTranslateDelta(const FVector& InStartPoint, const FVector& InEndPoint, FVector& OutDelta);
+	virtual FVector ComputePlanarTranslateDelta(const FVector& InStartPoint, const FVector& InEndPoint);
 
 	/** Compute scale planar delta based on world space start/end points */
-	virtual void ComputePlanarScaleDelta(const FVector& InStartPoint, const FVector& InEndPoint, FVector& OutDelta);
+	virtual FVector ComputePlanarScaleDelta(const FVector& InStartPoint, const FVector& InEndPoint);
 
-	/** Apply translate delta to transform proxy */
+	/**
+	 * Apply transform delta methods
+	 */
+	 
+	 /** Apply translate delta to transform proxy */
 	virtual void ApplyTranslateDelta(const FVector& InTranslateDelta);
 
 	/** Apply scale delta to transform proxy */
@@ -474,6 +543,15 @@ protected:
 	UPROPERTY()
 	TObjectPtr<UMaterialInstanceDynamic> OpaquePlaneMaterialXY;
 
+	/** Array of function pointers, indexed by gizmo part id, to handle click press behavior */
+	TArray<TFunction<void(UTransformGizmo* TransformGizmo, const FInputDeviceRay& PressPos)> > OnClickPressFunctions;
+
+	/** Array of function pointers, indexed by gizmo part id, to handle click drag behavior */
+	TArray<TFunction<void(UTransformGizmo* TransformGizmo, const FInputDeviceRay& DragPos)> > OnClickDragFunctions;
+
+	/** Array of function pointers, indexed by gizmo part id, to handle click release behavior */
+	TArray<TFunction<void(UTransformGizmo* TransformGizmo, const FInputDeviceRay& ReleasePos)> > OnClickReleaseFunctions;
+
 	/** Scale delta is multiplied by this amount */
 	UPROPERTY()
 	double ScaleMultiplier = 0.05;
