Added rotation, screen-space rotation, and screen-space translation interactions in new TRS gizmo.

diff --git a/Engine/Source/Editor/Experimental/EditorInteractiveToolsFramework/Private/EditorGizmos/EditorTransformGizmo.cpp b/Engine/Source/Editor/Experimental/EditorInteractiveToolsFramework/Private/EditorGizmos/EditorTransformGizmo.cpp
--- a/Engine/Source/Editor/Experimental/EditorInteractiveToolsFramework/Private/EditorGizmos/EditorTransformGizmo.cpp
+++ b/Engine/Source/Editor/Experimental/EditorInteractiveToolsFramework/Private/EditorGizmos/EditorTransformGizmo.cpp
@@ -24,6 +24,25 @@ void UEditorTransformGizmo::ApplyTranslateDelta(const FVector& InTranslateDelta)
 	}
 }
 
+void UEditorTransformGizmo::ApplyRotateDelta(const FQuat& InRotateDelta)
+{
+	check(ActiveTarget);
+
+	if (UEditorTransformProxy* EditorTransformProxy = Cast<UEditorTransformProxy>(ActiveTarget))
+	{
+		// Update the cached current delta.
+		// Applies rot delta after the current rotation.
+		FQuat NewRotation = InRotateDelta * CurrentTransform.GetRotation();
+		CurrentTransform.SetRotation(NewRotation);
+
+		EditorTransformProxy->InputRotateDelta(InRotateDelta.Rotator(), InteractionAxisList);
+	}
+	else
+	{
+		Super::ApplyRotateDelta(InRotateDelta);
+	}
+}
+
 void UEditorTransformGizmo::ApplyScaleDelta(const FVector& InScaleDelta)
 {
 	check(ActiveTarget);
diff --git a/Engine/Source/Editor/Experimental/EditorInteractiveToolsFramework/Private/EditorGizmos/EditorTransformGizmoBuilder.cpp b/Engine/Source/Editor/Experimental/EditorInteractiveToolsFramework/Private/EditorGizmos/EditorTransformGizmoBuilder.cpp
--- a/Engine/Source/Editor/Experimental/EditorInteractiveToolsFramework/Private/EditorGizmos/EditorTransformGizmoBuilder.cpp
+++ b/Engine/Source/Editor/Experimental/EditorInteractiveToolsFramework/Private/EditorGizmos/EditorTransformGizmoBuilder.cpp
@@ -10,12 +10,13 @@
 
 UInteractiveGizmo* UEditorTransformGizmoBuilder::BuildGizmo(const FToolBuilderState& SceneState) const
 {
+	UGizmoViewContext* GizmoViewContext = SceneState.ToolManager->GetContextObjectStore()->FindContext<UGizmoViewContext>();
+	check(GizmoViewContext && GizmoViewContext->IsValidLowLevel());
+
 	UEditorTransformGizmo* TransformGizmo = NewObject<UEditorTransformGizmo>(SceneState.GizmoManager);
 	TransformGizmo->Setup();
 	TransformGizmo->TransformGizmoSource = UEditorTransformGizmoSource::Construct(TransformGizmo);
-
-	UGizmoViewContext* GizmoViewContext = SceneState.ToolManager->GetContextObjectStore()->FindContext<UGizmoViewContext>();
-	check(GizmoViewContext && GizmoViewContext->IsValidLowLevel());
+	TransformGizmo->GizmoViewContext = GizmoViewContext;
 
 	// @todo: Gizmo element construction to be moved here from UTransformGizmo.
 	// A UGizmoElementRenderMultiTarget will be constructed and both the
diff --git a/Engine/Source/Editor/Experimental/EditorInteractiveToolsFramework/Private/EditorGizmos/TransformGizmo.cpp b/Engine/Source/Editor/Experimental/EditorInteractiveToolsFramework/Private/EditorGizmos/TransformGizmo.cpp
--- a/Engine/Source/Editor/Experimental/EditorInteractiveToolsFramework/Private/EditorGizmos/TransformGizmo.cpp
+++ b/Engine/Source/Editor/Experimental/EditorInteractiveToolsFramework/Private/EditorGizmos/TransformGizmo.cpp
@@ -720,7 +720,6 @@ UGizmoElementTorus* UTransformGizmo::MakeRotateAxis(ETransformGizmoPartIdentifie
 	RotateAxisElement->SetBeginAxis(TorusAxis0);
 	RotateAxisElement->SetPartial(true);
 	RotateAxisElement->SetAngle(PI);
-	RotateAxisElement->SetViewDependentType(EGizmoElementViewDependentType::Plane);
 	RotateAxisElement->SetViewDependentAxis(Normal);
 	RotateAxisElement->SetViewAlignType(EGizmoElementViewAlignType::Axial);
 	RotateAxisElement->SetViewAlignAxis(Normal);
@@ -845,41 +844,56 @@ void UTransformGizmo::SetupOnClickFunctions()
 	OnClickPressFunctions[static_cast<int>(ETransformGizmoPartIdentifier::TranslateXAxis)] = &UTransformGizmo::OnClickPressTranslateXAxis;
 	OnClickPressFunctions[static_cast<int>(ETransformGizmoPartIdentifier::TranslateYAxis)] = &UTransformGizmo::OnClickPressTranslateYAxis;
 	OnClickPressFunctions[static_cast<int>(ETransformGizmoPartIdentifier::TranslateZAxis)] = &UTransformGizmo::OnClickPressTranslateZAxis;
-	OnClickPressFunctions[static_cast<int>(ETransformGizmoPartIdentifier::ScaleXAxis)] = &UTransformGizmo::OnClickPressScaleXAxis;
-	OnClickPressFunctions[static_cast<int>(ETransformGizmoPartIdentifier::ScaleYAxis)] = &UTransformGizmo::OnClickPressScaleYAxis;
-	OnClickPressFunctions[static_cast<int>(ETransformGizmoPartIdentifier::ScaleZAxis)] = &UTransformGizmo::OnClickPressScaleZAxis;
 	OnClickPressFunctions[static_cast<int>(ETransformGizmoPartIdentifier::TranslateXYPlanar)] = &UTransformGizmo::OnClickPressTranslateXYPlanar;
 	OnClickPressFunctions[static_cast<int>(ETransformGizmoPartIdentifier::TranslateYZPlanar)] = &UTransformGizmo::OnClickPressTranslateYZPlanar;
 	OnClickPressFunctions[static_cast<int>(ETransformGizmoPartIdentifier::TranslateXZPlanar)] = &UTransformGizmo::OnClickPressTranslateXZPlanar;
+	OnClickPressFunctions[static_cast<int>(ETransformGizmoPartIdentifier::TranslateScreenSpace)] = &UTransformGizmo::OnClickPressScreenSpaceTranslate;
+	OnClickPressFunctions[static_cast<int>(ETransformGizmoPartIdentifier::ScaleXAxis)] = &UTransformGizmo::OnClickPressScaleXAxis;
+	OnClickPressFunctions[static_cast<int>(ETransformGizmoPartIdentifier::ScaleYAxis)] = &UTransformGizmo::OnClickPressScaleYAxis;
+	OnClickPressFunctions[static_cast<int>(ETransformGizmoPartIdentifier::ScaleZAxis)] = &UTransformGizmo::OnClickPressScaleZAxis;
 	OnClickPressFunctions[static_cast<int>(ETransformGizmoPartIdentifier::ScaleXYPlanar)] = &UTransformGizmo::OnClickPressScaleXYPlanar;
 	OnClickPressFunctions[static_cast<int>(ETransformGizmoPartIdentifier::ScaleYZPlanar)] = &UTransformGizmo::OnClickPressScaleYZPlanar;
 	OnClickPressFunctions[static_cast<int>(ETransformGizmoPartIdentifier::ScaleXZPlanar)] = &UTransformGizmo::OnClickPressScaleXZPlanar;
+	OnClickPressFunctions[static_cast<int>(ETransformGizmoPartIdentifier::RotateXAxis)] = &UTransformGizmo::OnClickPressRotateXAxis;
+	OnClickPressFunctions[static_cast<int>(ETransformGizmoPartIdentifier::RotateYAxis)] = &UTransformGizmo::OnClickPressRotateYAxis;
+	OnClickPressFunctions[static_cast<int>(ETransformGizmoPartIdentifier::RotateZAxis)] = &UTransformGizmo::OnClickPressRotateZAxis;
+	OnClickPressFunctions[static_cast<int>(ETransformGizmoPartIdentifier::RotateScreenSpace)] = &UTransformGizmo::OnClickPressScreenSpaceRotate;
 
 	OnClickDragFunctions[static_cast<int>(ETransformGizmoPartIdentifier::TranslateXAxis)] = &UTransformGizmo::OnClickDragTranslateAxis;
 	OnClickDragFunctions[static_cast<int>(ETransformGizmoPartIdentifier::TranslateYAxis)] = &UTransformGizmo::OnClickDragTranslateAxis;
 	OnClickDragFunctions[static_cast<int>(ETransformGizmoPartIdentifier::TranslateZAxis)] = &UTransformGizmo::OnClickDragTranslateAxis;
-	OnClickDragFunctions[static_cast<int>(ETransformGizmoPartIdentifier::ScaleXAxis)] = &UTransformGizmo::OnClickDragScaleAxis;
-	OnClickDragFunctions[static_cast<int>(ETransformGizmoPartIdentifier::ScaleYAxis)] = &UTransformGizmo::OnClickDragScaleAxis;
-	OnClickDragFunctions[static_cast<int>(ETransformGizmoPartIdentifier::ScaleZAxis)] = &UTransformGizmo::OnClickDragScaleAxis;
 	OnClickDragFunctions[static_cast<int>(ETransformGizmoPartIdentifier::TranslateXYPlanar)] = &UTransformGizmo::OnClickDragTranslatePlanar;
 	OnClickDragFunctions[static_cast<int>(ETransformGizmoPartIdentifier::TranslateYZPlanar)] = &UTransformGizmo::OnClickDragTranslatePlanar;
 	OnClickDragFunctions[static_cast<int>(ETransformGizmoPartIdentifier::TranslateXZPlanar)] = &UTransformGizmo::OnClickDragTranslatePlanar;
+	OnClickDragFunctions[static_cast<int>(ETransformGizmoPartIdentifier::TranslateScreenSpace)] = &UTransformGizmo::OnClickDragScreenSpaceTranslate;
+	OnClickDragFunctions[static_cast<int>(ETransformGizmoPartIdentifier::ScaleXAxis)] = &UTransformGizmo::OnClickDragScaleAxis;
+	OnClickDragFunctions[static_cast<int>(ETransformGizmoPartIdentifier::ScaleYAxis)] = &UTransformGizmo::OnClickDragScaleAxis;
+	OnClickDragFunctions[static_cast<int>(ETransformGizmoPartIdentifier::ScaleZAxis)] = &UTransformGizmo::OnClickDragScaleAxis;
 	OnClickDragFunctions[static_cast<int>(ETransformGizmoPartIdentifier::ScaleXYPlanar)] = &UTransformGizmo::OnClickDragScalePlanar;
 	OnClickDragFunctions[static_cast<int>(ETransformGizmoPartIdentifier::ScaleYZPlanar)] = &UTransformGizmo::OnClickDragScalePlanar;
 	OnClickDragFunctions[static_cast<int>(ETransformGizmoPartIdentifier::ScaleXZPlanar)] = &UTransformGizmo::OnClickDragScalePlanar;
+	OnClickDragFunctions[static_cast<int>(ETransformGizmoPartIdentifier::RotateXAxis)] = &UTransformGizmo::OnClickDragRotateAxis;
+	OnClickDragFunctions[static_cast<int>(ETransformGizmoPartIdentifier::RotateYAxis)] = &UTransformGizmo::OnClickDragRotateAxis;
+	OnClickDragFunctions[static_cast<int>(ETransformGizmoPartIdentifier::RotateZAxis)] = &UTransformGizmo::OnClickDragRotateAxis;
+	OnClickDragFunctions[static_cast<int>(ETransformGizmoPartIdentifier::RotateScreenSpace)] = &UTransformGizmo::OnClickDragScreenSpaceRotate;
 
 	OnClickReleaseFunctions[static_cast<int>(ETransformGizmoPartIdentifier::TranslateXAxis)] = &UTransformGizmo::OnClickReleaseTranslateAxis;
 	OnClickReleaseFunctions[static_cast<int>(ETransformGizmoPartIdentifier::TranslateYAxis)] = &UTransformGizmo::OnClickReleaseTranslateAxis;
 	OnClickReleaseFunctions[static_cast<int>(ETransformGizmoPartIdentifier::TranslateZAxis)] = &UTransformGizmo::OnClickReleaseTranslateAxis;
-	OnClickReleaseFunctions[static_cast<int>(ETransformGizmoPartIdentifier::ScaleXAxis)] = &UTransformGizmo::OnClickReleaseScaleAxis;
-	OnClickReleaseFunctions[static_cast<int>(ETransformGizmoPartIdentifier::ScaleYAxis)] = &UTransformGizmo::OnClickReleaseScaleAxis;
-	OnClickReleaseFunctions[static_cast<int>(ETransformGizmoPartIdentifier::ScaleZAxis)] = &UTransformGizmo::OnClickReleaseScaleAxis;
 	OnClickReleaseFunctions[static_cast<int>(ETransformGizmoPartIdentifier::TranslateXYPlanar)] = &UTransformGizmo::OnClickReleaseTranslatePlanar;
 	OnClickReleaseFunctions[static_cast<int>(ETransformGizmoPartIdentifier::TranslateYZPlanar)] = &UTransformGizmo::OnClickReleaseTranslatePlanar;
 	OnClickReleaseFunctions[static_cast<int>(ETransformGizmoPartIdentifier::TranslateXZPlanar)] = &UTransformGizmo::OnClickReleaseTranslatePlanar;
+	OnClickReleaseFunctions[static_cast<int>(ETransformGizmoPartIdentifier::TranslateScreenSpace)] = &UTransformGizmo::OnClickReleaseScreenSpaceTranslate;
+	OnClickReleaseFunctions[static_cast<int>(ETransformGizmoPartIdentifier::ScaleXAxis)] = &UTransformGizmo::OnClickReleaseScaleAxis;
+	OnClickReleaseFunctions[static_cast<int>(ETransformGizmoPartIdentifier::ScaleYAxis)] = &UTransformGizmo::OnClickReleaseScaleAxis;
+	OnClickReleaseFunctions[static_cast<int>(ETransformGizmoPartIdentifier::ScaleZAxis)] = &UTransformGizmo::OnClickReleaseScaleAxis;
 	OnClickReleaseFunctions[static_cast<int>(ETransformGizmoPartIdentifier::ScaleXYPlanar)] = &UTransformGizmo::OnClickReleaseScalePlanar;
 	OnClickReleaseFunctions[static_cast<int>(ETransformGizmoPartIdentifier::ScaleYZPlanar)] = &UTransformGizmo::OnClickReleaseScalePlanar;
 	OnClickReleaseFunctions[static_cast<int>(ETransformGizmoPartIdentifier::ScaleXZPlanar)] = &UTransformGizmo::OnClickReleaseScalePlanar;
+	OnClickReleaseFunctions[static_cast<int>(ETransformGizmoPartIdentifier::RotateXAxis)] = &UTransformGizmo::OnClickReleaseRotateAxis;
+	OnClickReleaseFunctions[static_cast<int>(ETransformGizmoPartIdentifier::RotateYAxis)] = &UTransformGizmo::OnClickReleaseRotateAxis;
+	OnClickReleaseFunctions[static_cast<int>(ETransformGizmoPartIdentifier::RotateZAxis)] = &UTransformGizmo::OnClickReleaseRotateAxis;
+	OnClickReleaseFunctions[static_cast<int>(ETransformGizmoPartIdentifier::RotateScreenSpace)] = &UTransformGizmo::OnClickReleaseScreenSpaceRotate;
 }
 
 
@@ -1217,12 +1231,188 @@ FVector UTransformGizmo::ComputePlanarScaleDelta(const FVector& InStartPoint, co
 		InteractionAxisList & EAxisList::Z ? ScaleApplied : 0.0);
 }
 
+void UTransformGizmo::OnClickPressScreenSpaceTranslate(const FInputDeviceRay& PressPos)
+{
+	check(GizmoViewContext);
+
+	InteractionPlanarOrigin = CurrentTransform.GetLocation();
+	InteractionPlanarNormal = -GizmoViewContext->GetViewDirection();
+	InteractionPlanarAxisX = GizmoViewContext->GetViewUp();
+	InteractionPlanarAxisY = GizmoViewContext->GetViewRight();
+	InteractionAxisList = EAxisList::Screen;
+	OnClickPressPlanar(PressPos);
+}
+
+void UTransformGizmo::OnClickDragScreenSpaceTranslate(const FInputDeviceRay& DragPos)
+{
+	OnClickDragTranslatePlanar(DragPos);
+}
+
+void UTransformGizmo::OnClickReleaseScreenSpaceTranslate(const FInputDeviceRay& InReleasePos)
+{
+	bInInteraction = false;
+}
+
+void UTransformGizmo::OnClickPressRotateXAxis(const FInputDeviceRay& PressPos)
+{
+	InteractionScreenAxisDirection = GetScreenRotateAxisDir(FVector::YAxisVector, FVector::ZAxisVector).GetSafeNormal();
+	InteractionAxisList = EAxisList::X;
+	InteractionScreenStartPos = InteractionScreenCurrPos = PressPos.ScreenPosition;
+	bInInteraction = true;
+}
+
+void UTransformGizmo::OnClickPressRotateYAxis(const FInputDeviceRay& PressPos)
+{
+	InteractionScreenAxisDirection = GetScreenRotateAxisDir(FVector::ZAxisVector, FVector::XAxisVector).GetSafeNormal();
+	InteractionAxisList = EAxisList::Y;
+	InteractionScreenStartPos = InteractionScreenCurrPos = PressPos.ScreenPosition;
+	bInInteraction = true;
+}
+
+void UTransformGizmo::OnClickPressRotateZAxis(const FInputDeviceRay& PressPos)
+{
+	InteractionScreenAxisDirection = GetScreenRotateAxisDir(FVector::XAxisVector, FVector::YAxisVector).GetSafeNormal();
+	InteractionAxisList = EAxisList::Z;
+	InteractionScreenStartPos = InteractionScreenCurrPos = PressPos.ScreenPosition;
+	bInInteraction = true;
+}
+
+FVector2D UTransformGizmo::GetScreenRotateAxisDir(const FVector& InAxis0, const FVector& InAxis1)
+{
+	check(GizmoViewContext);
+	const FVector DirectionToWidget = CurrentTransform.GetLocation() - GizmoViewContext->ViewLocation;
+
+	const FVector Axis0 = GetWorldAxis(InAxis0);
+	const FVector Axis1 = GetWorldAxis(InAxis1);
+
+	// Reverse the axes based on camera view
+	const bool bMirrorAxis0 = (FVector::DotProduct(Axis0, DirectionToWidget) <= 0.0f);
+	const bool bMirrorAxis1 = (FVector::DotProduct(Axis1, DirectionToWidget) <= 0.0f);
+	const float Direction = (bMirrorAxis0 ^ bMirrorAxis1) ? -1.0f : 1.0f;
+
+	const FVector AxisDir = (Axis1 - Axis0) * Direction;
+
+	return GetScreenProjectedAxis(AxisDir);
+}
+
+void UTransformGizmo::OnClickDragRotateAxis(const FInputDeviceRay& DragPos)
+{
+	FQuat DeltaRot = ComputeAxisRotateDelta(InteractionScreenCurrPos, DragPos.ScreenPosition);
+	ApplyRotateDelta(DeltaRot);
+	InteractionScreenCurrPos = DragPos.ScreenPosition;
+}
+
+FQuat UTransformGizmo::ComputeAxisRotateDelta(const FVector2D& InStartPos, const FVector2D& InEndPos)
+{
+	FVector2D DragDir = InEndPos - InStartPos;
+	FRotator DeltaRot(0.0, 0.0, 0.0);
+	if (InteractionAxisList == EAxisList::X)
+	{
+		DeltaRot.Roll = FVector2D::DotProduct(InteractionScreenAxisDirection, DragDir);
+	}
+	else if (InteractionAxisList == EAxisList::Y)
+	{
+		DeltaRot.Pitch = FVector2D::DotProduct(InteractionScreenAxisDirection, DragDir);
+	}
+	else
+	{
+		DeltaRot.Yaw = FVector2D::DotProduct(InteractionScreenAxisDirection, DragDir);
+	}
+
+	check(TransformGizmoSource);
+	if (TransformGizmoSource->GetGizmoCoordSystemSpace() == EToolContextCoordinateSystem::Local)
+	{
+		check(ActiveTarget);
+		FMatrix CurrCoordSystem = ActiveTarget->GetTransform().ToMatrixNoScale();
+		DeltaRot = (CurrCoordSystem.Inverse() * FRotationMatrix(DeltaRot) * CurrCoordSystem).Rotator();
+	}
+
+	return DeltaRot.Quaternion();
+}
+
+void UTransformGizmo::OnClickReleaseRotateAxis(const FInputDeviceRay& InReleasePos)
+{
+	bInInteraction = false;
+}
+
+void UTransformGizmo::OnClickPressScreenSpaceRotate(const FInputDeviceRay& PressPos)
+{
+	check(GizmoViewContext);
+	
+	InteractionPlanarOrigin = CurrentTransform.GetLocation();
+	InteractionPlanarNormal = -GizmoViewContext->GetViewDirection();
+	InteractionPlanarAxisX = GizmoViewContext->GetViewUp();
+	InteractionPlanarAxisY = GizmoViewContext->GetViewRight();
+	InteractionAxisList = EAxisList::Screen;
+
+	float HitDepth;
+	if (GetRayParamIntersectionWithInteractionPlane(PressPos, HitDepth))
+	{
+		FVector HitPoint = PressPos.WorldRay.Origin + PressPos.WorldRay.Direction * HitDepth;
+		InteractionStartAngle = GizmoMath::ComputeAngleInPlane(HitPoint,
+			InteractionPlanarOrigin, InteractionPlanarNormal, InteractionPlanarAxisX, InteractionPlanarAxisY);
+		InteractionCurrAngle = InteractionStartAngle;
+
+		bInInteraction = true;
+	}
+}
+
+void UTransformGizmo::OnClickDragScreenSpaceRotate(const FInputDeviceRay& DragPos)
+{
+	check(GizmoViewContext);
+	
+	float HitDepth;
+	if (GetRayParamIntersectionWithInteractionPlane(DragPos, HitDepth))
+	{
+		FVector HitPoint = DragPos.WorldRay.Origin + DragPos.WorldRay.Direction * HitDepth;
+		float HitAngle = GizmoMath::ComputeAngleInPlane(HitPoint,
+			InteractionPlanarOrigin, InteractionPlanarNormal, InteractionPlanarAxisX, InteractionPlanarAxisY);
+
+		FQuat Delta = ComputeAngularRotateDelta(InteractionCurrAngle, HitAngle);
+		ApplyRotateDelta(Delta);
+		InteractionCurrAngle = HitAngle;
+	}
+}
+
+FQuat UTransformGizmo::ComputeAngularRotateDelta(double InStartAngle, double InEndAngle)
+{
+	float DeltaAngle = InEndAngle - InStartAngle;
+	return FQuat(InteractionPlanarNormal, DeltaAngle);
+}
+
+void UTransformGizmo::OnClickReleaseScreenSpaceRotate(const FInputDeviceRay& InReleasePos)
+{
+	bInInteraction = false;
+}
+
+FVector2D UTransformGizmo::GetScreenProjectedAxis(const FVector& InAxis)
+{
+	check(GizmoViewContext);
+
+	if (FMath::IsNearlyEqual(FVector::DotProduct(GizmoViewContext->GetViewDirection(), InAxis), 1.0))
+	{
+		return FVector2D(1.0, 1.0);
+	}
+	
+	const float ScreenX = FVector::DotProduct(InAxis, GizmoViewContext->GetViewRight());
+	const float ScreenY = FVector::DotProduct(InAxis, GizmoViewContext->GetViewUp());
+	return FVector2D(ScreenX, ScreenY);
+}
+
 void UTransformGizmo::ApplyTranslateDelta(const FVector& InTranslateDelta)
 {
 	CurrentTransform.AddToTranslation(InTranslateDelta);
 	ActiveTarget->SetTransform(CurrentTransform);
 }
 
+void UTransformGizmo::ApplyRotateDelta(const FQuat& InRotateDelta)
+{
+	// Applies rot delta after the current rotation
+	FQuat NewRotation = InRotateDelta * CurrentTransform.GetRotation();
+	CurrentTransform.SetRotation(NewRotation);
+	ActiveTarget->SetTransform(CurrentTransform);
+}
+
 void UTransformGizmo::ApplyScaleDelta(const FVector& InScaleDelta)
 {
 	FVector StartScale = CurrentTransform.GetScale3D();
diff --git a/Engine/Source/Editor/Experimental/EditorInteractiveToolsFramework/Public/EditorGizmos/EditorTransformGizmo.h b/Engine/Source/Editor/Experimental/EditorInteractiveToolsFramework/Public/EditorGizmos/EditorTransformGizmo.h
--- a/Engine/Source/Editor/Experimental/EditorInteractiveToolsFramework/Public/EditorGizmos/EditorTransformGizmo.h
+++ b/Engine/Source/Editor/Experimental/EditorInteractiveToolsFramework/Public/EditorGizmos/EditorTransformGizmo.h
@@ -20,6 +20,9 @@ protected:
 	/** Apply translate delta to transform proxy */
 	virtual void ApplyTranslateDelta(const FVector& InTranslateDelta);
 
+	/** Apply rotate delta to transform proxy */
+	virtual void ApplyRotateDelta(const FQuat& InRotateDelta);
+
 	/** Apply scale delta to transform proxy */
 	virtual void ApplyScaleDelta(const FVector& InScaleDelta);
 };
\ No newline at end of file
diff --git a/Engine/Source/Editor/Experimental/EditorInteractiveToolsFramework/Public/EditorGizmos/TransformGizmo.h b/Engine/Source/Editor/Experimental/EditorInteractiveToolsFramework/Public/EditorGizmos/TransformGizmo.h
--- a/Engine/Source/Editor/Experimental/EditorInteractiveToolsFramework/Public/EditorGizmos/TransformGizmo.h
+++ b/Engine/Source/Editor/Experimental/EditorInteractiveToolsFramework/Public/EditorGizmos/TransformGizmo.h
@@ -6,6 +6,7 @@
 #include "BaseBehaviors/BehaviorTargetInterfaces.h"
 #include "BaseGizmos/GizmoElementHitTargets.h"
 #include "BaseGizmos/GizmoElementStateTargets.h"
+#include "BaseGizmos/GizmoViewContext.h"
 #include "BaseGizmos/TransformProxy.h"
 #include "EditorGizmos/TransformGizmoInterfaces.h"
 #include "InteractiveGizmo.h"
@@ -204,6 +205,10 @@ public:
 	UPROPERTY()
 	TObjectPtr<UGizmoElementGroup> GizmoElementRoot;
 
+	/** Gizmo view context, needed for screen space interactions */
+	UPROPERTY()
+	TObjectPtr<UGizmoViewContext> GizmoViewContext;
+
 	/** Whether gizmo is visible. */
 	UPROPERTY()
 	bool bVisible = false;
@@ -501,13 +506,78 @@ protected:
 	/** Compute scale planar delta based on world space start/end points */
 	virtual FVector ComputePlanarScaleDelta(const FVector& InStartPoint, const FVector& InEndPoint);
 
+	/**
+	 * Screen-space translate interaction methods
+	 */
+
+	 /** Handle click press for screen-space translate */
+	virtual void OnClickPressScreenSpaceTranslate(const FInputDeviceRay& PressPos);
+
+	/** Handle click drag for screen-space translate */
+	virtual void OnClickDragScreenSpaceTranslate(const FInputDeviceRay& DragPos);
+
+	/** Handle click release for screen-space translate */
+	virtual void OnClickReleaseScreenSpaceTranslate(const FInputDeviceRay& ReleasePos);
+
+	/**
+	 * Rotate interaction methods
+	 */
+
+	/** Handle click press for rotate X axis */
+	virtual void OnClickPressRotateXAxis(const FInputDeviceRay& PressPos);
+
+	/** Handle click press for rotate Y axis */
+	virtual void OnClickPressRotateYAxis(const FInputDeviceRay& PressPos);
+
+	/** Handle click press for rotate Z axis */
+	virtual void OnClickPressRotateZAxis(const FInputDeviceRay& PressPos);
+
+	/** Handle click drag for rotate axis */
+	virtual void OnClickDragRotateAxis(const FInputDeviceRay& DragPos);
+
+	/** Handle click release for rotate axes */
+	virtual void OnClickReleaseRotateAxis(const FInputDeviceRay& ReleasePos);
+
+	/** Get screen-space axis for rotation drag */
+	FVector2D GetScreenRotateAxisDir(const FVector& InAxis0, const FVector& InAxis1);
+
+	/** Compute rotate delta based on screen-space start/end positions */
+	virtual FQuat ComputeAxisRotateDelta(const FVector2D& InStartPos, const FVector2D& InEndPos);
+
+	/**
+	 * Screen-space rotate interaction methods
+	 */
+
+	 /** Handle click press for screen-space rotate */
+	virtual void OnClickPressScreenSpaceRotate(const FInputDeviceRay& PressPos);
+
+	/** Handle click drag for screen-space rotate */
+	virtual void OnClickDragScreenSpaceRotate(const FInputDeviceRay& DragPos);
+
+	/** Handle click release for screen-space rotate */
+	virtual void OnClickReleaseScreenSpaceRotate(const FInputDeviceRay& ReleasePos);
+
+	/** Compute rotate delta based on start/end angles */
+	virtual FQuat ComputeAngularRotateDelta(double InStartAngle, double InEndAngle);
+
+	/**
+ 	 * Screen-space helper methods
+	 */
+
+	/** Returns 2D vector projection of input axis onto the current viewing plane */
+	virtual FVector2D GetScreenProjectedAxis(const FVector& InAxis);
+
+
 	/**
 	 * Apply transform delta methods
 	 */
 	 
-	 /** Apply translate delta to transform proxy */
+	/** Apply translate delta to transform proxy */
 	virtual void ApplyTranslateDelta(const FVector& InTranslateDelta);
 
+	/** Apply rotate delta to transform proxy */
+	virtual void ApplyRotateDelta(const FQuat& InRotateDelta);
+
 	/** Apply scale delta to transform proxy */
 	virtual void ApplyScaleDelta(const FVector& InScaleDelta);
 
@@ -620,4 +690,25 @@ protected:
 	/** Active interaction current point planar (only valid between state target BeginModify/EndModify) */
 	UPROPERTY()
 	FVector InteractionPlanarCurrPoint;
+
+	/** Active interaction rotation start angle (only valid between state target BeginModify/EndModify) */
+	UPROPERTY()
+	float InteractionStartAngle;
+
+	/** Active interaction rotation curr angle (only valid between state target BeginModify/EndModify) */
+	UPROPERTY()
+	float InteractionCurrAngle;
+
+	/** Active interaction screen axis dir (only valid between state target BeginModify/EndModify) */
+	UPROPERTY()
+	FVector2D InteractionScreenAxisDirection;
+
+	/** Active interaction screen start pos planar (only valid between state target BeginModify/EndModify) */
+	UPROPERTY()
+	FVector2D InteractionScreenStartPos;
+
+	/** Active interaction screen current pos planar (only valid between state target BeginModify/EndModify) */
+	UPROPERTY()
+	FVector2D InteractionScreenCurrPos;
+
 };
