Fix new TRS gizmo scale crash when Preserve Non-Uniform scale is enabled.

diff --git a/Engine/Source/Editor/Experimental/EditorInteractiveToolsFramework/Private/EditorGizmos/EditorTransformProxy.cpp b/Engine/Source/Editor/Experimental/EditorInteractiveToolsFramework/Private/EditorGizmos/EditorTransformProxy.cpp
--- a/Engine/Source/Editor/Experimental/EditorInteractiveToolsFramework/Private/EditorGizmos/EditorTransformProxy.cpp
+++ b/Engine/Source/Editor/Experimental/EditorInteractiveToolsFramework/Private/EditorGizmos/EditorTransformProxy.cpp
@@ -29,7 +29,10 @@ void UEditorTransformProxy::InputTranslateDelta(const FVector& InDeltaTranslate,
 		FRotator Rot = FRotator::ZeroRotator;
 		FVector Scale = FVector::ZeroVector;
 
+		// Set legacy widget axis temporarily because InputWidgetDelta branches/overrides may expect it
+		ViewportClient->SetCurrentWidgetAxis(InAxisList);
 		ViewportClient->InputWidgetDelta(ViewportClient->Viewport, InAxisList, Translate, Rot, Scale);
+		ViewportClient->SetCurrentWidgetAxis(EAxisList::None);
 	}
 }
 
@@ -41,7 +44,10 @@ void UEditorTransformProxy::InputScaleDelta(const FVector& InDeltaScale, EAxisLi
 		FRotator Rot = FRotator::ZeroRotator;
 		FVector Scale = InDeltaScale;
 
+		// Set legacy widget axis temporarily because InputWidgetDelta validates the axis in some crashes and crashes if it is not set
+		ViewportClient->SetCurrentWidgetAxis(InAxisList);
 		ViewportClient->InputWidgetDelta(ViewportClient->Viewport, InAxisList, Translate, Rot, Scale);
+		ViewportClient->SetCurrentWidgetAxis(EAxisList::None);
 	}
 }
 
@@ -53,7 +59,10 @@ void UEditorTransformProxy::InputRotateDelta(const FRotator& InDeltaRotate, EAxi
 		FRotator Rot = InDeltaRotate;
 		FVector Scale = FVector::ZeroVector;
 
+		// Set legacy widget axis temporarily because InputWidgetDelta branches/overrides may expect it
+		ViewportClient->SetCurrentWidgetAxis(InAxisList);
 		ViewportClient->InputWidgetDelta(ViewportClient->Viewport, InAxisList, Translate, Rot, Scale);
+		ViewportClient->SetCurrentWidgetAxis(EAxisList::None);
 	}
 }
 
