New TRS gizmo: highlight related axes when planar scale or uniform scale are hovered or interacting.

diff --git a/Engine/Source/Editor/Experimental/EditorInteractiveToolsFramework/Private/EditorGizmos/TransformGizmo.cpp b/Engine/Source/Editor/Experimental/EditorInteractiveToolsFramework/Private/EditorGizmos/TransformGizmo.cpp
--- a/Engine/Source/Editor/Experimental/EditorInteractiveToolsFramework/Private/EditorGizmos/TransformGizmo.cpp
+++ b/Engine/Source/Editor/Experimental/EditorInteractiveToolsFramework/Private/EditorGizmos/TransformGizmo.cpp
@@ -180,12 +180,12 @@ FInputRayHit UTransformGizmo::UpdateHoveredPart(const FInputDeviceRay& PressPos)
 	{
 		if (LastHitPart != ETransformGizmoPartIdentifier::Default)
 		{
-			HitTarget->UpdateHoverState(false, static_cast<uint32>(LastHitPart));
+			UpdateHoverState(false, LastHitPart);
 		}
 
 		if (HitPart != ETransformGizmoPartIdentifier::Default)
 		{
-			HitTarget->UpdateHoverState(true, static_cast<uint32>(HitPart));
+			UpdateHoverState(true, HitPart);
 		}
 
 		LastHitPart = HitPart;
@@ -854,6 +854,7 @@ void UTransformGizmo::SetupOnClickFunctions()
 	OnClickPressFunctions[static_cast<int>(ETransformGizmoPartIdentifier::ScaleXYPlanar)] = &UTransformGizmo::OnClickPressScaleXYPlanar;
 	OnClickPressFunctions[static_cast<int>(ETransformGizmoPartIdentifier::ScaleYZPlanar)] = &UTransformGizmo::OnClickPressScaleYZPlanar;
 	OnClickPressFunctions[static_cast<int>(ETransformGizmoPartIdentifier::ScaleXZPlanar)] = &UTransformGizmo::OnClickPressScaleXZPlanar;
+	OnClickPressFunctions[static_cast<int>(ETransformGizmoPartIdentifier::ScaleUniform)] = &UTransformGizmo::OnClickPressScaleXYZ;
 	OnClickPressFunctions[static_cast<int>(ETransformGizmoPartIdentifier::RotateXAxis)] = &UTransformGizmo::OnClickPressRotateXAxis;
 	OnClickPressFunctions[static_cast<int>(ETransformGizmoPartIdentifier::RotateYAxis)] = &UTransformGizmo::OnClickPressRotateYAxis;
 	OnClickPressFunctions[static_cast<int>(ETransformGizmoPartIdentifier::RotateZAxis)] = &UTransformGizmo::OnClickPressRotateZAxis;
@@ -872,6 +873,7 @@ void UTransformGizmo::SetupOnClickFunctions()
 	OnClickDragFunctions[static_cast<int>(ETransformGizmoPartIdentifier::ScaleXYPlanar)] = &UTransformGizmo::OnClickDragScalePlanar;
 	OnClickDragFunctions[static_cast<int>(ETransformGizmoPartIdentifier::ScaleYZPlanar)] = &UTransformGizmo::OnClickDragScalePlanar;
 	OnClickDragFunctions[static_cast<int>(ETransformGizmoPartIdentifier::ScaleXZPlanar)] = &UTransformGizmo::OnClickDragScalePlanar;
+	OnClickDragFunctions[static_cast<int>(ETransformGizmoPartIdentifier::ScaleUniform)] = &UTransformGizmo::OnClickDragScaleXYZ;
 	OnClickDragFunctions[static_cast<int>(ETransformGizmoPartIdentifier::RotateXAxis)] = &UTransformGizmo::OnClickDragRotateAxis;
 	OnClickDragFunctions[static_cast<int>(ETransformGizmoPartIdentifier::RotateYAxis)] = &UTransformGizmo::OnClickDragRotateAxis;
 	OnClickDragFunctions[static_cast<int>(ETransformGizmoPartIdentifier::RotateZAxis)] = &UTransformGizmo::OnClickDragRotateAxis;
@@ -890,6 +892,7 @@ void UTransformGizmo::SetupOnClickFunctions()
 	OnClickReleaseFunctions[static_cast<int>(ETransformGizmoPartIdentifier::ScaleXYPlanar)] = &UTransformGizmo::OnClickReleaseScalePlanar;
 	OnClickReleaseFunctions[static_cast<int>(ETransformGizmoPartIdentifier::ScaleYZPlanar)] = &UTransformGizmo::OnClickReleaseScalePlanar;
 	OnClickReleaseFunctions[static_cast<int>(ETransformGizmoPartIdentifier::ScaleXZPlanar)] = &UTransformGizmo::OnClickReleaseScalePlanar;
+	OnClickReleaseFunctions[static_cast<int>(ETransformGizmoPartIdentifier::ScaleUniform)] = &UTransformGizmo::OnClickReleaseScaleXYZ;
 	OnClickReleaseFunctions[static_cast<int>(ETransformGizmoPartIdentifier::RotateXAxis)] = &UTransformGizmo::OnClickReleaseRotateAxis;
 	OnClickReleaseFunctions[static_cast<int>(ETransformGizmoPartIdentifier::RotateYAxis)] = &UTransformGizmo::OnClickReleaseRotateAxis;
 	OnClickReleaseFunctions[static_cast<int>(ETransformGizmoPartIdentifier::RotateZAxis)] = &UTransformGizmo::OnClickReleaseRotateAxis;
@@ -925,6 +928,65 @@ bool UTransformGizmo::GetRayParamIntersectionWithInteractionPlane(const FInputDe
 	return true;
 }
 
+void UTransformGizmo::UpdateHoverState(bool bInHover, ETransformGizmoPartIdentifier InHitPartId)
+{
+	HitTarget->UpdateHoverState(bInHover, static_cast<uint32>(InHitPartId));
+
+	switch (InHitPartId)
+	{
+	case ETransformGizmoPartIdentifier::ScaleUniform:
+		HitTarget->UpdateHoverState(bInHover, static_cast<uint32>(ETransformGizmoPartIdentifier::ScaleXAxis));
+		HitTarget->UpdateHoverState(bInHover, static_cast<uint32>(ETransformGizmoPartIdentifier::ScaleYAxis));
+		HitTarget->UpdateHoverState(bInHover, static_cast<uint32>(ETransformGizmoPartIdentifier::ScaleZAxis));
+		HitTarget->UpdateHoverState(bInHover, static_cast<uint32>(ETransformGizmoPartIdentifier::ScaleXYPlanar));
+		HitTarget->UpdateHoverState(bInHover, static_cast<uint32>(ETransformGizmoPartIdentifier::ScaleYZPlanar));
+		HitTarget->UpdateHoverState(bInHover, static_cast<uint32>(ETransformGizmoPartIdentifier::ScaleXZPlanar));
+		break;
+	case ETransformGizmoPartIdentifier::ScaleXYPlanar:
+		HitTarget->UpdateHoverState(bInHover, static_cast<uint32>(ETransformGizmoPartIdentifier::ScaleXAxis));
+		HitTarget->UpdateHoverState(bInHover, static_cast<uint32>(ETransformGizmoPartIdentifier::ScaleYAxis));
+		break;
+	case ETransformGizmoPartIdentifier::ScaleYZPlanar:
+		HitTarget->UpdateHoverState(bInHover, static_cast<uint32>(ETransformGizmoPartIdentifier::ScaleYAxis));
+		HitTarget->UpdateHoverState(bInHover, static_cast<uint32>(ETransformGizmoPartIdentifier::ScaleZAxis));
+		break;
+	case ETransformGizmoPartIdentifier::ScaleXZPlanar:
+		HitTarget->UpdateHoverState(bInHover, static_cast<uint32>(ETransformGizmoPartIdentifier::ScaleXAxis));
+		HitTarget->UpdateHoverState(bInHover, static_cast<uint32>(ETransformGizmoPartIdentifier::ScaleZAxis));
+		break;
+	}
+}
+
+void UTransformGizmo::UpdateInteractingState(bool bInInteracting, ETransformGizmoPartIdentifier InHitPartId)
+{
+	HitTarget->UpdateInteractingState(bInInteracting, static_cast<uint32>(InHitPartId));
+
+	switch (InHitPartId)
+	{
+	case ETransformGizmoPartIdentifier::ScaleUniform:
+		HitTarget->UpdateInteractingState(bInInteracting, static_cast<uint32>(ETransformGizmoPartIdentifier::ScaleXAxis));
+		HitTarget->UpdateInteractingState(bInInteracting, static_cast<uint32>(ETransformGizmoPartIdentifier::ScaleYAxis));
+		HitTarget->UpdateInteractingState(bInInteracting, static_cast<uint32>(ETransformGizmoPartIdentifier::ScaleZAxis));
+		HitTarget->UpdateInteractingState(bInInteracting, static_cast<uint32>(ETransformGizmoPartIdentifier::ScaleXYPlanar));
+		HitTarget->UpdateInteractingState(bInInteracting, static_cast<uint32>(ETransformGizmoPartIdentifier::ScaleYZPlanar));
+		HitTarget->UpdateInteractingState(bInInteracting, static_cast<uint32>(ETransformGizmoPartIdentifier::ScaleXZPlanar));
+		HitTarget->UpdateInteractingState(bInInteracting, static_cast<uint32>(ETransformGizmoPartIdentifier::ScaleUniform));
+		break;
+	case ETransformGizmoPartIdentifier::ScaleXYPlanar:
+		HitTarget->UpdateInteractingState(bInInteracting, static_cast<uint32>(ETransformGizmoPartIdentifier::ScaleXAxis));
+		HitTarget->UpdateInteractingState(bInInteracting, static_cast<uint32>(ETransformGizmoPartIdentifier::ScaleYAxis));
+		break;
+	case ETransformGizmoPartIdentifier::ScaleYZPlanar:
+		HitTarget->UpdateInteractingState(bInInteracting, static_cast<uint32>(ETransformGizmoPartIdentifier::ScaleYAxis));
+		HitTarget->UpdateInteractingState(bInInteracting, static_cast<uint32>(ETransformGizmoPartIdentifier::ScaleZAxis));
+		break;
+	case ETransformGizmoPartIdentifier::ScaleXZPlanar:
+		HitTarget->UpdateInteractingState(bInInteracting, static_cast<uint32>(ETransformGizmoPartIdentifier::ScaleXAxis));
+		HitTarget->UpdateInteractingState(bInInteracting, static_cast<uint32>(ETransformGizmoPartIdentifier::ScaleZAxis));
+		break;
+	}
+}
+
 void UTransformGizmo::OnClickPress(const FInputDeviceRay& PressPos)
 {
 	check(OnClickPressFunctions.Num() == static_cast<int>(ETransformGizmoPartIdentifier::Max));
@@ -938,7 +1000,7 @@ void UTransformGizmo::OnClickPress(const FInputDeviceRay& PressPos)
 	{
 		if (HitTarget && LastHitPart != ETransformGizmoPartIdentifier::Default)
 		{
-			HitTarget->UpdateInteractingState(true, static_cast<uint32>(LastHitPart));
+			UpdateInteractingState(true, LastHitPart);
 		}
 
 		if (StateTarget)
@@ -988,7 +1050,7 @@ void UTransformGizmo::OnClickRelease(const FInputDeviceRay& ReleasePos)
 
 	if (HitTarget && LastHitPart != ETransformGizmoPartIdentifier::Default)
 	{
-		HitTarget->UpdateInteractingState(false, static_cast<uint32>(LastHitPart));
+		UpdateInteractingState(false, LastHitPart);
 	}
 }
 
@@ -1253,6 +1315,21 @@ void UTransformGizmo::OnClickReleaseScreenSpaceTranslate(const FInputDeviceRay&
 	bInInteraction = false;
 }
 
+void UTransformGizmo::OnClickPressScaleXYZ(const FInputDeviceRay& PressPos)
+{
+	// @todo - interaction handling coming soon?	bInInteraction = true;
+}
+
+void UTransformGizmo::OnClickDragScaleXYZ(const FInputDeviceRay& DragPos)
+{
+	// @todo - interaction handling coming soon
+}
+
+void UTransformGizmo::OnClickReleaseScaleXYZ(const FInputDeviceRay& InReleasePos)
+{
+	bInInteraction = false;
+}
+
 void UTransformGizmo::OnClickPressRotateXAxis(const FInputDeviceRay& PressPos)
 {
 	InteractionScreenAxisDirection = GetScreenRotateAxisDir(FVector::YAxisVector, FVector::ZAxisVector).GetSafeNormal();
diff --git a/Engine/Source/Editor/Experimental/EditorInteractiveToolsFramework/Public/EditorGizmos/TransformGizmo.h b/Engine/Source/Editor/Experimental/EditorInteractiveToolsFramework/Public/EditorGizmos/TransformGizmo.h
--- a/Engine/Source/Editor/Experimental/EditorInteractiveToolsFramework/Public/EditorGizmos/TransformGizmo.h
+++ b/Engine/Source/Editor/Experimental/EditorInteractiveToolsFramework/Public/EditorGizmos/TransformGizmo.h
@@ -420,6 +420,12 @@ protected:
 	/** Return true if input ray intersects current interaction plane and return param along ray in OutHitParam */
 	virtual bool GetRayParamIntersectionWithInteractionPlane(const FInputDeviceRay& InRay, float& OutHitParam);
 
+	/** Update hover state for given part id */
+	virtual void UpdateHoverState(bool bInHover, ETransformGizmoPartIdentifier InPartId);
+
+	/** Update interacting state for given part id */
+	virtual void UpdateInteractingState(bool bInInteracting, ETransformGizmoPartIdentifier InPartId);
+
 	/**
 	 * Translate and scale axis click-drag handling methods 
 	 */ 
@@ -567,6 +573,18 @@ protected:
 	/** Returns 2D vector projection of input axis onto the current viewing plane */
 	virtual FVector2D GetScreenProjectedAxis(const FVector& InAxis);
 
+	/**
+	* Scale click-drag handling methods
+	*/
+
+	/** Handle click press for uniform scale */
+	virtual void OnClickPressScaleXYZ(const FInputDeviceRay& PressPos);
+
+	/** Handle click drag for uniform scale */
+	virtual void OnClickDragScaleXYZ(const FInputDeviceRay& DragPos);
+
+	/** Handle click release for uniform scale */
+	virtual void OnClickReleaseScaleXYZ(const FInputDeviceRay& ReleasePos);
 
 	/**
 	 * Apply transform delta methods
diff --git a/Engine/Source/Runtime/InteractiveToolsFramework/Private/BaseGizmos/GizmoElementHitTargets.cpp b/Engine/Source/Runtime/InteractiveToolsFramework/Private/BaseGizmos/GizmoElementHitTargets.cpp
--- a/Engine/Source/Runtime/InteractiveToolsFramework/Private/BaseGizmos/GizmoElementHitTargets.cpp
+++ b/Engine/Source/Runtime/InteractiveToolsFramework/Private/BaseGizmos/GizmoElementHitTargets.cpp
@@ -29,7 +29,9 @@ void UGizmoElementHitTarget::UpdateHoverState(bool bHovering)
 	{
 		GizmoElement->SetElementInteractionState(EGizmoElementInteractionState::Hovering);
 	}
-	else
+	// If this element is interacting, do not update the hover state. This is necessary because when transitioning
+	// from hovering to interacting, the gizmo framework updates interacting to true before updating hovering to false.
+	else if (GizmoElement->GetElementInteractionState() != EGizmoElementInteractionState::Interacting)
 	{
 		GizmoElement->SetElementInteractionState(EGizmoElementInteractionState::None);
 	}
@@ -78,16 +80,21 @@ void UGizmoElementHitMultiTarget::UpdateHoverState(bool bInHovering, uint32 Part
 		return;
 	}
 
-	if (bInHovering && !bHovering)
+	if (bInHovering)
 	{
 		GizmoElement->UpdatePartInteractionState(EGizmoElementInteractionState::Hovering, PartIdentifier);
-		bHovering = true;
-		bInteracting = false;
 	}
-	else if (!bInHovering && bHovering)
+	else
 	{
+		// If this element is interacting, do not update the hover state. This is necessary because when transitioning
+		// from hovering to interacting, the gizmo framework updates interacting to true before updating hovering to false.
+		TOptional<EGizmoElementInteractionState> InteractionStateResult = GizmoElement->GetPartInteractionState(PartIdentifier);
+		if (InteractionStateResult.IsSet() && InteractionStateResult.GetValue() == EGizmoElementInteractionState::Interacting)
+		{
+			return;
+		}
+
 		GizmoElement->UpdatePartInteractionState(EGizmoElementInteractionState::None, PartIdentifier);
-		bHovering = false;
 	}
 }
 
@@ -98,16 +105,13 @@ void UGizmoElementHitMultiTarget::UpdateInteractingState(bool bInInteracting, ui
 		return;
 	}
 
-	if (bInInteracting && !bInteracting)
+	if (bInInteracting)
 	{
 		GizmoElement->UpdatePartInteractionState(EGizmoElementInteractionState::Interacting, PartIdentifier);
-		bInteracting = true;
-		bHovering = false;
 	}
-	else if (!bInInteracting && bInteracting)
+	else
 	{
 		GizmoElement->UpdatePartInteractionState(EGizmoElementInteractionState::None, PartIdentifier);
-		bInteracting = false;
 	}
 }
 
diff --git a/Engine/Source/Runtime/InteractiveToolsFramework/Public/BaseGizmos/GizmoElementHitTargets.h b/Engine/Source/Runtime/InteractiveToolsFramework/Public/BaseGizmos/GizmoElementHitTargets.h
--- a/Engine/Source/Runtime/InteractiveToolsFramework/Public/BaseGizmos/GizmoElementHitTargets.h
+++ b/Engine/Source/Runtime/InteractiveToolsFramework/Public/BaseGizmos/GizmoElementHitTargets.h
@@ -102,14 +102,6 @@ public:
 
 	virtual void UpdateHittableState(bool bHittable, uint32 PartIdentifier);
 
-protected:
-
-	UPROPERTY()
-	bool bHovering = false;
-
-	UPROPERTY()
-	bool bInteracting = false;
-
 public:
 	static UGizmoElementHitMultiTarget* Construct(
 		UGizmoElementBase* InGizmoElement,
