Fixed win64 non-unity compile issue.

diff --git a/Engine/Source/Editor/Experimental/HordeExecutor/Private/ContentAddressableStorage.cpp b/Engine/Source/Editor/Experimental/HordeExecutor/Private/ContentAddressableStorage.cpp
--- a/Engine/Source/Editor/Experimental/HordeExecutor/Private/ContentAddressableStorage.cpp
+++ b/Engine/Source/Editor/Experimental/HordeExecutor/Private/ContentAddressableStorage.cpp
@@ -6,7 +6,6 @@
 #include "CoreGlobals.h"
 #include "Delegates/Delegate.h"
 #include "HAL/CriticalSection.h"
-#include "HAL/PThreadCriticalSection.h"
 #include "HordeMessages.h"
 #include "HttpModule.h"
 #include "IO/IoHash.h"
