New TRS gizmo: added uniform scale interaction and modified scale and planar scale interaction to behave more similarly to current widget.

diff --git a/Engine/Source/Editor/Experimental/EditorInteractiveToolsFramework/Private/EditorGizmos/EditorTransformGizmoSource.cpp b/Engine/Source/Editor/Experimental/EditorInteractiveToolsFramework/Private/EditorGizmos/EditorTransformGizmoSource.cpp
--- a/Engine/Source/Editor/Experimental/EditorInteractiveToolsFramework/Private/EditorGizmos/EditorTransformGizmoSource.cpp
+++ b/Engine/Source/Editor/Experimental/EditorInteractiveToolsFramework/Private/EditorGizmos/EditorTransformGizmoSource.cpp
@@ -89,6 +89,16 @@ bool UEditorTransformGizmoSource::GetVisible() const
 	return false;
 }
 
+EGizmoTransformScaleType UEditorTransformGizmoSource::GetScaleType() const
+{
+	if (GEditor->UsePercentageBasedScaling())
+	{
+		return EGizmoTransformScaleType::PercentageBased;
+	}
+
+	return EGizmoTransformScaleType::Default;
+}
+
 FEditorModeTools& UEditorTransformGizmoSource::GetModeTools() const
 {
 	return GLevelEditorModeTools();
diff --git a/Engine/Source/Editor/Experimental/EditorInteractiveToolsFramework/Private/EditorGizmos/TransformGizmo.cpp b/Engine/Source/Editor/Experimental/EditorInteractiveToolsFramework/Private/EditorGizmos/TransformGizmo.cpp
--- a/Engine/Source/Editor/Experimental/EditorInteractiveToolsFramework/Private/EditorGizmos/TransformGizmo.cpp
+++ b/Engine/Source/Editor/Experimental/EditorInteractiveToolsFramework/Private/EditorGizmos/TransformGizmo.cpp
@@ -827,7 +827,7 @@ FVector UTransformGizmo::GetWorldAxis(const FVector& InAxis)
 {
 	if (TransformGizmoSource->GetGizmoCoordSystemSpace() == EToolContextCoordinateSystem::Local)
 	{
-		return CurrentTransform.GetRotation().RotateVector(InAxis);
+		return CurrentTransform.TransformVectorNoScale(InAxis);
 	}
 	
 	return InAxis;
@@ -938,9 +938,6 @@ void UTransformGizmo::UpdateHoverState(bool bInHover, ETransformGizmoPartIdentif
 		HitTarget->UpdateHoverState(bInHover, static_cast<uint32>(ETransformGizmoPartIdentifier::ScaleXAxis));
 		HitTarget->UpdateHoverState(bInHover, static_cast<uint32>(ETransformGizmoPartIdentifier::ScaleYAxis));
 		HitTarget->UpdateHoverState(bInHover, static_cast<uint32>(ETransformGizmoPartIdentifier::ScaleZAxis));
-		HitTarget->UpdateHoverState(bInHover, static_cast<uint32>(ETransformGizmoPartIdentifier::ScaleXYPlanar));
-		HitTarget->UpdateHoverState(bInHover, static_cast<uint32>(ETransformGizmoPartIdentifier::ScaleYZPlanar));
-		HitTarget->UpdateHoverState(bInHover, static_cast<uint32>(ETransformGizmoPartIdentifier::ScaleXZPlanar));
 		break;
 	case ETransformGizmoPartIdentifier::ScaleXYPlanar:
 		HitTarget->UpdateHoverState(bInHover, static_cast<uint32>(ETransformGizmoPartIdentifier::ScaleXAxis));
@@ -964,14 +961,16 @@ void UTransformGizmo::UpdateInteractingState(bool bInInteracting, ETransformGizm
 	switch (InHitPartId)
 	{
 	case ETransformGizmoPartIdentifier::ScaleUniform:
+	{
 		HitTarget->UpdateInteractingState(bInInteracting, static_cast<uint32>(ETransformGizmoPartIdentifier::ScaleXAxis));
 		HitTarget->UpdateInteractingState(bInInteracting, static_cast<uint32>(ETransformGizmoPartIdentifier::ScaleYAxis));
 		HitTarget->UpdateInteractingState(bInInteracting, static_cast<uint32>(ETransformGizmoPartIdentifier::ScaleZAxis));
-		HitTarget->UpdateInteractingState(bInInteracting, static_cast<uint32>(ETransformGizmoPartIdentifier::ScaleXYPlanar));
-		HitTarget->UpdateInteractingState(bInInteracting, static_cast<uint32>(ETransformGizmoPartIdentifier::ScaleYZPlanar));
-		HitTarget->UpdateInteractingState(bInInteracting, static_cast<uint32>(ETransformGizmoPartIdentifier::ScaleXZPlanar));
 		HitTarget->UpdateInteractingState(bInInteracting, static_cast<uint32>(ETransformGizmoPartIdentifier::ScaleUniform));
+		GizmoElementRoot->UpdatePartVisibleState(!bInInteracting, static_cast<uint32>(ETransformGizmoPartIdentifier::ScaleXYPlanar));
+		GizmoElementRoot->UpdatePartVisibleState(!bInInteracting, static_cast<uint32>(ETransformGizmoPartIdentifier::ScaleYZPlanar));
+		GizmoElementRoot->UpdatePartVisibleState(!bInInteracting, static_cast<uint32>(ETransformGizmoPartIdentifier::ScaleXZPlanar));
 		break;
+	}
 	case ETransformGizmoPartIdentifier::ScaleXYPlanar:
 		HitTarget->UpdateInteractingState(bInInteracting, static_cast<uint32>(ETransformGizmoPartIdentifier::ScaleXAxis));
 		HitTarget->UpdateInteractingState(bInInteracting, static_cast<uint32>(ETransformGizmoPartIdentifier::ScaleYAxis));
@@ -1097,30 +1096,6 @@ void UTransformGizmo::OnClickPressTranslateZAxis(const FInputDeviceRay& PressPos
 	OnClickPressAxis(PressPos);
 }
 
-void UTransformGizmo::OnClickPressScaleXAxis(const FInputDeviceRay& PressPos)
-{
-	InteractionAxisOrigin = CurrentTransform.GetLocation();
-	InteractionAxisDirection = GetWorldAxis(FVector::XAxisVector);
-	InteractionAxisList = EAxisList::X;
-	OnClickPressAxis(PressPos);
-}
-
-void UTransformGizmo::OnClickPressScaleYAxis(const FInputDeviceRay& PressPos)
-{
-	InteractionAxisOrigin = CurrentTransform.GetLocation();
-	InteractionAxisDirection = GetWorldAxis(FVector::YAxisVector);
-	InteractionAxisList = EAxisList::Y;
-	OnClickPressAxis(PressPos);
-}
-
-void UTransformGizmo::OnClickPressScaleZAxis(const FInputDeviceRay& PressPos)
-{
-	InteractionAxisOrigin = CurrentTransform.GetLocation();
-	InteractionAxisDirection = GetWorldAxis(FVector::ZAxisVector);
-	InteractionAxisList = EAxisList::Z;
-	OnClickPressAxis(PressPos);
-}
-
 void UTransformGizmo::OnClickPressAxis(const FInputDeviceRay& PressPos)
 {
 	InteractionAxisStartParam = GetNearestRayParamToInteractionAxis(PressPos);
@@ -1136,24 +1111,11 @@ void UTransformGizmo::OnClickDragTranslateAxis(const FInputDeviceRay& DragPos)
 	InteractionAxisCurrParam = AxisNearestParam;
 }
 
-void UTransformGizmo::OnClickDragScaleAxis(const FInputDeviceRay& DragPos)
-{
-	float AxisNearestParam = GetNearestRayParamToInteractionAxis(DragPos);
-	FVector Delta = ComputeAxisScaleDelta(InteractionAxisCurrParam, AxisNearestParam);
-	ApplyScaleDelta(Delta);
-	InteractionAxisCurrParam = AxisNearestParam;
-}
-
 void UTransformGizmo::OnClickReleaseTranslateAxis(const FInputDeviceRay& InReleasePos)
 {
 	bInInteraction = false;
 }
 
-void UTransformGizmo::OnClickReleaseScaleAxis(const FInputDeviceRay& InReleasePos)
-{
-	bInInteraction = false;
-}
-
 void UTransformGizmo::OnClickPressTranslateXYPlanar(const FInputDeviceRay& PressPos)
 {
 	InteractionPlanarOrigin = CurrentTransform.GetLocation();
@@ -1184,36 +1146,6 @@ void UTransformGizmo::OnClickPressTranslateXZPlanar(const FInputDeviceRay& Press
 	OnClickPressPlanar(PressPos);
 }
 
-void UTransformGizmo::OnClickPressScaleXYPlanar(const FInputDeviceRay& PressPos)
-{
-	InteractionPlanarOrigin = CurrentTransform.GetLocation();
-	InteractionPlanarNormal = GetWorldAxis(FVector::ZAxisVector);
-	InteractionPlanarAxisX = GetWorldAxis(FVector::XAxisVector);
-	InteractionPlanarAxisY = GetWorldAxis(FVector::YAxisVector);
-	InteractionAxisList = EAxisList::XY;
-	OnClickPressPlanar(PressPos);
-}
-
-void UTransformGizmo::OnClickPressScaleYZPlanar(const FInputDeviceRay& PressPos)
-{
-	InteractionPlanarOrigin = CurrentTransform.GetLocation();
-	InteractionPlanarNormal = GetWorldAxis(FVector::XAxisVector);
-	InteractionPlanarAxisX = GetWorldAxis(FVector::YAxisVector);
-	InteractionPlanarAxisY = GetWorldAxis(FVector::ZAxisVector);
-	InteractionAxisList = EAxisList::YZ;
-	OnClickPressPlanar(PressPos);
-}
-
-void UTransformGizmo::OnClickPressScaleXZPlanar(const FInputDeviceRay& PressPos)
-{
-	InteractionPlanarOrigin = CurrentTransform.GetLocation();
-	InteractionPlanarNormal = GetWorldAxis(FVector::YAxisVector);
-	InteractionPlanarAxisX = GetWorldAxis(FVector::ZAxisVector);
-	InteractionPlanarAxisY = GetWorldAxis(FVector::XAxisVector);
-	InteractionAxisList = EAxisList::XZ;
-	OnClickPressPlanar(PressPos);
-}
-
 void UTransformGizmo::OnClickPressPlanar(const FInputDeviceRay& PressPos)
 {
 	float HitDepth;
@@ -1237,62 +1169,22 @@ void UTransformGizmo::OnClickDragTranslatePlanar(const FInputDeviceRay& DragPos)
 	}
 }
 
-void UTransformGizmo::OnClickDragScalePlanar(const FInputDeviceRay& DragPos)
-{
-	float HitDepth;
-	if (GetRayParamIntersectionWithInteractionPlane(DragPos, HitDepth))
-	{
-		FVector HitPoint = DragPos.WorldRay.Origin + DragPos.WorldRay.Direction * HitDepth;
-		FVector Delta = ComputePlanarScaleDelta(InteractionPlanarCurrPoint, HitPoint);
-		ApplyScaleDelta(Delta);
-		InteractionPlanarCurrPoint = HitPoint;
-	}
-}
-
 void UTransformGizmo::OnClickReleaseTranslatePlanar(const FInputDeviceRay& InReleasePos)
 {
 	bInInteraction = false;
 }
 
-void UTransformGizmo::OnClickReleaseScalePlanar(const FInputDeviceRay& InReleasePos)
-{
-	bInInteraction = false;
-}
-
 FVector UTransformGizmo::ComputeAxisTranslateDelta(double InStartParam, double InEndParam)
 {
 	const double ParamDelta = InEndParam - InStartParam;
 	return InteractionAxisDirection * ParamDelta;
 }
 
-FVector UTransformGizmo::ComputeAxisScaleDelta(double InStartParam, double InEndParam)
-{
-	const double ParamDelta = InEndParam - InStartParam;
-	const float ScaleApplied = ParamDelta * ScaleMultiplier;
-
-	return FVector(InteractionAxisList & EAxisList::X ? ScaleApplied : 0.0,
-		InteractionAxisList & EAxisList::Y ? ScaleApplied : 0.0,
-		InteractionAxisList & EAxisList::Z ? ScaleApplied : 0.0);
-}
-
 FVector UTransformGizmo::ComputePlanarTranslateDelta(const FVector& InStartPoint, const FVector& InEndPoint)
 {
 	return InEndPoint - InStartPoint;
 }
 
-FVector UTransformGizmo::ComputePlanarScaleDelta(const FVector& InStartPoint, const FVector& InEndPoint)
-{
-	FVector Delta = InEndPoint - InStartPoint;
-	float DragUp = FVector::DotProduct(Delta, InteractionPlanarAxisX);
-	float DragSide = FVector::DotProduct(Delta, InteractionPlanarAxisY);
-	const float ScaleApplied = FMath::Abs(DragUp) > FMath::Abs(DragSide) ? DragUp * ScaleMultiplier : DragSide * ScaleMultiplier;
-
-	return FVector(
-		InteractionAxisList & EAxisList::X ? ScaleApplied : 0.0,
-		InteractionAxisList & EAxisList::Y ? ScaleApplied : 0.0,
-		InteractionAxisList & EAxisList::Z ? ScaleApplied : 0.0);
-}
-
 void UTransformGizmo::OnClickPressScreenSpaceTranslate(const FInputDeviceRay& PressPos)
 {
 	check(GizmoViewContext);
@@ -1315,14 +1207,114 @@ void UTransformGizmo::OnClickReleaseScreenSpaceTranslate(const FInputDeviceRay&
 	bInInteraction = false;
 }
 
+void UTransformGizmo::OnClickPressScaleXAxis(const FInputDeviceRay& PressPos)
+{
+	InteractionAxisList = EAxisList::X;
+	OnClickPressScale(PressPos);
+}
+
+void UTransformGizmo::OnClickPressScaleYAxis(const FInputDeviceRay& PressPos)
+{
+	InteractionAxisList = EAxisList::Y;
+	OnClickPressScale(PressPos);
+}
+
+void UTransformGizmo::OnClickPressScaleZAxis(const FInputDeviceRay& PressPos)
+{
+	InteractionAxisList = EAxisList::Z;
+	OnClickPressScale(PressPos);
+}
+
+void UTransformGizmo::OnClickPressScaleXYPlanar(const FInputDeviceRay& PressPos)
+{
+	InteractionAxisList = EAxisList::XY;
+	OnClickPressScale(PressPos);
+}
+
+void UTransformGizmo::OnClickPressScaleYZPlanar(const FInputDeviceRay& PressPos)
+{
+	InteractionAxisList = EAxisList::YZ;
+	OnClickPressScale(PressPos);
+}
+
+void UTransformGizmo::OnClickPressScaleXZPlanar(const FInputDeviceRay& PressPos)
+{
+	InteractionAxisList = EAxisList::XZ;
+	OnClickPressScale(PressPos);
+}
+
 void UTransformGizmo::OnClickPressScaleXYZ(const FInputDeviceRay& PressPos)
 {
-	// @todo - interaction handling coming soon?	bInInteraction = true;
+	InteractionAxisList = EAxisList::XYZ;
+	OnClickPressScale(PressPos);
+}
+
+void UTransformGizmo::OnClickPressScale(const FInputDeviceRay& PressPos)
+{
+	FVector2D AxisDir(0.0, 0.0);
+
+	if (InteractionAxisList & EAxisList::X)
+	{
+		AxisDir += GetScreenProjectedAxis(GizmoViewContext, FVector::XAxisVector, CurrentTransform);
+	}
+	if (InteractionAxisList & EAxisList::Y)
+	{
+		AxisDir += GetScreenProjectedAxis(GizmoViewContext, FVector::YAxisVector, CurrentTransform);
+	}
+	if (InteractionAxisList & EAxisList::Z)
+	{
+		AxisDir += GetScreenProjectedAxis(GizmoViewContext, FVector::ZAxisVector, CurrentTransform);
+	}
+
+	InteractionScreenAxisDirection = AxisDir.GetSafeNormal();
+	InteractionScreenStartPos = InteractionScreenEndPos = InteractionScreenCurrPos = PressPos.ScreenPosition;
+	bInInteraction = true;
+}
+
+void UTransformGizmo::OnClickDragScaleAxis(const FInputDeviceRay& DragPos)
+{
+	OnClickDragScale(DragPos);
+}
+
+void UTransformGizmo::OnClickDragScalePlanar(const FInputDeviceRay& DragPos)
+{
+	OnClickDragScale(DragPos);
 }
 
 void UTransformGizmo::OnClickDragScaleXYZ(const FInputDeviceRay& DragPos)
 {
-	// @todo - interaction handling coming soon
+	OnClickDragScale(DragPos);
+}
+
+void UTransformGizmo::OnClickDragScale(const FInputDeviceRay& DragPos)
+{
+	FVector2D ScreenDelta = DragPos.ScreenPosition - InteractionScreenCurrPos;
+
+	if (TransformGizmoSource->GetScaleType() != EGizmoTransformScaleType::PercentageBased)
+	{
+		ScreenDelta *= ScaleMultiplier;
+	}
+
+	InteractionScreenEndPos += ScreenDelta;
+
+	FVector ScaleDelta = ComputeScaleDelta(InteractionScreenStartPos, InteractionScreenEndPos, ScreenDelta);
+
+	if (ScaleDelta.X != 0.0 || ScaleDelta.Y != 0.0 || ScaleDelta.Z != 0.0)
+	{
+		ApplyScaleDelta(ScaleDelta);
+		InteractionScreenEndPos -= ScreenDelta;
+		InteractionScreenCurrPos = DragPos.ScreenPosition;
+	}
+}
+
+void UTransformGizmo::OnClickReleaseScaleAxis(const FInputDeviceRay& InReleasePos)
+{
+	bInInteraction = false;
+}
+
+void UTransformGizmo::OnClickReleaseScalePlanar(const FInputDeviceRay& InReleasePos)
+{
+	bInInteraction = false;
 }
 
 void UTransformGizmo::OnClickReleaseScaleXYZ(const FInputDeviceRay& InReleasePos)
@@ -1330,6 +1322,25 @@ void UTransformGizmo::OnClickReleaseScaleXYZ(const FInputDeviceRay& InReleasePos
 	bInInteraction = false;
 }
 
+FVector UTransformGizmo::ComputeScaleDelta(const FVector2D& InStartPos, const FVector2D& InEndPos, FVector2D& OutScreenDelta)
+{
+	const FVector2D DragDir = InEndPos - InStartPos;
+	const double ScaleDelta = FVector2D::DotProduct(InteractionScreenAxisDirection, DragDir);
+
+	const FVector Scale(
+		(InteractionAxisList & EAxisList::X) ? ScaleDelta : 0.0f,
+		(InteractionAxisList & EAxisList::Y) ? ScaleDelta : 0.0f,
+		(InteractionAxisList & EAxisList::Z) ? ScaleDelta : 0.0f);
+
+	const float ScaleMax = Scale.GetMax();
+	const float ScaleMin = Scale.GetMin();
+	const float ScaleApplied = (ScaleMax > -ScaleMin) ? ScaleMax : ScaleMin;
+
+	OutScreenDelta = InteractionScreenAxisDirection * ScaleApplied;
+
+	return Scale;
+}
+
 void UTransformGizmo::OnClickPressRotateXAxis(const FInputDeviceRay& PressPos)
 {
 	InteractionScreenAxisDirection = GetScreenRotateAxisDir(FVector::YAxisVector, FVector::ZAxisVector).GetSafeNormal();
@@ -1369,7 +1380,7 @@ FVector2D UTransformGizmo::GetScreenRotateAxisDir(const FVector& InAxis0, const
 
 	const FVector AxisDir = (Axis1 - Axis0) * Direction;
 
-	return GetScreenProjectedAxis(AxisDir);
+	return GetScreenProjectedAxis(GizmoViewContext, AxisDir);
 }
 
 void UTransformGizmo::OnClickDragRotateAxis(const FInputDeviceRay& DragPos)
@@ -1462,18 +1473,32 @@ void UTransformGizmo::OnClickReleaseScreenSpaceRotate(const FInputDeviceRay& InR
 	bInInteraction = false;
 }
 
-FVector2D UTransformGizmo::GetScreenProjectedAxis(const FVector& InAxis)
+FVector2D UTransformGizmo::GetScreenProjectedAxis(const UGizmoViewContext* View, const FVector& InLocalAxis, const FTransform& InLocalToWorld) const
 {
-	check(GizmoViewContext);
+	FVector2D Origin;
+	FVector2D AxisEnd;
+	const FVector WorldOrigin = InLocalToWorld.GetTranslation();
+	const FVector WorldAxisEnd = WorldOrigin + InLocalToWorld.TransformVectorNoScale(InLocalAxis * 64.0);
 
-	if (FMath::IsNearlyEqual(FVector::DotProduct(GizmoViewContext->GetViewDirection(), InAxis), 1.0))
+	if (View->ScreenToPixel(View->WorldToScreen(WorldOrigin), Origin) &&
+		View->ScreenToPixel(View->WorldToScreen(WorldAxisEnd), AxisEnd))
 	{
-		return FVector2D(1.0, 1.0);
+		// If both the origin and the axis endpoint are in front of the camera, trivially calculate the viewport space axis direction
+		return (AxisEnd - Origin).GetSafeNormal();
 	}
 	
-	const float ScreenX = FVector::DotProduct(InAxis, GizmoViewContext->GetViewRight());
-	const float ScreenY = FVector::DotProduct(InAxis, GizmoViewContext->GetViewUp());
-	return FVector2D(ScreenX, ScreenY);
+	// If either the origin or axis endpoint are behind the camera, translate the entire widget in front of the camera in the view direction before performing the
+	// viewport space calculation
+	const FMatrix InvViewMatrix = View->ViewMatrices.GetInvViewMatrix();
+	const FVector ViewLocation = InvViewMatrix.GetOrigin();
+	const FVector ViewDirection = InvViewMatrix.GetUnitAxis(EAxis::Z);
+	const FVector Offset = ViewDirection * (FVector::DotProduct(ViewLocation - WorldOrigin, ViewDirection) + 100.0f);
+	const FVector AdjustedWidgetOrigin = WorldOrigin + Offset;
+	const FVector AdjustedWidgetAxisEnd = WorldAxisEnd + Offset;
+
+	View->ScreenToPixel(View->WorldToScreen(AdjustedWidgetOrigin), Origin);
+	View->ScreenToPixel(View->WorldToScreen(AdjustedWidgetAxisEnd), AxisEnd);
+	return -(AxisEnd - Origin).GetSafeNormal();	
 }
 
 void UTransformGizmo::ApplyTranslateDelta(const FVector& InTranslateDelta)
diff --git a/Engine/Source/Editor/Experimental/EditorInteractiveToolsFramework/Public/EditorGizmos/EditorTransformGizmoSource.h b/Engine/Source/Editor/Experimental/EditorInteractiveToolsFramework/Public/EditorGizmos/EditorTransformGizmoSource.h
--- a/Engine/Source/Editor/Experimental/EditorInteractiveToolsFramework/Public/EditorGizmos/EditorTransformGizmoSource.h
+++ b/Engine/Source/Editor/Experimental/EditorInteractiveToolsFramework/Public/EditorGizmos/EditorTransformGizmoSource.h
@@ -54,6 +54,11 @@ public:
 	 */
 	virtual bool GetVisible() const;
 
+	/**
+ 	 * Get current scale type
+	 */
+	virtual EGizmoTransformScaleType GetScaleType() const;
+
 public:
 	static UEditorTransformGizmoSource* Construct(
 		UObject* Outer = (UObject*)GetTransientPackage())
diff --git a/Engine/Source/Editor/Experimental/EditorInteractiveToolsFramework/Public/EditorGizmos/TransformGizmo.h b/Engine/Source/Editor/Experimental/EditorInteractiveToolsFramework/Public/EditorGizmos/TransformGizmo.h
--- a/Engine/Source/Editor/Experimental/EditorInteractiveToolsFramework/Public/EditorGizmos/TransformGizmo.h
+++ b/Engine/Source/Editor/Experimental/EditorInteractiveToolsFramework/Public/EditorGizmos/TransformGizmo.h
@@ -427,7 +427,7 @@
 	virtual void UpdateInteractingState(bool bInInteracting, ETransformGizmoPartIdentifier InPartId);
 
 	/**
-	 * Translate and scale axis click-drag handling methods 
+	 * Translate axis click-drag handling methods 
 	 */ 
 
 	/** Handle click press for translate X axis */
@@ -439,30 +439,15 @@
 	/** Handle click press for translate Z axis */
 	virtual void OnClickPressTranslateZAxis(const FInputDeviceRay& PressPos);
 
-	/** Handle click press for scale X axis */
-	virtual void OnClickPressScaleXAxis(const FInputDeviceRay& PressPos);
-
-	/** Handle click press for scale Y axis */
-	virtual void OnClickPressScaleYAxis(const FInputDeviceRay& PressPos);
-
-	/** Handle click press for scale Z axis */
-	virtual void OnClickPressScaleZAxis(const FInputDeviceRay& PressPos);
-
 	/** Handle click press for translate and scale X axes */
 	virtual void OnClickPressAxis(const FInputDeviceRay& PressPos);
 
 	/** Handle click drag for translate axes */
 	virtual void OnClickDragTranslateAxis(const FInputDeviceRay& DragPos);
 
-	/** Handle click drag for scale axes */
-	virtual void OnClickDragScaleAxis(const FInputDeviceRay& DragPos);
-
 	/** Handle click release for translate axes */
 	virtual void OnClickReleaseTranslateAxis(const FInputDeviceRay& ReleasePos);
 
-	/** Handle click release for scale axes */
-	virtual void OnClickReleaseScaleAxis(const FInputDeviceRay& ReleasePos);
-
 	/**
 	 * Translate and scale planar click-drag handling methods
 	 */
@@ -476,42 +461,21 @@
 	/** Handle click press for translate XZ planar */
 	virtual void OnClickPressTranslateXZPlanar(const FInputDeviceRay& PressPos);
 
-	/** Handle click press for scale XY planar */
-	virtual void OnClickPressScaleXYPlanar(const FInputDeviceRay& PressPos);
-
-	/** Handle click press for scale YZ planar */
-	virtual void OnClickPressScaleYZPlanar(const FInputDeviceRay& PressPos);
-
-	/** Handle click press for scale XZ planar */
-	virtual void OnClickPressScaleXZPlanar(const FInputDeviceRay& PressPos);
-
 	/** Handle click press for generic planar */
 	virtual void OnClickPressPlanar(const FInputDeviceRay& PressPos);
 
 	/** Handle click drag for translate planar */
 	virtual void OnClickDragTranslatePlanar(const FInputDeviceRay& DragPos);
 
-	/** Handle click drag for scale planar */
-	virtual void OnClickDragScalePlanar(const FInputDeviceRay& DragPos);
-
 	/** Handle click release for translate planar */
 	virtual void OnClickReleaseTranslatePlanar(const FInputDeviceRay& ReleasePos);
 
-	/** Handle click release for scale planar */
-	virtual void OnClickReleaseScalePlanar(const FInputDeviceRay& ReleasePos);
-
 	/** Compute translate axis delta based on start/end params */
 	virtual FVector ComputeAxisTranslateDelta(double InStartParam, double InEndParam);
 
-	/** Compute scale axis delta based on start/end params */
-	virtual FVector ComputeAxisScaleDelta(double InStartParam, double InEndParam);
-
 	/** Compute translate planar delta based on world space start/end points */
 	virtual FVector ComputePlanarTranslateDelta(const FVector& InStartPoint, const FVector& InEndPoint);
 
-	/** Compute scale planar delta based on world space start/end points */
-	virtual FVector ComputePlanarScaleDelta(const FVector& InStartPoint, const FVector& InEndPoint);
-
 	/**
 	 * Screen-space translate interaction methods
 	 */
@@ -567,25 +531,64 @@
 	virtual FQuat ComputeAngularRotateDelta(double InStartAngle, double InEndAngle);
 
 	/**
- 	 * Screen-space helper methods
-	 */
-
-	/** Returns 2D vector projection of input axis onto the current viewing plane */
-	virtual FVector2D GetScreenProjectedAxis(const FVector& InAxis);
-
-	/**
 	* Scale click-drag handling methods
 	*/
 
+	/** Handle click press for scale X axis */
+	virtual void OnClickPressScaleXAxis(const FInputDeviceRay& PressPos);
+
+	/** Handle click press for scale Y axis */
+	virtual void OnClickPressScaleYAxis(const FInputDeviceRay& PressPos);
+
+	/** Handle click press for scale Z axis */
+	virtual void OnClickPressScaleZAxis(const FInputDeviceRay& PressPos);
+
+	/** Handle click press for scale XY planar */
+	virtual void OnClickPressScaleXYPlanar(const FInputDeviceRay& PressPos);
+
+	/** Handle click press for scale YZ planar */
+	virtual void OnClickPressScaleYZPlanar(const FInputDeviceRay& PressPos);
+
+	/** Handle click press for scale XZ planar */
+	virtual void OnClickPressScaleXZPlanar(const FInputDeviceRay& PressPos);
+
 	/** Handle click press for uniform scale */
 	virtual void OnClickPressScaleXYZ(const FInputDeviceRay& PressPos);
 
+	/** Handle click press for all scale methods */
+	virtual void OnClickPressScale(const FInputDeviceRay& PressPos);
+
+	/** Handle click drag for scale axes */
+	virtual void OnClickDragScaleAxis(const FInputDeviceRay& DragPos);
+
+	/** Handle click drag for scale planar */
+	virtual void OnClickDragScalePlanar(const FInputDeviceRay& DragPos);
+
 	/** Handle click drag for uniform scale */
 	virtual void OnClickDragScaleXYZ(const FInputDeviceRay& DragPos);
 
+	/** Handle click drag for all scale */
+	virtual void OnClickDragScale(const FInputDeviceRay& DragPos);
+
+	/** Handle click release for scale axes */
+	virtual void OnClickReleaseScaleAxis(const FInputDeviceRay& ReleasePos);
+
+	/** Handle click release for scale planar */
+	virtual void OnClickReleaseScalePlanar(const FInputDeviceRay& ReleasePos);
+
 	/** Handle click release for uniform scale */
 	virtual void OnClickReleaseScaleXYZ(const FInputDeviceRay& ReleasePos);
 
+	/** Compute scale delta based on screen space start/end positions */
+	virtual FVector ComputeScaleDelta(const FVector2D& InStartPos, const FVector2D& InEndPos, FVector2D& OutScreenDelta);
+
 	/**
+	 * Screen-space helper method
+	 */
+
+	/** Returns 2D vector projection of input axis onto input view plane */
+	FVector2D GetScreenProjectedAxis(const UGizmoViewContext* View, const FVector& InLocalAxis, const FTransform& InLocalToWorld = FTransform::Identity) const;
+
+	/**
 	 * Apply transform delta methods
 	 */
@@ -640,9 +643,15 @@
 	/** Array of function pointers, indexed by gizmo part id, to handle click release behavior */
 	TArray<TFunction<void(UTransformGizmo* TransformGizmo, const FInputDeviceRay& ReleasePos)> > OnClickReleaseFunctions;
 
-	/** Scale delta is multiplied by this amount */
+	/** Array of function pointers, indexed by gizmo part id, to handle update hovering state */
+	TArray<TFunction<void(UTransformGizmo* TransformGizmo, bool bInHover, uint32 InHitParts)> > OnUpdateHoverFunctions;
+
+	/** Array of function pointers, indexed by gizmo part id, to handle update interacting state */
+	TArray<TFunction<void(UTransformGizmo* TransformGizmo, bool bInInteracting, uint32 InHitPart)> > OnUpdateInteractingFunctions;
+
+	/** Percentage-based scale multiplier */
 	UPROPERTY()
-	double ScaleMultiplier = 0.05;
+	double ScaleMultiplier = 0.01;
 
 	/** Current transform */
 	UPROPERTY()
@@ -721,11 +730,15 @@
 	UPROPERTY()
 	FVector2D InteractionScreenAxisDirection;
 
-	/** Active interaction screen start pos planar (only valid between state target BeginModify/EndModify) */
+	/** Active interaction screen start pos (only valid between state target BeginModify/EndModify) */
 	UPROPERTY()
 	FVector2D InteractionScreenStartPos;
 
-	/** Active interaction screen current pos planar (only valid between state target BeginModify/EndModify) */
+	/** Active interaction screen end pos (only valid between state target BeginModify/EndModify) */
 	UPROPERTY()
+	FVector2D InteractionScreenEndPos;
+
+	/** Active interaction screen current pos (only valid between state target BeginModify/EndModify) */
+	UPROPERTY()
 	FVector2D InteractionScreenCurrPos;
 
diff --git a/Engine/Source/Editor/Experimental/EditorInteractiveToolsFramework/Public/EditorGizmos/TransformGizmoInterfaces.h b/Engine/Source/Editor/Experimental/EditorInteractiveToolsFramework/Public/EditorGizmos/TransformGizmoInterfaces.h
--- a/Engine/Source/Editor/Experimental/EditorInteractiveToolsFramework/Public/EditorGizmos/TransformGizmoInterfaces.h
+++ b/Engine/Source/Editor/Experimental/EditorInteractiveToolsFramework/Public/EditorGizmos/TransformGizmoInterfaces.h
@@ -20,6 +20,13 @@ enum class EGizmoTransformMode : uint8
 	Max
 };
 
+UENUM()
+enum class EGizmoTransformScaleType : uint8
+{
+	Default,
+	PercentageBased
+};
+
 UINTERFACE()
 class EDITORINTERACTIVETOOLSFRAMEWORK_API UTransformGizmoSource : public UInterface
 {
@@ -58,27 +65,10 @@ public:
 	 */
 	virtual bool GetVisible() const PURE_VIRTUAL(ITransformGizmoSource::GetVisible, return false;);
 
-};
-
-
-UINTERFACE()
-class EDITORINTERACTIVETOOLSFRAMEWORK_API UTransformGizmoTarget : public UInterface
-{
-	GENERATED_BODY()
-};
-
-/**
- * ITransformGizmoTarget is an interface for updating transform gizmo related state
- * which matches interfaces provided in the Editor.
- */
-class EDITORINTERACTIVETOOLSFRAMEWORK_API ITransformGizmoTarget
-{
-	GENERATED_BODY()
-public:
-
 	/*
-	 * Set the current axis when the gizmo is interacting
+	 * Get current scale type
 	 */
-	virtual void SetCurrentAxis(EAxisList::Type CurrentAxis) PURE_VIRTUAL(ITransformGizmoTarget::SetCurrentAxis, );
+	virtual EGizmoTransformScaleType GetScaleType() const PURE_VIRTUAL(ITransformGizmoSource::GetScaleType, return EGizmoTransformScaleType::Default;);
 };
 
+
