New TRS gizmo: added StateTarget to support undo/redo.

diff --git a/Engine/Source/Editor/Experimental/EditorInteractiveToolsFramework/Private/EditorGizmos/TransformGizmo.cpp b/Engine/Source/Editor/Experimental/EditorInteractiveToolsFramework/Private/EditorGizmos/TransformGizmo.cpp
--- a/Engine/Source/Editor/Experimental/EditorInteractiveToolsFramework/Private/EditorGizmos/TransformGizmo.cpp
+++ b/Engine/Source/Editor/Experimental/EditorInteractiveToolsFramework/Private/EditorGizmos/TransformGizmo.cpp
@@ -577,18 +577,8 @@ void UTransformGizmo::SetActiveTarget(UTransformProxy* Target, IToolContextTrans
 		TransactionProvider = GetGizmoManager();
 	}
 	
-	/* @todo - Add state target 
-	
-	// This state target emits an explicit FChange that moves the GizmoActor root component during undo/redo.
-	// It also opens/closes the Transaction that saves/restores the target object locations.
-	if (TransactionProvider == nullptr)
-	{
-		TransactionProvider = GetGizmoManager();
-	}
-	StateTarget = UGizmoTransformChangeStateTarget::Construct(GizmoElementRoot,
-		LOCTEXT("UCombinedTransformGizmoTransaction", "Transform"), TransactionProvider, this);
-	StateTarget->DependentChangeSources.Add(MakeUnique<FTransformProxyChangeSource>(Target));
-	*/
+	StateTarget = UGizmoObjectModifyStateTarget::Construct(Target,
+		LOCTEXT("UTransformGizmoTransaction", "Transform"), TransactionProvider, this);
 
 	CameraAxisSource = NewObject<UGizmoConstantFrameAxisSource>(this);
 }
diff --git a/Engine/Source/Editor/Experimental/EditorInteractiveToolsFramework/Public/EditorGizmos/TransformGizmo.h b/Engine/Source/Editor/Experimental/EditorInteractiveToolsFramework/Public/EditorGizmos/TransformGizmo.h
--- a/Engine/Source/Editor/Experimental/EditorInteractiveToolsFramework/Public/EditorGizmos/TransformGizmo.h
+++ b/Engine/Source/Editor/Experimental/EditorInteractiveToolsFramework/Public/EditorGizmos/TransformGizmo.h
@@ -5,8 +5,8 @@
 #include "CoreMinimal.h"
 #include "BaseBehaviors/BehaviorTargetInterfaces.h"
 #include "BaseGizmos/GizmoElementHitTargets.h"
-#include "BaseGizmos/GizmoElementStateTargets.h"
 #include "BaseGizmos/GizmoViewContext.h"
+#include "BaseGizmos/StateTargets.h"
 #include "BaseGizmos/TransformProxy.h"
 #include "EditorGizmos/TransformGizmoInterfaces.h"
 #include "InteractiveGizmo.h"
@@ -342,7 +342,7 @@ protected:
 
 	/** The state target is created internally during SetActiveTarget() */
 	UPROPERTY()
-	TObjectPtr<IGizmoStateTarget> StateTarget;
+	TObjectPtr<UGizmoObjectModifyStateTarget> StateTarget;
 
 	/**
 	 * These are used to let the translation subgizmos use raycasts into the scene to align the gizmo with scene geometry.
