Fixed compile issue found on the farm.

diff --git a/Engine/Source/Editor/Experimental/HordeExecutor/Private/HordeExecutorSettings.h b/Engine/Source/Editor/Experimental/HordeExecutor/Private/HordeExecutorSettings.h
--- a/Engine/Source/Editor/Experimental/HordeExecutor/Private/HordeExecutorSettings.h
+++ b/Engine/Source/Editor/Experimental/HordeExecutor/Private/HordeExecutorSettings.h
@@ -4,10 +4,11 @@
 
 #include "Containers/Map.h"
 #include "Containers/UnrealString.h"
-#include "HordeExecutorSettings"
 #include "UObject/Object.h"
 #include "UObject/UObjectGlobals.h"
 
+#include "HordeExecutorSettings.generated.h"
+
 UCLASS(config = EditorSettings)
 class UHordeExecutorSettings : public UObject
 {
diff --git a/Engine/Source/Editor/Experimental/RemoteExecution/Private/RemoteExecutionSettings.h b/Engine/Source/Editor/Experimental/RemoteExecution/Private/RemoteExecutionSettings.h
--- a/Engine/Source/Editor/Experimental/RemoteExecution/Private/RemoteExecutionSettings.h
+++ b/Engine/Source/Editor/Experimental/RemoteExecution/Private/RemoteExecutionSettings.h
@@ -3,10 +3,11 @@
 #pragma once
 
 #include "Containers/UnrealString.h"
-#include "RemoteExecutionSettings"
 #include "UObject/Object.h"
 #include "UObject/UObjectGlobals.h"
 
+#include "RemoteExecutionSettings.generated.h"
+
 UCLASS(config = EditorSettings)
 class URemoteExecutionSettings : public UObject
 {
