Cleaned up build.cs files by removing any include paths that were already being added by UBT. This was done to help identify how include paths are being added and to help with future refactoring.

diff --git a/Engine/Plugins/2D/Paper2D/Source/Paper2D/Paper2D.Build.cs b/Engine/Plugins/2D/Paper2D/Source/Paper2D/Paper2D.Build.cs
--- a/Engine/Plugins/2D/Paper2D/Source/Paper2D/Paper2D.Build.cs
+++ b/Engine/Plugins/2D/Paper2D/Source/Paper2D/Paper2D.Build.cs
@@ -6,8 +6,6 @@ public class Paper2D : ModuleRules
 {
 	public Paper2D(ReadOnlyTargetRules Target) : base(Target)
 	{
-		PrivateIncludePaths.Add("Paper2D/Private");
-
 		PublicDependencyModuleNames.AddRange(
 			new string[] {
 				"Core",
diff --git a/Engine/Plugins/2D/Paper2D/Source/Paper2DEditor/Paper2DEditor.Build.cs b/Engine/Plugins/2D/Paper2D/Source/Paper2DEditor/Paper2DEditor.Build.cs
--- a/Engine/Plugins/2D/Paper2D/Source/Paper2DEditor/Paper2DEditor.Build.cs
+++ b/Engine/Plugins/2D/Paper2D/Source/Paper2DEditor/Paper2DEditor.Build.cs
@@ -6,8 +6,6 @@ public class Paper2DEditor : ModuleRules
 {
 	public Paper2DEditor(ReadOnlyTargetRules Target) : base(Target)
 	{
-		PrivateIncludePaths.Add("Paper2DEditor/Private");
-
 		PrivateDependencyModuleNames.AddRange(
 			new string[] {
 				"Core",
diff --git a/Engine/Plugins/AI/EnvironmentQueryEditor/Source/EnvironmentQueryEditor/EnvironmentQueryEditor.Build.cs b/Engine/Plugins/AI/EnvironmentQueryEditor/Source/EnvironmentQueryEditor/EnvironmentQueryEditor.Build.cs
--- a/Engine/Plugins/AI/EnvironmentQueryEditor/Source/EnvironmentQueryEditor/EnvironmentQueryEditor.Build.cs
+++ b/Engine/Plugins/AI/EnvironmentQueryEditor/Source/EnvironmentQueryEditor/EnvironmentQueryEditor.Build.cs
@@ -12,7 +12,6 @@ public class EnvironmentQueryEditor : ModuleRules
 		// These nodes are not public so are hard to subclass
 		PrivateIncludePaths.AddRange(
 			new string[] {
-					"EnvironmentQueryEditor/Private",
 					 Path.Combine(EngineDir, @"Source/Editor/GraphEditor/Private"),
 					 Path.Combine(EngineDir, @"Source/Editor/AIGraph/Private")
 			});
diff --git a/Engine/Plugins/AI/HTNPlanner/Source/HTNPlanner/HTNPlanner.Build.cs b/Engine/Plugins/AI/HTNPlanner/Source/HTNPlanner/HTNPlanner.Build.cs
--- a/Engine/Plugins/AI/HTNPlanner/Source/HTNPlanner/HTNPlanner.Build.cs
+++ b/Engine/Plugins/AI/HTNPlanner/Source/HTNPlanner/HTNPlanner.Build.cs
@@ -6,16 +6,6 @@ namespace UnrealBuildTool.Rules
     {
         public HTNPlanner(ReadOnlyTargetRules Target) : base(Target)
         {
-            PublicIncludePaths.AddRange(
-                    new string[] {
-                    }
-                    );
-
-            PrivateIncludePaths.AddRange(
-                new string[] {
-                }
-                );
-
             PublicDependencyModuleNames.AddRange(
                 new string[] {
                         "Core",
@@ -27,12 +17,6 @@ namespace UnrealBuildTool.Rules
                 }
                 );
 
-            DynamicallyLoadedModuleNames.AddRange(
-                new string[] {
-                    // ... add any modules that your module loads dynamically here ...
-                }
-                );
-
             if (Target.bBuildEditor == true)
             {
                 PrivateDependencyModuleNames.Add("EditorFramework");
diff --git a/Engine/Plugins/AI/MLAdapter/Source/MLAdapter/MLAdapter.Build.cs b/Engine/Plugins/AI/MLAdapter/Source/MLAdapter/MLAdapter.Build.cs
--- a/Engine/Plugins/AI/MLAdapter/Source/MLAdapter/MLAdapter.Build.cs
+++ b/Engine/Plugins/AI/MLAdapter/Source/MLAdapter/MLAdapter.Build.cs
@@ -19,8 +19,6 @@ namespace UnrealBuildTool.Rules
                 }
             );
 
-            PrivateIncludePaths.Add("MLAdapter/Private");
-
             PublicDependencyModuleNames.AddRange(
                 new string[] {
                     "Core",
diff --git a/Engine/Plugins/AI/MassAI/Source/MassAIBehavior/MassAIBehavior.Build.cs b/Engine/Plugins/AI/MassAI/Source/MassAIBehavior/MassAIBehavior.Build.cs
--- a/Engine/Plugins/AI/MassAI/Source/MassAIBehavior/MassAIBehavior.Build.cs
+++ b/Engine/Plugins/AI/MassAI/Source/MassAIBehavior/MassAIBehavior.Build.cs
@@ -15,12 +15,6 @@ namespace UnrealBuildTool.Rules
 				}
 			);
 
-
-			PrivateIncludePaths.AddRange(
-				new string[] {
-				} 
-			);
-
 			PublicDependencyModuleNames.AddRange(
 				new string[] {
 					"MassEntity",
diff --git a/Engine/Plugins/AI/MassAI/Source/MassAIDebug/MassAIDebug.Build.cs b/Engine/Plugins/AI/MassAI/Source/MassAIDebug/MassAIDebug.Build.cs
--- a/Engine/Plugins/AI/MassAI/Source/MassAIDebug/MassAIDebug.Build.cs
+++ b/Engine/Plugins/AI/MassAI/Source/MassAIDebug/MassAIDebug.Build.cs
@@ -15,12 +15,6 @@ namespace UnrealBuildTool.Rules
 				}
 			);
 
-
-			PrivateIncludePaths.AddRange(
-				new string[] {
-				}
-			);
-
 			PublicDependencyModuleNames.AddRange(
 				new string[] {
 					"Core",
diff --git a/Engine/Plugins/Animation/ControlRig/Source/ControlRig/ControlRig.Build.cs b/Engine/Plugins/Animation/ControlRig/Source/ControlRig/ControlRig.Build.cs
--- a/Engine/Plugins/Animation/ControlRig/Source/ControlRig/ControlRig.Build.cs
+++ b/Engine/Plugins/Animation/ControlRig/Source/ControlRig/ControlRig.Build.cs
@@ -6,7 +6,6 @@ namespace UnrealBuildTool.Rules
     {
         public ControlRig(ReadOnlyTargetRules Target) : base(Target)
         {
-            PrivateIncludePaths.Add("ControlRig/Private");
             PrivateIncludePaths.Add("ControlRig/Private/Sequencer");
             PrivateIncludePaths.Add("ControlRig/Private/Units");
             PrivateIncludePaths.Add("ControlRig/ThirdParty/AHEasing");
diff --git a/Engine/Plugins/Animation/ControlRig/Source/ControlRigDeveloper/ControlRigDeveloper.Build.cs b/Engine/Plugins/Animation/ControlRig/Source/ControlRigDeveloper/ControlRigDeveloper.Build.cs
--- a/Engine/Plugins/Animation/ControlRig/Source/ControlRigDeveloper/ControlRigDeveloper.Build.cs
+++ b/Engine/Plugins/Animation/ControlRig/Source/ControlRigDeveloper/ControlRigDeveloper.Build.cs
@@ -7,7 +7,6 @@ namespace UnrealBuildTool.Rules
         public ControlRigDeveloper(ReadOnlyTargetRules Target) : base(Target)
         {
             PrivateIncludePaths.Add("ControlRig/Private");
-            PrivateIncludePaths.Add("ControlRigDeveloper/Private");
 
             // Copying some these from ControlRig.build.cs, our deps are likely leaner
             // and therefore these could be pruned if needed:
diff --git a/Engine/Plugins/Animation/ControlRig/Source/ControlRigEditor/ControlRigEditor.Build.cs b/Engine/Plugins/Animation/ControlRig/Source/ControlRigEditor/ControlRigEditor.Build.cs
--- a/Engine/Plugins/Animation/ControlRig/Source/ControlRigEditor/ControlRigEditor.Build.cs
+++ b/Engine/Plugins/Animation/ControlRig/Source/ControlRigEditor/ControlRigEditor.Build.cs
@@ -8,7 +8,6 @@ namespace UnrealBuildTool.Rules
         {
             PrivateIncludePaths.AddRange(
                 new string[] {
-                    "ControlRigEditor/Private",
                     "../../../../Source/Editor/UnrealEd/Private", 
 					"../../../../Source/Developer/MessageLog/Private", //compatibility for FBX importer
 				}
diff --git a/Engine/Plugins/Animation/MotionWarping/Source/MotionWarping/MotionWarping.Build.cs b/Engine/Plugins/Animation/MotionWarping/Source/MotionWarping/MotionWarping.Build.cs
--- a/Engine/Plugins/Animation/MotionWarping/Source/MotionWarping/MotionWarping.Build.cs
+++ b/Engine/Plugins/Animation/MotionWarping/Source/MotionWarping/MotionWarping.Build.cs
@@ -7,21 +7,7 @@ public class MotionWarping : ModuleRules
 	public MotionWarping(ReadOnlyTargetRules Target) : base(Target)
 	{
 		PCHUsage = ModuleRules.PCHUsageMode.UseExplicitOrSharedPCHs;
-		
-		PublicIncludePaths.AddRange(
-			new string[] {
-				// ... add public include paths required here ...
-			}
-			);
-				
-		
-		PrivateIncludePaths.AddRange(
-			new string[] {
-				// ... add other private include paths required here ...
-			}
-			);
-			
-		
+
 		PublicDependencyModuleNames.AddRange(
 			new string[]
 			{
@@ -30,8 +16,7 @@ public class MotionWarping : ModuleRules
 				"Engine",
 				"NetCore"
 			}
-			);
-			
+			);			
 		
 		PrivateDependencyModuleNames.AddRange(
 			new string[]
@@ -40,13 +25,5 @@ public class MotionWarping : ModuleRules
 				"SlateCore"
             }
 			);
-		
-		
-		DynamicallyLoadedModuleNames.AddRange(
-			new string[]
-			{
-				// ... add any modules that your module loads dynamically here ...
-			}
-			);
 	}
 }
diff --git a/Engine/Plugins/AudioGameplay/Source/AudioGameplay/AudioGameplay.Build.cs b/Engine/Plugins/AudioGameplay/Source/AudioGameplay/AudioGameplay.Build.cs
--- a/Engine/Plugins/AudioGameplay/Source/AudioGameplay/AudioGameplay.Build.cs
+++ b/Engine/Plugins/AudioGameplay/Source/AudioGameplay/AudioGameplay.Build.cs
@@ -6,19 +6,7 @@ public class AudioGameplay : ModuleRules
 {
 	public AudioGameplay(ReadOnlyTargetRules Target) : base(Target)
 	{
-		PCHUsage = ModuleRules.PCHUsageMode.UseExplicitOrSharedPCHs;
-
-		PublicIncludePaths.AddRange(
-			new string[] {
-			}
-			);
-				
-		
-		PrivateIncludePaths.AddRange(
-			new string[] {
-			}
-			);
-			
+		PCHUsage = ModuleRules.PCHUsageMode.UseExplicitOrSharedPCHs;			
 		
 		PublicDependencyModuleNames.AddRange(
 			new string[]
diff --git a/Engine/Plugins/AudioGameplayVolume/Source/AudioGameplayVolume/AudioGameplayVolume.Build.cs b/Engine/Plugins/AudioGameplayVolume/Source/AudioGameplayVolume/AudioGameplayVolume.Build.cs
--- a/Engine/Plugins/AudioGameplayVolume/Source/AudioGameplayVolume/AudioGameplayVolume.Build.cs
+++ b/Engine/Plugins/AudioGameplayVolume/Source/AudioGameplayVolume/AudioGameplayVolume.Build.cs
@@ -6,27 +6,14 @@ public class AudioGameplayVolume : ModuleRules
 {
 	public AudioGameplayVolume(ReadOnlyTargetRules Target) : base(Target)
 	{
-		PCHUsage = ModuleRules.PCHUsageMode.UseExplicitOrSharedPCHs;
-
-		PublicIncludePaths.AddRange(
-			new string[] {
-			}
-			);
-				
-		
-		PrivateIncludePaths.AddRange(
-			new string[] {
-			}
-			);
-			
+		PCHUsage = ModuleRules.PCHUsageMode.UseExplicitOrSharedPCHs;			
 		
 		PublicDependencyModuleNames.AddRange(
 			new string[]
 			{
 				"Core",
 			}
-			);
-			
+			);			
 		
 		PrivateDependencyModuleNames.AddRange(
 			new string[]
diff --git a/Engine/Plugins/BlueprintFileUtils/Source/BlueprintFileUtils/BlueprintFileUtils.Build.cs b/Engine/Plugins/BlueprintFileUtils/Source/BlueprintFileUtils/BlueprintFileUtils.Build.cs
--- a/Engine/Plugins/BlueprintFileUtils/Source/BlueprintFileUtils/BlueprintFileUtils.Build.cs
+++ b/Engine/Plugins/BlueprintFileUtils/Source/BlueprintFileUtils/BlueprintFileUtils.Build.cs
@@ -6,23 +6,7 @@ public class BlueprintFileUtils : ModuleRules
 {
 	public BlueprintFileUtils(ReadOnlyTargetRules Target) : base(Target)
 	{
-		PCHUsage = ModuleRules.PCHUsageMode.UseExplicitOrSharedPCHs;
-		
-		PublicIncludePaths.AddRange(
-			new string[] {
-				// ... add public include paths required here ...
-			}
-			);
-				
-		
-		PrivateIncludePaths.AddRange(
-			new string[] {
-				"BlueprintFileUtils/Private",
-				
-				// ... add other private include paths required here ...
-			}
-			);
-			
+		PCHUsage = ModuleRules.PCHUsageMode.UseExplicitOrSharedPCHs;			
 		
 		PublicDependencyModuleNames.AddRange(
 			new string[]
diff --git a/Engine/Plugins/Compositing/Composure/Source/Composure/Composure.Build.cs b/Engine/Plugins/Compositing/Composure/Source/Composure/Composure.Build.cs
--- a/Engine/Plugins/Compositing/Composure/Source/Composure/Composure.Build.cs
+++ b/Engine/Plugins/Compositing/Composure/Source/Composure/Composure.Build.cs
@@ -9,7 +9,6 @@ namespace UnrealBuildTool.Rules
 			PrivateIncludePaths.AddRange(
 				new string[] {
                     "../../../../Source/Runtime/Engine/",
-					"Composure/Private/"
 				}
 				);
             
diff --git a/Engine/Plugins/Compositing/Composure/Source/ComposureEditor/ComposureEditor.Build.cs b/Engine/Plugins/Compositing/Composure/Source/ComposureEditor/ComposureEditor.Build.cs
--- a/Engine/Plugins/Compositing/Composure/Source/ComposureEditor/ComposureEditor.Build.cs
+++ b/Engine/Plugins/Compositing/Composure/Source/ComposureEditor/ComposureEditor.Build.cs
@@ -6,9 +6,6 @@ namespace UnrealBuildTool.Rules
 	{
 		public ComposureEditor(ReadOnlyTargetRules Target) : base(Target)
 		{
-			PrivateIncludePaths.AddRange(new string[] {
-			});
-
 			PublicDependencyModuleNames.AddRange(new string[] {
 				"Core",
 				"CoreUObject",
diff --git a/Engine/Plugins/Compositing/OpenCVLensDistortion/Source/OpenCVLensCalibration/OpenCVLensCalibration.Build.cs b/Engine/Plugins/Compositing/OpenCVLensDistortion/Source/OpenCVLensCalibration/OpenCVLensCalibration.Build.cs
--- a/Engine/Plugins/Compositing/OpenCVLensDistortion/Source/OpenCVLensCalibration/OpenCVLensCalibration.Build.cs
+++ b/Engine/Plugins/Compositing/OpenCVLensDistortion/Source/OpenCVLensCalibration/OpenCVLensCalibration.Build.cs
@@ -6,21 +6,6 @@ public class OpenCVLensCalibration : ModuleRules
 {
 	public OpenCVLensCalibration(ReadOnlyTargetRules Target) : base(Target)
 	{
-		PublicIncludePaths.AddRange(
-			new string[] {
-            }
-		);
-
-		PrivateIncludePaths.AddRange(
-			new string[] {
-			}
-		);
-
-		PublicDependencyModuleNames.AddRange(
-			new string[] {
-			}
-		);
-
 		PrivateDependencyModuleNames.AddRange(
 			new string[] {
 				"Core",
diff --git a/Engine/Plugins/Compositing/OpenCVLensDistortion/Source/OpenCVLensDistortion/OpenCVLensDistortion.Build.cs b/Engine/Plugins/Compositing/OpenCVLensDistortion/Source/OpenCVLensDistortion/OpenCVLensDistortion.Build.cs
--- a/Engine/Plugins/Compositing/OpenCVLensDistortion/Source/OpenCVLensDistortion/OpenCVLensDistortion.Build.cs
+++ b/Engine/Plugins/Compositing/OpenCVLensDistortion/Source/OpenCVLensDistortion/OpenCVLensDistortion.Build.cs
@@ -5,19 +5,7 @@ namespace UnrealBuildTool.Rules
 	public class OpenCVLensDistortion : ModuleRules
 	{
 		public OpenCVLensDistortion(ReadOnlyTargetRules Target) : base(Target)
-		{
-			PublicIncludePaths.AddRange(
-				new string[] {
-					// ... add public include paths required here ...
-				}
-			);
-
-			PrivateIncludePaths.AddRange(
-				new string[] {
-					// ... add other private include paths required here ...
-				}
-			);
-            
+		{            
             PublicDependencyModuleNames.AddRange(
 				new string[]
 				{
diff --git a/Engine/Plugins/Compositing/OpenColorIO/Source/OpenColorIOEditor/OpenColorIOEditor.Build.cs b/Engine/Plugins/Compositing/OpenColorIO/Source/OpenColorIOEditor/OpenColorIOEditor.Build.cs
--- a/Engine/Plugins/Compositing/OpenColorIO/Source/OpenColorIOEditor/OpenColorIOEditor.Build.cs
+++ b/Engine/Plugins/Compositing/OpenColorIO/Source/OpenColorIOEditor/OpenColorIOEditor.Build.cs
@@ -6,16 +6,6 @@ namespace UnrealBuildTool.Rules
 	{
 		public OpenColorIOEditor(ReadOnlyTargetRules Target) : base(Target)
 		{
-			PublicIncludePaths.AddRange(
-				new string[] {
-				}
-				);
-
-			PrivateIncludePaths.AddRange(
-				new string[] {
-				}
-				);
-
 			PrivateDependencyModuleNames.AddRange(
 				new string[] {
 					"Core",
diff --git a/Engine/Plugins/Developer/AnimationSharing/Source/AnimationSharing/AnimationSharing.Build.cs b/Engine/Plugins/Developer/AnimationSharing/Source/AnimationSharing/AnimationSharing.Build.cs
--- a/Engine/Plugins/Developer/AnimationSharing/Source/AnimationSharing/AnimationSharing.Build.cs
+++ b/Engine/Plugins/Developer/AnimationSharing/Source/AnimationSharing/AnimationSharing.Build.cs
@@ -21,11 +21,6 @@ public class AnimationSharing : ModuleRules
         if (Target.Type == TargetType.Editor)
         {
             PrivateDependencyModuleNames.Add("TargetPlatform");
-        }
-
-        PrivateIncludePaths.AddRange(
-            new string[] {
-                "AnimationSharing/Private"
-         });       
+        }     
     }
 }
diff --git a/Engine/Plugins/Developer/AnimationSharing/Source/AnimationSharingEd/AnimationSharingEd.Build.cs b/Engine/Plugins/Developer/AnimationSharing/Source/AnimationSharingEd/AnimationSharingEd.Build.cs
--- a/Engine/Plugins/Developer/AnimationSharing/Source/AnimationSharingEd/AnimationSharingEd.Build.cs
+++ b/Engine/Plugins/Developer/AnimationSharing/Source/AnimationSharingEd/AnimationSharingEd.Build.cs
@@ -22,11 +22,6 @@ public class AnimationSharingEd : ModuleRules
                 "PropertyEditor",
             }
 		);
-
-        PrivateIncludePaths.AddRange(
-            new string[] {
-                "AnimationSharingEd/Private"
-         });       
     }
 
 }
diff --git a/Engine/Plugins/Editor/AssetManagerEditor/Source/AssetManagerEditor/AssetManagerEditor.Build.cs b/Engine/Plugins/Editor/AssetManagerEditor/Source/AssetManagerEditor/AssetManagerEditor.Build.cs
--- a/Engine/Plugins/Editor/AssetManagerEditor/Source/AssetManagerEditor/AssetManagerEditor.Build.cs
+++ b/Engine/Plugins/Editor/AssetManagerEditor/Source/AssetManagerEditor/AssetManagerEditor.Build.cs
@@ -6,12 +6,6 @@ public class AssetManagerEditor : ModuleRules
 {
 	public AssetManagerEditor(ReadOnlyTargetRules Target) : base(Target)
 	{
-		PrivateIncludePaths.AddRange(
-			new string[] {
-				"AssetManagerEditor/Private",
-			}
-		);
-
 		PublicDependencyModuleNames.AddRange(
 			new string[] { 
 				"Core",
diff --git a/Engine/Plugins/Editor/BlueprintMaterialTextureNodes/Source/BlueprintMaterialTextureNodes/BlueprintMaterialTextureNodes.Build.cs b/Engine/Plugins/Editor/BlueprintMaterialTextureNodes/Source/BlueprintMaterialTextureNodes/BlueprintMaterialTextureNodes.Build.cs
--- a/Engine/Plugins/Editor/BlueprintMaterialTextureNodes/Source/BlueprintMaterialTextureNodes/BlueprintMaterialTextureNodes.Build.cs
+++ b/Engine/Plugins/Editor/BlueprintMaterialTextureNodes/Source/BlueprintMaterialTextureNodes/BlueprintMaterialTextureNodes.Build.cs
@@ -14,17 +14,7 @@ public class BlueprintMaterialTextureNodes : ModuleRules
 				
 				// ... add public include paths required here ...
 			}
-			);
-				
-		
-		PrivateIncludePaths.AddRange(
-			new string[] {
-				"BlueprintMaterialTextureNodes/Private",
-				
-				// ... add other private include paths required here ...
-			}
-			);
-			
+			);			
 		
 		PublicDependencyModuleNames.AddRange(
 			new string[]
diff --git a/Engine/Plugins/Editor/EditorDebugTools/Source/EditorDebugTools/EditorDebugTools.Build.cs b/Engine/Plugins/Editor/EditorDebugTools/Source/EditorDebugTools/EditorDebugTools.Build.cs
--- a/Engine/Plugins/Editor/EditorDebugTools/Source/EditorDebugTools/EditorDebugTools.Build.cs
+++ b/Engine/Plugins/Editor/EditorDebugTools/Source/EditorDebugTools/EditorDebugTools.Build.cs
@@ -6,27 +6,14 @@ public class EditorDebugTools : ModuleRules
 {
 	public EditorDebugTools(ReadOnlyTargetRules Target) : base(Target)
 	{
-		PCHUsage = ModuleRules.PCHUsageMode.UseExplicitOrSharedPCHs;
-		
-		PublicIncludePaths.AddRange(
-			new string[] {
-			}
-			);
-				
-		
-		PrivateIncludePaths.AddRange(
-			new string[] {
-			}
-			);
-			
+		PCHUsage = ModuleRules.PCHUsageMode.UseExplicitOrSharedPCHs;			
 		
 		PublicDependencyModuleNames.AddRange(
 			new string[]
 			{
 				"Core",
 			}
-			);
-			
+			);			
 		
 		PrivateDependencyModuleNames.AddRange(
 			new string[]
diff --git a/Engine/Plugins/Editor/GLTFImporter/Source/GLTFImporter/GLTFImporter.Build.cs b/Engine/Plugins/Editor/GLTFImporter/Source/GLTFImporter/GLTFImporter.Build.cs
--- a/Engine/Plugins/Editor/GLTFImporter/Source/GLTFImporter/GLTFImporter.Build.cs
+++ b/Engine/Plugins/Editor/GLTFImporter/Source/GLTFImporter/GLTFImporter.Build.cs
@@ -6,22 +6,6 @@ namespace UnrealBuildTool.Rules
     {
         public GLTFImporter(ReadOnlyTargetRules Target) : base(Target)
         {
-            PublicIncludePaths.AddRange(
-                new string[] {
-                }
-                );
-
-            PrivateIncludePaths.AddRange(
-                new string[] {
-                }
-                );
-
-            PublicDependencyModuleNames.AddRange(
-                new string[]
-                {
-                }
-                );
-
             PrivateDependencyModuleNames.AddRange(
                 new string[]
                 {
@@ -43,12 +27,6 @@ namespace UnrealBuildTool.Rules
 					"GLTFCore",
                 }
                 );
-
-            DynamicallyLoadedModuleNames.AddRange(
-                new string[]
-                {
-                }
-                );
         }
     }
 }
diff --git a/Engine/Plugins/Editor/GuidedTutorials/Source/IntroTutorials/IntroTutorials.Build.cs b/Engine/Plugins/Editor/GuidedTutorials/Source/IntroTutorials/IntroTutorials.Build.cs
--- a/Engine/Plugins/Editor/GuidedTutorials/Source/IntroTutorials/IntroTutorials.Build.cs
+++ b/Engine/Plugins/Editor/GuidedTutorials/Source/IntroTutorials/IntroTutorials.Build.cs
@@ -22,12 +22,6 @@ namespace UnrealBuildTool.Rules
 				}
 			);
 
-			PrivateIncludePaths.AddRange(
-				new string[] {
-					"IntroTutorials/Private",
-				}
-			);
-
 			PrivateDependencyModuleNames.AddRange(
 				new string[] {
                     "AppFramework",
diff --git a/Engine/Plugins/Editor/MobileLauncherProfileWizard/Source/MobileLauncherProfileWizard/MobileLauncherProfileWizard.Build.cs b/Engine/Plugins/Editor/MobileLauncherProfileWizard/Source/MobileLauncherProfileWizard/MobileLauncherProfileWizard.Build.cs
--- a/Engine/Plugins/Editor/MobileLauncherProfileWizard/Source/MobileLauncherProfileWizard/MobileLauncherProfileWizard.Build.cs
+++ b/Engine/Plugins/Editor/MobileLauncherProfileWizard/Source/MobileLauncherProfileWizard/MobileLauncherProfileWizard.Build.cs
@@ -6,18 +6,6 @@ namespace UnrealBuildTool.Rules
 	{
 		public MobileLauncherProfileWizard(ReadOnlyTargetRules Target) : base(Target)
 		{
-			PublicIncludePaths.AddRange(
-				new string[] {
-					// ... add public include paths required here ...
-				}
-				);
-
-			PrivateIncludePaths.AddRange(
-				new string[] {
-					// ... add other private include paths required here ...
-				}
-				);
-
 			PublicDependencyModuleNames.AddRange(
 				new string[]
 				{
@@ -42,13 +30,6 @@ namespace UnrealBuildTool.Rules
                     // ... add private dependencies that you statically link with here ...
 				}
                 );
-
-			DynamicallyLoadedModuleNames.AddRange(
-				new string[]
-				{
-					// ... add any modules that your module loads dynamically here ...
-				}
-				);
 		}
 	}
 }
diff --git a/Engine/Plugins/Editor/ModelingToolsEditorMode/Source/ModelingToolsEditorMode/ModelingToolsEditorMode.Build.cs b/Engine/Plugins/Editor/ModelingToolsEditorMode/Source/ModelingToolsEditorMode/ModelingToolsEditorMode.Build.cs
--- a/Engine/Plugins/Editor/ModelingToolsEditorMode/Source/ModelingToolsEditorMode/ModelingToolsEditorMode.Build.cs
+++ b/Engine/Plugins/Editor/ModelingToolsEditorMode/Source/ModelingToolsEditorMode/ModelingToolsEditorMode.Build.cs
@@ -8,20 +8,6 @@ public class ModelingToolsEditorMode : ModuleRules
 	{
 		PCHUsage = ModuleRules.PCHUsageMode.UseExplicitOrSharedPCHs;
 
-		PublicIncludePaths.AddRange(
-			new string[]
-			{
-				// ... add public include paths required here ...
-			}
-		);
-
-		PrivateIncludePaths.AddRange(
-			new string[]
-			{
-				// ... add other private include paths required here ...
-			}
-		);
-
 		PublicDependencyModuleNames.AddRange(
 			new string[]
 			{
diff --git a/Engine/Plugins/Editor/SpeedTreeImporter/Source/SpeedTreeImporter/SpeedTreeImporter.Build.cs b/Engine/Plugins/Editor/SpeedTreeImporter/Source/SpeedTreeImporter/SpeedTreeImporter.Build.cs
--- a/Engine/Plugins/Editor/SpeedTreeImporter/Source/SpeedTreeImporter/SpeedTreeImporter.Build.cs
+++ b/Engine/Plugins/Editor/SpeedTreeImporter/Source/SpeedTreeImporter/SpeedTreeImporter.Build.cs
@@ -6,22 +6,6 @@ namespace UnrealBuildTool.Rules
 	{
 		public SpeedTreeImporter(ReadOnlyTargetRules Target) : base(Target)
 		{
-			PublicIncludePaths.AddRange(
-				new string[] {
-				}
-				);
-
-			PrivateIncludePaths.AddRange(
-				new string[] {
-				}
-				);
-
-			PublicDependencyModuleNames.AddRange(
-				new string[]
-				{
-				}
-				);
-
 			PrivateDependencyModuleNames.AddRange(
 				new string[]
 				{
@@ -38,12 +22,6 @@ namespace UnrealBuildTool.Rules
 					"StaticMeshDescription"
                 }
 				);
-
-			DynamicallyLoadedModuleNames.AddRange(
-				new string[]
-				{
-				}
-				);
 				
 			AddEngineThirdPartyPrivateStaticDependencies(Target, "SpeedTree");
 		}
diff --git a/Engine/Plugins/Editor/StylusInput/Source/StylusInput/StylusInput.Build.cs b/Engine/Plugins/Editor/StylusInput/Source/StylusInput/StylusInput.Build.cs
--- a/Engine/Plugins/Editor/StylusInput/Source/StylusInput/StylusInput.Build.cs
+++ b/Engine/Plugins/Editor/StylusInput/Source/StylusInput/StylusInput.Build.cs
@@ -6,18 +6,6 @@ namespace UnrealBuildTool.Rules
 	{
 		public StylusInput(ReadOnlyTargetRules Target) : base(Target)
 		{
-			PublicIncludePaths.AddRange(
-				new string[] {
-					// ... add public include paths required here ...
-				}
-				);
-
-			PrivateIncludePaths.AddRange(
-				new string[] {
-					// ... add other private include paths required here ...
-				}
-				);
-
 			PublicDependencyModuleNames.AddRange(
 				new string[]
 				{
@@ -41,13 +29,6 @@ namespace UnrealBuildTool.Rules
 					// ... add private dependencies that you statically link with here ...
 				}
 				);
-
-			DynamicallyLoadedModuleNames.AddRange(
-				new string[]
-				{
-					// ... add any modules that your module loads dynamically here ...
-				}
-				);
 		}
 	}
 }
diff --git a/Engine/Plugins/Experimental/ActorPalette/Source/ActorPalette/ActorPalette.Build.cs b/Engine/Plugins/Experimental/ActorPalette/Source/ActorPalette/ActorPalette.Build.cs
--- a/Engine/Plugins/Experimental/ActorPalette/Source/ActorPalette/ActorPalette.Build.cs
+++ b/Engine/Plugins/Experimental/ActorPalette/Source/ActorPalette/ActorPalette.Build.cs
@@ -6,29 +6,14 @@ public class ActorPalette : ModuleRules
 {
 	public ActorPalette(ReadOnlyTargetRules Target) : base(Target)
 	{
-		PCHUsage = ModuleRules.PCHUsageMode.UseExplicitOrSharedPCHs;
-		
-		PublicIncludePaths.AddRange(
-			new string[] {
-				// ... add public include paths required here ...
-			}
-			);
-				
-		
-		PrivateIncludePaths.AddRange(
-			new string[] {
-				// ... add other private include paths required here ...
-			}
-			);
-			
+		PCHUsage = ModuleRules.PCHUsageMode.UseExplicitOrSharedPCHs;			
 		
 		PublicDependencyModuleNames.AddRange(
 			new string[]
 			{
 				"Core",
 			}
-			);
-			
+			);			
 		
 		PrivateDependencyModuleNames.AddRange(
 			new string[]
diff --git a/Engine/Plugins/Experimental/Animation/ContextualAnimation/Source/ContextualAnimation/ContextualAnimation.Build.cs b/Engine/Plugins/Experimental/Animation/ContextualAnimation/Source/ContextualAnimation/ContextualAnimation.Build.cs
--- a/Engine/Plugins/Experimental/Animation/ContextualAnimation/Source/ContextualAnimation/ContextualAnimation.Build.cs
+++ b/Engine/Plugins/Experimental/Animation/ContextualAnimation/Source/ContextualAnimation/ContextualAnimation.Build.cs
@@ -6,21 +6,7 @@ public class ContextualAnimation : ModuleRules
 {
 	public ContextualAnimation(ReadOnlyTargetRules Target) : base(Target)
 	{
-		PCHUsage = ModuleRules.PCHUsageMode.UseExplicitOrSharedPCHs;
-		
-		PublicIncludePaths.AddRange(
-			new string[] {
-				// ... add public include paths required here ...
-			}
-			);
-				
-		
-		PrivateIncludePaths.AddRange(
-			new string[] {
-				// ... add other private include paths required here ...
-			}
-			);
-			
+		PCHUsage = ModuleRules.PCHUsageMode.UseExplicitOrSharedPCHs;			
 		
 		PublicDependencyModuleNames.AddRange(
 			new string[]
@@ -32,8 +18,7 @@ public class ContextualAnimation : ModuleRules
 				"MotionWarping",
 				"IKRig"
 			}
-			);
-			
+			);			
 		
 		PrivateDependencyModuleNames.AddRange(
 			new string[]
@@ -42,13 +27,5 @@ public class ContextualAnimation : ModuleRules
 				"SlateCore"
             }
 			);
-		
-		
-		DynamicallyLoadedModuleNames.AddRange(
-			new string[]
-			{
-				// ... add any modules that your module loads dynamically here ...
-			}
-			);
 	}
 }
diff --git a/Engine/Plugins/Experimental/Animation/ContextualAnimation/Source/ContextualAnimationEditor/ContextualAnimationEditor.Build.cs b/Engine/Plugins/Experimental/Animation/ContextualAnimation/Source/ContextualAnimationEditor/ContextualAnimationEditor.Build.cs
--- a/Engine/Plugins/Experimental/Animation/ContextualAnimation/Source/ContextualAnimationEditor/ContextualAnimationEditor.Build.cs
+++ b/Engine/Plugins/Experimental/Animation/ContextualAnimation/Source/ContextualAnimationEditor/ContextualAnimationEditor.Build.cs
@@ -8,18 +8,6 @@ public class ContextualAnimationEditor : ModuleRules
 	{
 		PCHUsage = ModuleRules.PCHUsageMode.UseExplicitOrSharedPCHs;
 
-		PublicIncludePaths.AddRange(
-			new string[] {
-			}
-			);
-				
-		
-		PrivateIncludePaths.AddRange(
-			new string[] {
-			}
-			);
-			
-		
 		PublicDependencyModuleNames.AddRange(
 			new string[]
 			{
@@ -30,8 +18,7 @@ public class ContextualAnimationEditor : ModuleRules
 				"MotionWarping",
 				"GameplayTags"
 			}
-			);
-			
+			);			
 		
 		PrivateDependencyModuleNames.AddRange(
 			new string[]
@@ -65,13 +52,5 @@ public class ContextualAnimationEditor : ModuleRules
 				"AnimGraph"
 			}
 			);
-		
-		
-		DynamicallyLoadedModuleNames.AddRange(
-			new string[]
-			{
-				// ... add any modules that your module loads dynamically here ...
-			}
-			);
 	}
 }
diff --git a/Engine/Plugins/Experimental/Animation/PoseSearch/Source/Runtime/PoseSearch.Build.cs b/Engine/Plugins/Experimental/Animation/PoseSearch/Source/Runtime/PoseSearch.Build.cs
--- a/Engine/Plugins/Experimental/Animation/PoseSearch/Source/Runtime/PoseSearch.Build.cs
+++ b/Engine/Plugins/Experimental/Animation/PoseSearch/Source/Runtime/PoseSearch.Build.cs
@@ -13,20 +13,6 @@ public class PoseSearch : ModuleRules
 			"nanoflann"
 		);
 
-		PublicIncludePaths.AddRange(
-			new string[]
-			{
-				System.IO.Path.Combine(ModuleDirectory, "Public"),
-			}
-		);
-
-		PrivateIncludePaths.AddRange(
-			new string[]
-			{
-				System.IO.Path.Combine(ModuleDirectory, "Private"),
-			}
-		);
-
 		PublicDependencyModuleNames.AddRange(
 			new string[]
 			{
diff --git a/Engine/Plugins/Experimental/AppleVision/Source/AppleVision/AppleVision.Build.cs b/Engine/Plugins/Experimental/AppleVision/Source/AppleVision/AppleVision.Build.cs
--- a/Engine/Plugins/Experimental/AppleVision/Source/AppleVision/AppleVision.Build.cs
+++ b/Engine/Plugins/Experimental/AppleVision/Source/AppleVision/AppleVision.Build.cs
@@ -6,12 +6,6 @@ namespace UnrealBuildTool.Rules
 	{
 		public AppleVision(ReadOnlyTargetRules Target) : base(Target)
 		{
-            PublicIncludePaths.AddRange(
-                new string[] {
-                    // ... add public include paths required here ...
-                }
-                );
-
 
             PrivateIncludePaths.AddRange(
                 new string[] {
@@ -51,13 +45,6 @@ namespace UnrealBuildTool.Rules
 					// ... add private dependencies that you statically link with here ...
 				}
 				);
-
-			DynamicallyLoadedModuleNames.AddRange(
-				new string[]
-				{
-					// ... add any modules that your module loads dynamically here ...
-				}
-				);
 		}
 	}
 }
diff --git a/Engine/Plugins/Experimental/AppleVision/Source/AppleVisionBlueprintSupport/AppleVisionBlueprintSupport.Build.cs b/Engine/Plugins/Experimental/AppleVision/Source/AppleVisionBlueprintSupport/AppleVisionBlueprintSupport.Build.cs
--- a/Engine/Plugins/Experimental/AppleVision/Source/AppleVisionBlueprintSupport/AppleVisionBlueprintSupport.Build.cs
+++ b/Engine/Plugins/Experimental/AppleVision/Source/AppleVisionBlueprintSupport/AppleVisionBlueprintSupport.Build.cs
@@ -6,12 +6,6 @@ namespace UnrealBuildTool.Rules
 	{
 		public AppleVisionBlueprintSupport(ReadOnlyTargetRules Target) : base(Target)
 		{
-            PublicIncludePaths.AddRange(
-                new string[] {
-                    // ... add public include paths required here ...
-                }
-                );
-
 
             PrivateIncludePaths.AddRange(
                 new string[] {
@@ -31,13 +25,6 @@ namespace UnrealBuildTool.Rules
                     "AppleVision"
                     // ... add other public dependencies that you statically link with here ...
                 }
-                );
-
-            PublicFrameworks.AddRange(
-                new string[]
-                {
-                    // ... add other public dependencies that you statically link with here ...
-                }
                 );
 
 			PrivateDependencyModuleNames.AddRange(
@@ -47,13 +34,6 @@ namespace UnrealBuildTool.Rules
 					// ... add private dependencies that you statically link with here ...
 				}
 				);
-
-			DynamicallyLoadedModuleNames.AddRange(
-				new string[]
-				{
-					// ... add any modules that your module loads dynamically here ...
-				}
-				);
 		}
 	}
 }
diff --git a/Engine/Source/Editor/EditorSettingsViewer/EditorSettingsViewer.Build.cs b/Engine/Source/Editor/EditorSettingsViewer/EditorSettingsViewer.Build.cs
--- a/Engine/Source/Editor/EditorSettingsViewer/EditorSettingsViewer.Build.cs
+++ b/Engine/Source/Editor/EditorSettingsViewer/EditorSettingsViewer.Build.cs
@@ -38,12 +38,6 @@ namespace UnrealBuildTool.Rules
 					"Settings",
 				}
 			);
-
-			PrivateIncludePaths.AddRange(
-				new string[] {
-					"Editor/EditorSettingsViewer/Private",
-				}
-			);
 		}
 	}
 }
diff --git a/Engine/Source/Editor/Experimental/EditorInteractiveToolsFramework/EditorInteractiveToolsFramework.Build.cs b/Engine/Source/Editor/Experimental/EditorInteractiveToolsFramework/EditorInteractiveToolsFramework.Build.cs
--- a/Engine/Source/Editor/Experimental/EditorInteractiveToolsFramework/EditorInteractiveToolsFramework.Build.cs
+++ b/Engine/Source/Editor/Experimental/EditorInteractiveToolsFramework/EditorInteractiveToolsFramework.Build.cs
@@ -6,22 +6,7 @@ public class EditorInteractiveToolsFramework : ModuleRules
 {
 	public EditorInteractiveToolsFramework(ReadOnlyTargetRules Target) : base(Target)
 	{
-		PCHUsage = ModuleRules.PCHUsageMode.UseExplicitOrSharedPCHs;
-		
-		PublicIncludePaths.AddRange(
-			new string[] {
-                //"ContentBrowser"
-				// ... add public include paths required here ...
-			}
-			);
-				
-		
-		PrivateIncludePaths.AddRange(
-			new string[] {
-				// ... add other private include paths required here ...
-			}
-			);
-			
+		PCHUsage = ModuleRules.PCHUsageMode.UseExplicitOrSharedPCHs;			
 		
 		PublicDependencyModuleNames.AddRange(
 			new string[]
@@ -30,8 +15,7 @@ public class EditorInteractiveToolsFramework : ModuleRules
 				"TypedElementFramework",
 				// ... add other public dependencies that you statically link with here ...
 			}
-			);
-			
+			);			
 		
 		PrivateDependencyModuleNames.AddRange(
 			new string[]
@@ -57,13 +41,5 @@ public class EditorInteractiveToolsFramework : ModuleRules
 				// ... add private dependencies that you statically link with here ...	
 			}
             );
-		
-		
-		DynamicallyLoadedModuleNames.AddRange(
-			new string[]
-			{
-				// ... add any modules that your module loads dynamically here ...
-			}
-			);
 	}
 }
diff --git a/Engine/Source/Editor/Kismet/Kismet.Build.cs b/Engine/Source/Editor/Kismet/Kismet.Build.cs
--- a/Engine/Source/Editor/Kismet/Kismet.Build.cs
+++ b/Engine/Source/Editor/Kismet/Kismet.Build.cs
@@ -6,8 +6,6 @@ public class Kismet : ModuleRules
 {
 	public Kismet(ReadOnlyTargetRules Target) : base(Target)
 	{
-		PrivateIncludePaths.Add("Editor/Kismet/Private");
-
 		PrivateIncludePathModuleNames.AddRange(
 			new string[] { 
 				"AssetRegistry", 
diff --git a/Engine/Source/Editor/KismetWidgets/KismetWidgets.Build.cs b/Engine/Source/Editor/KismetWidgets/KismetWidgets.Build.cs
--- a/Engine/Source/Editor/KismetWidgets/KismetWidgets.Build.cs
+++ b/Engine/Source/Editor/KismetWidgets/KismetWidgets.Build.cs
@@ -6,8 +6,6 @@ public class KismetWidgets : ModuleRules
 {
 	public KismetWidgets(ReadOnlyTargetRules Target) : base(Target)
 	{
-		PrivateIncludePaths.Add("Editor/KismetWidgets/Private");
-
 		PrivateDependencyModuleNames.AddRange(
 			new string[] {
 				"Core", 
diff --git a/Engine/Source/Editor/LevelAssetEditor/LevelAssetEditor.Build.cs b/Engine/Source/Editor/LevelAssetEditor/LevelAssetEditor.Build.cs
--- a/Engine/Source/Editor/LevelAssetEditor/LevelAssetEditor.Build.cs
+++ b/Engine/Source/Editor/LevelAssetEditor/LevelAssetEditor.Build.cs
@@ -8,20 +8,6 @@ public class LevelAssetEditor : ModuleRules
 	{
 		PCHUsage = ModuleRules.PCHUsageMode.UseExplicitOrSharedPCHs;
 		
-		PublicIncludePaths.AddRange(
-			new string[] {
-				// ... add public include paths required here ...
-			}
-			);
-				
-		
-		PrivateIncludePaths.AddRange(
-			new string[] {
-				// ... add other private include paths required here ...
-			}
-			);
-			
-		
 		PublicDependencyModuleNames.AddRange(
 			new string[]
 			{
@@ -32,8 +18,7 @@ public class LevelAssetEditor : ModuleRules
 				"TypedElementFramework",
 				// ... add other public dependencies that you statically link with here ...
 			}
-			);
-			
+			);			
 		
 		PrivateDependencyModuleNames.AddRange(
 			new string[]
@@ -50,13 +35,5 @@ public class LevelAssetEditor : ModuleRules
 				// ... add private dependencies that you statically link with here ...	
 			}
 			);
-		
-		
-		DynamicallyLoadedModuleNames.AddRange(
-			new string[]
-			{
-				// ... add any modules that your module loads dynamically here ...
-			}
-			);
 	}
 }
diff --git a/Engine/Source/Editor/LevelInstanceEditor/LevelInstanceEditor.Build.cs b/Engine/Source/Editor/LevelInstanceEditor/LevelInstanceEditor.Build.cs
--- a/Engine/Source/Editor/LevelInstanceEditor/LevelInstanceEditor.Build.cs
+++ b/Engine/Source/Editor/LevelInstanceEditor/LevelInstanceEditor.Build.cs
@@ -6,8 +6,6 @@ public class LevelInstanceEditor : ModuleRules
 {
     public LevelInstanceEditor(ReadOnlyTargetRules Target) : base(Target)
     {
-        PrivateIncludePaths.Add("Editor/LevelInstanceEditor/Private");	// For PCH includes (because they don't work with relative paths, yet)
-
         PrivateIncludePathModuleNames.AddRange(
             new string[] {
             }
diff --git a/Engine/Source/Editor/LocalizationDashboard/LocalizationDashboard.Build.cs b/Engine/Source/Editor/LocalizationDashboard/LocalizationDashboard.Build.cs
--- a/Engine/Source/Editor/LocalizationDashboard/LocalizationDashboard.Build.cs
+++ b/Engine/Source/Editor/LocalizationDashboard/LocalizationDashboard.Build.cs
@@ -39,20 +39,6 @@ public class LocalizationDashboard : ModuleRules
 			}
 		);
 
-		PrivateIncludePaths.AddRange(
-			new string[]
-			{
-				"Editor/LocalizationDashboard/Private",
-			}
-		);
-
-        PublicIncludePaths.AddRange(
-			new string[]
-			{
-				"Editor/LocalizationDashboard/Public",
-			}
-		);
-
         CircularlyReferencedDependentModules.AddRange(
            new string[] {
                 "LocalizationService",
diff --git a/Engine/Source/Editor/MaterialEditor/MaterialEditor.Build.cs b/Engine/Source/Editor/MaterialEditor/MaterialEditor.Build.cs
--- a/Engine/Source/Editor/MaterialEditor/MaterialEditor.Build.cs
+++ b/Engine/Source/Editor/MaterialEditor/MaterialEditor.Build.cs
@@ -6,12 +6,6 @@ public class MaterialEditor : ModuleRules
 {
 	public MaterialEditor(ReadOnlyTargetRules Target) : base(Target)
 	{
-		PrivateIncludePaths.AddRange(
-			new string[] {
-				"Editor/MaterialEditor/Private"
-			}
-		);
-
 		PrivateIncludePathModuleNames.AddRange(
 			new string[] 
 			{
diff --git a/Engine/Source/Editor/MovieSceneCaptureDialog/MovieSceneCaptureDialog.Build.cs b/Engine/Source/Editor/MovieSceneCaptureDialog/MovieSceneCaptureDialog.Build.cs
--- a/Engine/Source/Editor/MovieSceneCaptureDialog/MovieSceneCaptureDialog.Build.cs
+++ b/Engine/Source/Editor/MovieSceneCaptureDialog/MovieSceneCaptureDialog.Build.cs
@@ -6,12 +6,6 @@ public class MovieSceneCaptureDialog : ModuleRules
 {
 	public MovieSceneCaptureDialog(ReadOnlyTargetRules Target) : base(Target)
 	{
-		PrivateIncludePaths.AddRange(
-			new string[] {
-				"Editor/MovieSceneCaptureDialog/Private"
-			}
-		);
-
 		PrivateDependencyModuleNames.AddRange(
 			new string[] {
 				"Core",
diff --git a/Engine/Source/Editor/OverlayEditor/OverlayEditor.Build.cs b/Engine/Source/Editor/OverlayEditor/OverlayEditor.Build.cs
--- a/Engine/Source/Editor/OverlayEditor/OverlayEditor.Build.cs
+++ b/Engine/Source/Editor/OverlayEditor/OverlayEditor.Build.cs
@@ -27,20 +27,6 @@ namespace UnrealBuildTool.Rules
 					"AssetTools",
 				}
 			);
-			
-            PublicIncludePaths.AddRange(
-                new string[]
-                {
-                    "Editor/OverlayEditor/Public",
-                }
-            );
-
-			PrivateIncludePaths.AddRange(
-				new string[] {
-					"Editor/OverlayEditor/Private",
-					"Editor/OverlayEditor/Private/Factories",
-				}
-			);
 
             PrivateIncludePathModuleNames.AddRange(
                 new string[] {
diff --git a/Engine/Source/Editor/Persona/Persona.Build.cs b/Engine/Source/Editor/Persona/Persona.Build.cs
--- a/Engine/Source/Editor/Persona/Persona.Build.cs
+++ b/Engine/Source/Editor/Persona/Persona.Build.cs
@@ -6,7 +6,6 @@ public class Persona : ModuleRules
 {
     public Persona(ReadOnlyTargetRules Target) : base(Target)
     {
-        PrivateIncludePaths.Add("Editor/Persona/Private");  // For PCH includes (because they don't work with relative paths, yet)
         PrivateIncludePaths.Add("Editor/Persona/Private/AnimTimeline");
 
         PublicIncludePathModuleNames.AddRange(
diff --git a/Engine/Source/Editor/PhysicsAssetEditor/PhysicsAssetEditor.Build.cs b/Engine/Source/Editor/PhysicsAssetEditor/PhysicsAssetEditor.Build.cs
--- a/Engine/Source/Editor/PhysicsAssetEditor/PhysicsAssetEditor.Build.cs
+++ b/Engine/Source/Editor/PhysicsAssetEditor/PhysicsAssetEditor.Build.cs
@@ -15,7 +15,6 @@ public class PhysicsAssetEditor : ModuleRules
 
         PrivateIncludePaths.AddRange(
             new string[] {
-                "Editor/PhysicsAssetEditor/Private",
                 "Editor/PhysicsAssetEditor/Private/PhysicsAssetGraph",
             }
         );
diff --git a/Engine/Source/Editor/PinnedCommandList/PinnedCommandList.Build.cs b/Engine/Source/Editor/PinnedCommandList/PinnedCommandList.Build.cs
--- a/Engine/Source/Editor/PinnedCommandList/PinnedCommandList.Build.cs
+++ b/Engine/Source/Editor/PinnedCommandList/PinnedCommandList.Build.cs
@@ -6,23 +6,6 @@ public class PinnedCommandList : ModuleRules
 {
     public PinnedCommandList(ReadOnlyTargetRules Target) : base(Target)
     {
-        PrivateIncludePaths.Add("Editor/PinnedCommandList/Private");
-
-        PublicIncludePathModuleNames.AddRange(
-            new string[] {
-            }
-        );
-
-        PublicDependencyModuleNames.AddRange(
-            new string[] {
-            }
-        );
-
-        PrivateIncludePathModuleNames.AddRange(
-            new string[] {
-            }
-        );
-
         PrivateDependencyModuleNames.AddRange(
             new string[] {
                 "Core",
@@ -33,10 +16,5 @@ public class PinnedCommandList : ModuleRules
                 "InputCore",
             }
         );
-
-        DynamicallyLoadedModuleNames.AddRange(
-            new string[] {
-            }
-        );
     }
 }
diff --git a/Engine/Source/Editor/ProjectSettingsViewer/ProjectSettingsViewer.Build.cs b/Engine/Source/Editor/ProjectSettingsViewer/ProjectSettingsViewer.Build.cs
--- a/Engine/Source/Editor/ProjectSettingsViewer/ProjectSettingsViewer.Build.cs
+++ b/Engine/Source/Editor/ProjectSettingsViewer/ProjectSettingsViewer.Build.cs
@@ -31,13 +31,6 @@ namespace UnrealBuildTool.Rules
 				}
 			);
 
-			PrivateIncludePaths.AddRange(
-				new string[]
-				{
-					"Editor/ProjectSettingsViewer/Private",
-				}
-			);
-
 			if (Target.bBuildTargetDeveloperTools)
 			{
 				PrivateDependencyModuleNames.Add("ProjectTargetPlatformEditor");
diff --git a/Engine/Source/Editor/ProjectTargetPlatformEditor/ProjectTargetPlatformEditor.Build.cs b/Engine/Source/Editor/ProjectTargetPlatformEditor/ProjectTargetPlatformEditor.Build.cs
--- a/Engine/Source/Editor/ProjectTargetPlatformEditor/ProjectTargetPlatformEditor.Build.cs
+++ b/Engine/Source/Editor/ProjectTargetPlatformEditor/ProjectTargetPlatformEditor.Build.cs
@@ -44,7 +44,6 @@ public class ProjectTargetPlatformEditor : ModuleRules
 		PrivateIncludePaths.AddRange(
 			new string[]
 			{
-				"Editor/ProjectTargetPlatformEditor/Private",
 				"Editor/ProjectTargetPlatformEditor/Private/Widgets",
 			}
 		);
diff --git a/Engine/Source/Editor/PropertyEditor/PropertyEditor.Build.cs b/Engine/Source/Editor/PropertyEditor/PropertyEditor.Build.cs
--- a/Engine/Source/Editor/PropertyEditor/PropertyEditor.Build.cs
+++ b/Engine/Source/Editor/PropertyEditor/PropertyEditor.Build.cs
@@ -23,7 +23,6 @@ public class PropertyEditor : ModuleRules
 
 		PrivateIncludePaths.AddRange(
 			new string[] {
-				"Editor/PropertyEditor/Private",
 				"Editor/PropertyEditor/Private/Presentation",
 				"Editor/PropertyEditor/Private/Presentation/PropertyTable",
 				"Editor/PropertyEditor/Private/Presentation/PropertyEditor",
diff --git a/Engine/Source/Editor/SequenceRecorder/SequenceRecorder.Build.cs b/Engine/Source/Editor/SequenceRecorder/SequenceRecorder.Build.cs
--- a/Engine/Source/Editor/SequenceRecorder/SequenceRecorder.Build.cs
+++ b/Engine/Source/Editor/SequenceRecorder/SequenceRecorder.Build.cs
@@ -6,14 +6,8 @@ namespace UnrealBuildTool.Rules
 	{
 		public SequenceRecorder(ReadOnlyTargetRules Target) : base(Target)
 		{
-			PublicIncludePaths.AddRange(
-				new string[] {
-				}
-				);
-
 			PrivateIncludePaths.AddRange(
 				new string[] {
-					"Editor/SequenceRecorder/Private",
 					"Editor/SequenceRecorder/Private/Sections",
 				}
 				);
diff --git a/Engine/Source/Editor/SequenceRecorderSections/SequenceRecorderSections.Build.cs b/Engine/Source/Editor/SequenceRecorderSections/SequenceRecorderSections.Build.cs
--- a/Engine/Source/Editor/SequenceRecorderSections/SequenceRecorderSections.Build.cs
+++ b/Engine/Source/Editor/SequenceRecorderSections/SequenceRecorderSections.Build.cs
@@ -6,24 +6,6 @@ namespace UnrealBuildTool.Rules
 	{
 		public SequenceRecorderSections(ReadOnlyTargetRules Target) : base(Target)
 		{
-			PublicIncludePaths.AddRange(
-				new string[] {
-					"Editor/SequenceRecorderSections/Public",
-				}
-				);
-
-			PrivateIncludePaths.AddRange(
-				new string[] {
-					"Editor/SequenceRecorderSections/Private",
-				}
-				);
-
-			PublicDependencyModuleNames.AddRange(
-				new string[]
-				{
-				}
-				);
-
 			PrivateDependencyModuleNames.AddRange(
 				new string[]
 				{
@@ -36,12 +18,6 @@ namespace UnrealBuildTool.Rules
 					"TimeManagement"
 				}
 				);
-
-			DynamicallyLoadedModuleNames.AddRange(
-				new string[]
-				{
-				}
-				);
 		}
 	}
 }
diff --git a/Engine/Source/Editor/SequencerCore/SequencerCore.Build.cs b/Engine/Source/Editor/SequencerCore/SequencerCore.Build.cs
--- a/Engine/Source/Editor/SequencerCore/SequencerCore.Build.cs
+++ b/Engine/Source/Editor/SequencerCore/SequencerCore.Build.cs
@@ -6,12 +6,6 @@ public class SequencerCore : ModuleRules
 {
 	public SequencerCore(ReadOnlyTargetRules Target) : base(Target)
 	{
-		PrivateIncludePaths.AddRange(
-			new string[] {
-				"Editor/SequencerCore/Private",
-			}
-		);
-
 		PublicDependencyModuleNames.AddRange(
 			new string[] {
 				"Core",
diff --git a/Engine/Source/Editor/SequencerWidgets/SequencerWidgets.Build.cs b/Engine/Source/Editor/SequencerWidgets/SequencerWidgets.Build.cs
--- a/Engine/Source/Editor/SequencerWidgets/SequencerWidgets.Build.cs
+++ b/Engine/Source/Editor/SequencerWidgets/SequencerWidgets.Build.cs
@@ -6,8 +6,6 @@ public class SequencerWidgets : ModuleRules
 {
 	public SequencerWidgets(ReadOnlyTargetRules Target) : base(Target)
 	{
-		PrivateIncludePaths.Add("Editor/SequencerWidgets/Private");
-
         PrivateIncludePathModuleNames.AddRange(
             new string[] {
                 "MovieScene",
diff --git a/Engine/Source/Editor/SourceControlWindowExtender/SourceControlWindowExtender.Build.cs b/Engine/Source/Editor/SourceControlWindowExtender/SourceControlWindowExtender.Build.cs
--- a/Engine/Source/Editor/SourceControlWindowExtender/SourceControlWindowExtender.Build.cs
+++ b/Engine/Source/Editor/SourceControlWindowExtender/SourceControlWindowExtender.Build.cs
@@ -6,10 +6,6 @@ public class SourceControlWindowExtender : ModuleRules
 {
 	public SourceControlWindowExtender(ReadOnlyTargetRules Target) : base(Target)
 	{
-		PublicIncludePaths.Add("Editor/SourceControlWindowExtender/Public");
-
-		PrivateIncludePaths.Add("Editor/SourceControlWindowExtender/Private");
-
 		PrivateDependencyModuleNames.AddRange(
 			new string[] {
 				"Core", 
diff --git a/Engine/Source/Editor/SourceControlWindows/SourceControlWindows.Build.cs b/Engine/Source/Editor/SourceControlWindows/SourceControlWindows.Build.cs
--- a/Engine/Source/Editor/SourceControlWindows/SourceControlWindows.Build.cs
+++ b/Engine/Source/Editor/SourceControlWindows/SourceControlWindows.Build.cs
@@ -6,10 +6,6 @@ public class SourceControlWindows : ModuleRules
 {
 	public SourceControlWindows(ReadOnlyTargetRules Target) : base(Target)
 	{
-		PublicIncludePaths.Add("Editor/SourceControlWindows/Public");
-
-		PrivateIncludePaths.Add("Editor/SourceControlWindows/Private");
-
 		PrivateDependencyModuleNames.AddRange(
 			new string[] {
 				"Core", 
diff --git a/Engine/Source/Editor/SubobjectDataInterface/SubobjectDataInterface.Build.cs b/Engine/Source/Editor/SubobjectDataInterface/SubobjectDataInterface.Build.cs
--- a/Engine/Source/Editor/SubobjectDataInterface/SubobjectDataInterface.Build.cs
+++ b/Engine/Source/Editor/SubobjectDataInterface/SubobjectDataInterface.Build.cs
@@ -7,9 +7,6 @@ public class SubobjectDataInterface : ModuleRules
 {
 	public SubobjectDataInterface(ReadOnlyTargetRules Target) : base(Target)
 	{
-		PublicIncludePaths.Add(Path.Combine(ModuleDirectory, "Public"));
-		PrivateIncludePaths.Add(Path.Combine(ModuleDirectory, "Private"));
-
 		PublicDependencyModuleNames.AddRange(
 			new string[]
 			{
diff --git a/Engine/Source/Editor/SubobjectEditor/SubobjectEditor.Build.cs b/Engine/Source/Editor/SubobjectEditor/SubobjectEditor.Build.cs
--- a/Engine/Source/Editor/SubobjectEditor/SubobjectEditor.Build.cs
+++ b/Engine/Source/Editor/SubobjectEditor/SubobjectEditor.Build.cs
@@ -7,9 +7,6 @@ public class SubobjectEditor : ModuleRules
 {
 	public SubobjectEditor(ReadOnlyTargetRules Target) : base(Target)
 	{
-		PublicIncludePaths.Add(Path.Combine(ModuleDirectory, "Public"));
-		PrivateIncludePaths.Add(Path.Combine(ModuleDirectory, "Private"));
-
 		PublicDependencyModuleNames.AddRange(
 			new string[]
 			{
diff --git a/Engine/Source/Editor/TranslationEditor/TranslationEditor.Build.cs b/Engine/Source/Editor/TranslationEditor/TranslationEditor.Build.cs
--- a/Engine/Source/Editor/TranslationEditor/TranslationEditor.Build.cs
+++ b/Engine/Source/Editor/TranslationEditor/TranslationEditor.Build.cs
@@ -55,19 +55,5 @@ public class TranslationEditor : ModuleRules
 				"DesktopPlatform",
 			}
 		);
-
-        PrivateIncludePaths.AddRange(
-            new string[]
-			{
-				"Editor/TranslationEditor/Private",
-			}
-        );
-
-        PublicIncludePaths.AddRange(
-            new string[]
-			{
-				"Editor/TranslationEditor/Public",
-			}
-        );
 	}
 }
diff --git a/Engine/Source/Editor/TurnkeySupport/TurnkeySupport.Build.cs b/Engine/Source/Editor/TurnkeySupport/TurnkeySupport.Build.cs
--- a/Engine/Source/Editor/TurnkeySupport/TurnkeySupport.Build.cs
+++ b/Engine/Source/Editor/TurnkeySupport/TurnkeySupport.Build.cs
@@ -6,7 +6,6 @@ public class TurnkeySupport : ModuleRules
 {
 	public TurnkeySupport(ReadOnlyTargetRules Target) : base(Target)
 	{
-
 		PrivateDependencyModuleNames.AddRange(
 			new string[] {
 				"Core",
@@ -56,12 +55,5 @@ public class TurnkeySupport : ModuleRules
 				}
 			);
 		}
-
-
-		PrivateIncludePaths.AddRange(
-			new string[] {
-				"Editor/TurnkeySupport/Private",
-			}
-		);
 	}
 }
diff --git a/Engine/Source/Editor/UnrealEd/UnrealEd.Build.cs b/Engine/Source/Editor/UnrealEd/UnrealEd.Build.cs
--- a/Engine/Source/Editor/UnrealEd/UnrealEd.Build.cs
+++ b/Engine/Source/Editor/UnrealEd/UnrealEd.Build.cs
@@ -19,7 +19,6 @@ public class UnrealEd : ModuleRules
 		PrivateIncludePaths.AddRange(
 			new string[]
 			{
-				"Editor/UnrealEd/Private",
 				"Editor/UnrealEd/Private/Settings",
 				"Editor/PackagesDialog/Public",
 				"Developer/TargetPlatform/Public",
diff --git a/Engine/Source/Editor/UnrealEdMessages/UnrealEdMessages.Build.cs b/Engine/Source/Editor/UnrealEdMessages/UnrealEdMessages.Build.cs
--- a/Engine/Source/Editor/UnrealEdMessages/UnrealEdMessages.Build.cs
+++ b/Engine/Source/Editor/UnrealEdMessages/UnrealEdMessages.Build.cs
@@ -17,12 +17,6 @@ namespace UnrealBuildTool.Rules
 				{
 					"Core",
 				});
-
-			PrivateIncludePaths.AddRange(
-				new string[]
-				{
-					"Editor/UnrealEdMessages/Private",
-				});
 		}
 	}
 }
diff --git a/Engine/Source/Editor/VirtualTexturingEditor/VirtualTexturingEditor.Build.cs b/Engine/Source/Editor/VirtualTexturingEditor/VirtualTexturingEditor.Build.cs
--- a/Engine/Source/Editor/VirtualTexturingEditor/VirtualTexturingEditor.Build.cs
+++ b/Engine/Source/Editor/VirtualTexturingEditor/VirtualTexturingEditor.Build.cs
@@ -6,8 +6,6 @@ public class VirtualTexturingEditor : ModuleRules
 {
 	public VirtualTexturingEditor(ReadOnlyTargetRules Target) : base(Target)
 	{
-        PrivateIncludePaths.Add("Editor/VirtualTexturingEditor/Private");
-
         PrivateDependencyModuleNames.AddRange(
             new string[] {
 				"AppFramework",
diff --git a/Engine/Source/Editor/WorkspaceMenuStructure/WorkspaceMenuStructure.Build.cs b/Engine/Source/Editor/WorkspaceMenuStructure/WorkspaceMenuStructure.Build.cs
--- a/Engine/Source/Editor/WorkspaceMenuStructure/WorkspaceMenuStructure.Build.cs
+++ b/Engine/Source/Editor/WorkspaceMenuStructure/WorkspaceMenuStructure.Build.cs
@@ -6,8 +6,6 @@ public class WorkspaceMenuStructure : ModuleRules
 {
 	public WorkspaceMenuStructure(ReadOnlyTargetRules Target) : base(Target)
 	{
-		PrivateIncludePaths.Add("Editor/WorkspaceMenuStructure/Private");
-
 		PrivateDependencyModuleNames.AddRange(
 			new string[] {
 				"Core",
diff --git a/Engine/Source/Editor/WorldBrowser/WorldBrowser.Build.cs b/Engine/Source/Editor/WorldBrowser/WorldBrowser.Build.cs
--- a/Engine/Source/Editor/WorldBrowser/WorldBrowser.Build.cs
+++ b/Engine/Source/Editor/WorldBrowser/WorldBrowser.Build.cs
@@ -6,10 +6,6 @@ public class WorldBrowser : ModuleRules
 {
     public WorldBrowser(ReadOnlyTargetRules Target) : base(Target)
     {
-        PublicIncludePaths.Add("Editor/WorldBrowser/Public");
-
-        PrivateIncludePaths.Add("Editor/WorldBrowser/Private");	// For PCH includes (because they don't work with relative paths, yet)
-
         PrivateIncludePathModuleNames.AddRange(
             new string[] {
                 "AssetRegistry",
diff --git a/Engine/Source/Editor/WorldPartitionEditor/WorldPartitionEditor.Build.cs b/Engine/Source/Editor/WorldPartitionEditor/WorldPartitionEditor.Build.cs
--- a/Engine/Source/Editor/WorldPartitionEditor/WorldPartitionEditor.Build.cs
+++ b/Engine/Source/Editor/WorldPartitionEditor/WorldPartitionEditor.Build.cs
@@ -5,16 +5,7 @@ using UnrealBuildTool;
 public class WorldPartitionEditor : ModuleRules
 {
     public WorldPartitionEditor(ReadOnlyTargetRules Target) : base(Target)
-    {
-        PublicIncludePaths.Add("Editor/WorldPartitionEditor/Public");
-
-        PrivateIncludePaths.Add("Editor/WorldPartitionEditor/Private");	// For PCH includes (because they don't work with relative paths, yet)
-
-        PrivateIncludePathModuleNames.AddRange(
-            new string[] {
-			}
-		);
-     
+    {     
         PrivateDependencyModuleNames.AddRange(
             new string[] {
 				"Core",
diff --git a/Engine/Source/Runtime/AIModule/AIModule.Build.cs b/Engine/Source/Runtime/AIModule/AIModule.Build.cs
--- a/Engine/Source/Runtime/AIModule/AIModule.Build.cs
+++ b/Engine/Source/Runtime/AIModule/AIModule.Build.cs
@@ -15,7 +15,6 @@ namespace UnrealBuildTool.Rules
 
             PrivateIncludePaths.AddRange(
                 new string[] {
-                    "Runtime/AIModule/Private",
                     "Runtime/Engine/Private",
                 }
                 );
diff --git a/Engine/Source/Runtime/AVEncoder/AVEncoder.Build.cs b/Engine/Source/Runtime/AVEncoder/AVEncoder.Build.cs
--- a/Engine/Source/Runtime/AVEncoder/AVEncoder.Build.cs
+++ b/Engine/Source/Runtime/AVEncoder/AVEncoder.Build.cs
@@ -17,14 +17,6 @@ public class AVEncoder : ModuleRules
 
 		// PrecompileForTargets = PrecompileTargetsType.None;
 
-		PublicIncludePaths.AddRange(new string[] {
-			// ... add public include paths required here ...
-		});
-
-		PrivateIncludePaths.AddRange(new string[] {
-			// ... add other private include paths required here ...
-		});
-
 		PrivateDependencyModuleNames.AddRange(new string[] {
 			"Engine"
 		});
@@ -37,10 +29,6 @@ public class AVEncoder : ModuleRules
 			// ... add other public dependencies that you statically link with here ...
 		});
 
-		DynamicallyLoadedModuleNames.AddRange(new string[] {
-			// ... add any modules that your module loads dynamically here ...
-		});
-
 		if (Target.IsInPlatformGroup(UnrealPlatformGroup.Windows) || Target.IsInPlatformGroup(UnrealPlatformGroup.Unix))
 		{
 			PublicIncludePathModuleNames.Add("Vulkan");
diff --git a/Engine/Source/Runtime/AVIWriter/AVIWriter.Build.cs b/Engine/Source/Runtime/AVIWriter/AVIWriter.Build.cs
--- a/Engine/Source/Runtime/AVIWriter/AVIWriter.Build.cs
+++ b/Engine/Source/Runtime/AVIWriter/AVIWriter.Build.cs
@@ -6,12 +6,6 @@ namespace UnrealBuildTool.Rules
 	{
 		public AVIWriter(ReadOnlyTargetRules Target) : base(Target)
 		{
-			PrivateIncludePaths.AddRange(
-				new string[] {
-					"Runtime/AVIWriter/Private",
-				}
-				);
-
 			PublicDependencyModuleNames.AddRange(
 				new string[]
 				{
diff --git a/Engine/Source/Runtime/Advertising/Advertising/Advertising.Build.cs b/Engine/Source/Runtime/Advertising/Advertising/Advertising.Build.cs
--- a/Engine/Source/Runtime/Advertising/Advertising/Advertising.Build.cs
+++ b/Engine/Source/Runtime/Advertising/Advertising/Advertising.Build.cs
@@ -5,20 +5,7 @@ namespace UnrealBuildTool.Rules
 	public class Advertising : ModuleRules
 	{
 		public Advertising( ReadOnlyTargetRules Target ) : base(Target)
-		{
-			PublicIncludePaths.AddRange(
-				new string[] {
-					// ... add public include paths required here ...
-				}
-				);
-
-			PrivateIncludePaths.AddRange(
-				new string[] {
-					// ... add other private include paths required here ...
-				}
-				);
-
-			
+		{			
 			PublicDependencyModuleNames.AddRange(
 				new string[]
 				{
@@ -26,13 +13,6 @@ namespace UnrealBuildTool.Rules
                     // ... add other public dependencies that you statically link with here ...
 				}
 				);
-
-			DynamicallyLoadedModuleNames.AddRange(
-				new string[]
-				{
-					// ... add any modules that your module loads dynamically here ...
-				}
-				);
 		}
 	}
 }
diff --git a/Engine/Source/Runtime/Advertising/Android/AndroidAdvertising/AndroidAdvertising.Build.cs b/Engine/Source/Runtime/Advertising/Android/AndroidAdvertising/AndroidAdvertising.Build.cs
--- a/Engine/Source/Runtime/Advertising/Android/AndroidAdvertising/AndroidAdvertising.Build.cs
+++ b/Engine/Source/Runtime/Advertising/Android/AndroidAdvertising/AndroidAdvertising.Build.cs
@@ -7,20 +7,7 @@ namespace UnrealBuildTool.Rules
 	public class AndroidAdvertising : ModuleRules
 	{
 		public AndroidAdvertising( ReadOnlyTargetRules Target ) : base(Target)
-		{
-			PublicIncludePaths.AddRange(
-				new string[] {
-					// ... add public include paths required here ...
-				}
-				);
-
-			PrivateIncludePaths.AddRange(
-				new string[] {
-                    // ... add other private include paths required here ...
-				}
-				);
-
-			
+		{			
 			PublicDependencyModuleNames.AddRange(
 				new string[]
 				{
@@ -38,20 +25,6 @@ namespace UnrealBuildTool.Rules
 				);
 			PublicIncludePathModuleNames.Add("Advertising");
 
-            PrivateIncludePathModuleNames.AddRange(
-                new string[]
-                {
-                    // ... add any private module dependencies that should include paths
-                }
-                );
-
-			DynamicallyLoadedModuleNames.AddRange(
-				new string[]
-				{
-                    // ... add any modules that your module loads dynamically here ...
-				}
-				);
-
 			if (Target.Platform == UnrealTargetPlatform.Android)
 			{
 				string PluginPath = Utils.MakePathRelativeTo(ModuleDirectory, Target.RelativeEnginePath);
diff --git a/Engine/Source/Runtime/Advertising/IOS/IOSAdvertising/IOSAdvertising.Build.cs b/Engine/Source/Runtime/Advertising/IOS/IOSAdvertising/IOSAdvertising.Build.cs
--- a/Engine/Source/Runtime/Advertising/IOS/IOSAdvertising/IOSAdvertising.Build.cs
+++ b/Engine/Source/Runtime/Advertising/IOS/IOSAdvertising/IOSAdvertising.Build.cs
@@ -5,20 +5,7 @@ namespace UnrealBuildTool.Rules
 	public class IOSAdvertising : ModuleRules
 	{
 		public IOSAdvertising( ReadOnlyTargetRules Target ) : base(Target)
-		{
-			PublicIncludePaths.AddRange(
-				new string[] {
-					// ... add public include paths required here ...
-				}
-				);
-
-			PrivateIncludePaths.AddRange(
-				new string[] {
-                    // ... add other private include paths required here ...
-				}
-				);
-
-			
+		{			
 			PublicDependencyModuleNames.AddRange(
 				new string[]
 				{
@@ -36,20 +23,6 @@ namespace UnrealBuildTool.Rules
 				}
 				);
 			PublicIncludePathModuleNames.Add("Advertising");
-
-            PrivateIncludePathModuleNames.AddRange(
-                new string[]
-                {
-                    // ... add any private module dependencies that should include paths
-                }
-                );
-
-			DynamicallyLoadedModuleNames.AddRange(
-				new string[]
-				{
-                    // ... add any modules that your module loads dynamically here ...
-				}
-				);
 		}
 	}
 }
diff --git a/Engine/Source/Runtime/Analytics/Analytics/Analytics.Build.cs b/Engine/Source/Runtime/Analytics/Analytics/Analytics.Build.cs
--- a/Engine/Source/Runtime/Analytics/Analytics/Analytics.Build.cs
+++ b/Engine/Source/Runtime/Analytics/Analytics/Analytics.Build.cs
@@ -5,20 +5,7 @@ namespace UnrealBuildTool.Rules
 	public class Analytics : ModuleRules
 	{
 		public Analytics(ReadOnlyTargetRules Target) : base(Target)
-		{
-			PublicIncludePaths.AddRange(
-				new string[] {
-					// ... add public include paths required here ...
-				}
-				);
-
-			PrivateIncludePaths.AddRange(
-				new string[] {
-					// ... add other private include paths required here ...
-				}
-				);
-
-			
+		{			
 			PublicDependencyModuleNames.AddRange(
 				new string[]
 				{
@@ -26,13 +13,6 @@ namespace UnrealBuildTool.Rules
                     // ... add other public dependencies that you statically link with here ...
 				}
 				);
-
-            DynamicallyLoadedModuleNames.AddRange(
-				new string[]
-				{
-					// ... add any modules that your module loads dynamically here ...
-				}
-				);
 		}
 	}
 }
diff --git a/Engine/Source/Runtime/AnimGraphRuntime/AnimGraphRuntime.Build.cs b/Engine/Source/Runtime/AnimGraphRuntime/AnimGraphRuntime.Build.cs
--- a/Engine/Source/Runtime/AnimGraphRuntime/AnimGraphRuntime.Build.cs
+++ b/Engine/Source/Runtime/AnimGraphRuntime/AnimGraphRuntime.Build.cs
@@ -7,8 +7,6 @@ public class AnimGraphRuntime : ModuleRules
 {
 	public AnimGraphRuntime(ReadOnlyTargetRules Target) : base(Target)
 	{
-		PrivateIncludePaths.Add("Runtime/AnimGraphRuntime/Private");
-
 		PublicDependencyModuleNames.AddRange(
 			new string[] { 
 				"Core", 
diff --git a/Engine/Source/Runtime/AudioMixerCore/AudioMixerCore.Build.cs b/Engine/Source/Runtime/AudioMixerCore/AudioMixerCore.Build.cs
--- a/Engine/Source/Runtime/AudioMixerCore/AudioMixerCore.Build.cs
+++ b/Engine/Source/Runtime/AudioMixerCore/AudioMixerCore.Build.cs
@@ -8,14 +8,6 @@ namespace UnrealBuildTool.Rules
 		{
             PublicIncludePathModuleNames.Add("SignalProcessing");
 
-
-            PrivateIncludePaths.AddRange(
-				new string[]
-				{
-					"Runtime/AudioMixerCore/Private",
-				}
-			);
-
 			PrivateDependencyModuleNames.AddRange(
 				new string[]
                 {
diff --git a/Engine/Source/Runtime/AugmentedReality/AugmentedReality.build.cs b/Engine/Source/Runtime/AugmentedReality/AugmentedReality.build.cs
--- a/Engine/Source/Runtime/AugmentedReality/AugmentedReality.build.cs
+++ b/Engine/Source/Runtime/AugmentedReality/AugmentedReality.build.cs
@@ -6,9 +6,6 @@ namespace UnrealBuildTool.Rules
     {
         public AugmentedReality(ReadOnlyTargetRules Target) : base(Target)
         {
-            PrivateIncludePaths.Add("Runtime/AugmentedReality/Private");
-            PublicIncludePaths.Add("Runtime/AugmentedReality/Public");
-
             PrivateDependencyModuleNames.AddRange(
                 new string[]
                 {
diff --git a/Engine/Source/Runtime/AutomationMessages/AutomationMessages.Build.cs b/Engine/Source/Runtime/AutomationMessages/AutomationMessages.Build.cs
--- a/Engine/Source/Runtime/AutomationMessages/AutomationMessages.Build.cs
+++ b/Engine/Source/Runtime/AutomationMessages/AutomationMessages.Build.cs
@@ -18,11 +18,6 @@ namespace UnrealBuildTool.Rules
 				{
 					"Core"
 				});
-
-			PrivateIncludePaths.AddRange(
-				new string[]
-				{
-				});
 		}
 	}
 }
diff --git a/Engine/Source/Runtime/AutomationWorker/AutomationWorker.Build.cs b/Engine/Source/Runtime/AutomationWorker/AutomationWorker.Build.cs
--- a/Engine/Source/Runtime/AutomationWorker/AutomationWorker.Build.cs
+++ b/Engine/Source/Runtime/AutomationWorker/AutomationWorker.Build.cs
@@ -33,13 +33,6 @@ namespace UnrealBuildTool.Rules
 				}
 			);
 
-			PrivateIncludePaths.AddRange(
-				new string[]
-				{
-					"Runtime/AutomationWorker/Private",
-				}
-			);
-
 			if (Target.bCompileAgainstEngine)
 			{
 				PrivateDependencyModuleNames.Add("Engine");
diff --git a/Engine/Source/Runtime/BlueprintRuntime/BlueprintRuntime.Build.cs b/Engine/Source/Runtime/BlueprintRuntime/BlueprintRuntime.Build.cs
--- a/Engine/Source/Runtime/BlueprintRuntime/BlueprintRuntime.Build.cs
+++ b/Engine/Source/Runtime/BlueprintRuntime/BlueprintRuntime.Build.cs
@@ -6,12 +6,6 @@ namespace UnrealBuildTool.Rules
 	{
         public BlueprintRuntime(ReadOnlyTargetRules Target) : base(Target)
 		{
-			PrivateIncludePaths.AddRange(
-				new string[] {
-					"Runtime/BlueprintRuntime/Private",
-				}
-			);
-
 			PrivateDependencyModuleNames.AddRange(
 				new string[]
 				{
diff --git a/Engine/Source/Runtime/CEF3Utils/CEF3Utils.Build.cs b/Engine/Source/Runtime/CEF3Utils/CEF3Utils.Build.cs
--- a/Engine/Source/Runtime/CEF3Utils/CEF3Utils.Build.cs
+++ b/Engine/Source/Runtime/CEF3Utils/CEF3Utils.Build.cs
@@ -6,9 +6,6 @@ public class CEF3Utils : ModuleRules
 {
 	public CEF3Utils(ReadOnlyTargetRules Target) : base(Target)
 	{
-		PublicIncludePaths.Add("Runtime/CEF3Utils/Public");
-		PrivateIncludePaths.Add("Runtime/CEF3Utils/Private");
-
 		PrivateDependencyModuleNames.AddRange(
 			new string[]
 			{
diff --git a/Engine/Source/Runtime/CinematicCamera/CinematicCamera.Build.cs b/Engine/Source/Runtime/CinematicCamera/CinematicCamera.Build.cs
--- a/Engine/Source/Runtime/CinematicCamera/CinematicCamera.Build.cs
+++ b/Engine/Source/Runtime/CinematicCamera/CinematicCamera.Build.cs
@@ -16,11 +16,5 @@ public class CinematicCamera : ModuleRules
                 "SlateCore"
 			}
 		);
-
-        PrivateIncludePaths.AddRange(
-            new string[] {
-                "Runtime/CinematicCamera/Private"
-            })
-		;
 	}
 }
diff --git a/Engine/Source/Runtime/ClothingSystemRuntimeCommon/ClothingSystemRuntimeCommon.Build.cs b/Engine/Source/Runtime/ClothingSystemRuntimeCommon/ClothingSystemRuntimeCommon.Build.cs
--- a/Engine/Source/Runtime/ClothingSystemRuntimeCommon/ClothingSystemRuntimeCommon.Build.cs
+++ b/Engine/Source/Runtime/ClothingSystemRuntimeCommon/ClothingSystemRuntimeCommon.Build.cs
@@ -6,7 +6,6 @@ public class ClothingSystemRuntimeCommon : ModuleRules
 {
 	public ClothingSystemRuntimeCommon(ReadOnlyTargetRules Target) : base(Target)
 	{
-		PrivateIncludePaths.Add("Runtime/ClothingSystemRuntimeCommon/Private");
 		SetupModulePhysicsSupport(Target);
 
 		PublicIncludePathModuleNames.Add("ClothingSystemRuntimeInterface");
diff --git a/Engine/Source/Runtime/ClothingSystemRuntimeInterface/ClothingSystemRuntimeInterface.Build.cs b/Engine/Source/Runtime/ClothingSystemRuntimeInterface/ClothingSystemRuntimeInterface.Build.cs
--- a/Engine/Source/Runtime/ClothingSystemRuntimeInterface/ClothingSystemRuntimeInterface.Build.cs
+++ b/Engine/Source/Runtime/ClothingSystemRuntimeInterface/ClothingSystemRuntimeInterface.Build.cs
@@ -8,7 +8,6 @@ public class ClothingSystemRuntimeInterface : ModuleRules
 	{
 		ShortName = "ClothSysRuntimeIntrfc";
 
-        PrivateIncludePaths.Add("Runtime/ClothingSystemRuntimeInterface/Private");
         SetupModulePhysicsSupport(Target);
 
         PublicDependencyModuleNames.AddRange(
diff --git a/Engine/Source/Runtime/ClothingSystemRuntimeNv/ClothingSystemRuntimeNv.Build.cs b/Engine/Source/Runtime/ClothingSystemRuntimeNv/ClothingSystemRuntimeNv.Build.cs
--- a/Engine/Source/Runtime/ClothingSystemRuntimeNv/ClothingSystemRuntimeNv.Build.cs
+++ b/Engine/Source/Runtime/ClothingSystemRuntimeNv/ClothingSystemRuntimeNv.Build.cs
@@ -6,8 +6,6 @@ public class ClothingSystemRuntimeNv : ModuleRules
 {
 	public ClothingSystemRuntimeNv(ReadOnlyTargetRules Target) : base(Target)
 	{
-		PrivateIncludePaths.Add("Runtime/ClothingSystemRuntimeNv/Private");
-        
         SetupModulePhysicsSupport(Target);
 
         PublicIncludePathModuleNames.Add("ClothingSystemRuntimeInterface");
diff --git a/Engine/Source/Runtime/ColorManagement/ColorManagement.Build.cs b/Engine/Source/Runtime/ColorManagement/ColorManagement.Build.cs
--- a/Engine/Source/Runtime/ColorManagement/ColorManagement.Build.cs
+++ b/Engine/Source/Runtime/ColorManagement/ColorManagement.Build.cs
@@ -11,11 +11,5 @@ public class ColorManagement : ModuleRules
 				"Core",
 			}
 		);
-
-		PrivateIncludePaths.AddRange(
-			new string[] {
-				"Runtime/ColorManagement/Private"
-			}
-		);
 	}
 }
diff --git a/Engine/Source/Runtime/CoreUObject/CoreUObject.Build.cs b/Engine/Source/Runtime/CoreUObject/CoreUObject.Build.cs
--- a/Engine/Source/Runtime/CoreUObject/CoreUObject.Build.cs
+++ b/Engine/Source/Runtime/CoreUObject/CoreUObject.Build.cs
@@ -10,8 +10,6 @@ public class CoreUObject : ModuleRules
 
 		SharedPCHHeaderFile = "Public/CoreUObjectSharedPCH.h";
 
-		PrivateIncludePaths.Add("Runtime/CoreUObject/Private");
-
         PrivateIncludePathModuleNames.AddRange(
                 new string[] 
 			    {
diff --git a/Engine/Source/Runtime/CrashReportCore/CrashReportCore.Build.cs b/Engine/Source/Runtime/CrashReportCore/CrashReportCore.Build.cs
--- a/Engine/Source/Runtime/CrashReportCore/CrashReportCore.Build.cs
+++ b/Engine/Source/Runtime/CrashReportCore/CrashReportCore.Build.cs
@@ -6,12 +6,6 @@ public class CrashReportCore : ModuleRules
 {
 	public CrashReportCore( ReadOnlyTargetRules Target ) : base(Target)
 	{
-		PrivateIncludePaths.AddRange(
-		new string[] {
-				"Runtime/CrashReportCore/Private/",
-            	}
-        );
-
 		PublicDependencyModuleNames.AddRange(
 			new string[] {
 				"Core",
diff --git a/Engine/Source/Runtime/EngineMessages/EngineMessages.Build.cs b/Engine/Source/Runtime/EngineMessages/EngineMessages.Build.cs
--- a/Engine/Source/Runtime/EngineMessages/EngineMessages.Build.cs
+++ b/Engine/Source/Runtime/EngineMessages/EngineMessages.Build.cs
@@ -17,12 +17,6 @@ namespace UnrealBuildTool.Rules
 				{
 					"Core",
 				});
-
-			PrivateIncludePaths.AddRange(
-				new string[]
-				{
-					"Runtime/EngineMessages/Private",
-				});
 		}
 	}
 }
diff --git a/Engine/Source/Runtime/EngineSettings/EngineSettings.Build.cs b/Engine/Source/Runtime/EngineSettings/EngineSettings.Build.cs
--- a/Engine/Source/Runtime/EngineSettings/EngineSettings.Build.cs
+++ b/Engine/Source/Runtime/EngineSettings/EngineSettings.Build.cs
@@ -17,12 +17,6 @@ namespace UnrealBuildTool.Rules
 				{
 					"Core",
 				});
-
-			PrivateIncludePaths.AddRange(
-				new string[]
-				{
-					"Runtime/EngineSettings/Private",
-				});
 		}
 	}
 }
diff --git a/Engine/Source/Runtime/Experimental/Animation/Constraints/Constraints.Build.cs b/Engine/Source/Runtime/Experimental/Animation/Constraints/Constraints.Build.cs
--- a/Engine/Source/Runtime/Experimental/Animation/Constraints/Constraints.Build.cs
+++ b/Engine/Source/Runtime/Experimental/Animation/Constraints/Constraints.Build.cs
@@ -6,10 +6,6 @@ public class Constraints : ModuleRules
 {
 	public Constraints(ReadOnlyTargetRules Target) : base(Target)
 	{
-		PCHUsage = ModuleRules.PCHUsageMode.UseExplicitOrSharedPCHs;
-		PrivateIncludePaths.Add("Runtime/Experimental/Animation/Constraints/Private");
-		PublicIncludePaths.Add(ModuleDirectory + "/Public");
-
 		PublicDependencyModuleNames.AddRange(
 			new string[]
 			{
diff --git a/Engine/Source/Runtime/Experimental/Animation/DynamicPlayRate/DynamicPlayRate.Build.cs b/Engine/Source/Runtime/Experimental/Animation/DynamicPlayRate/DynamicPlayRate.Build.cs
--- a/Engine/Source/Runtime/Experimental/Animation/DynamicPlayRate/DynamicPlayRate.Build.cs
+++ b/Engine/Source/Runtime/Experimental/Animation/DynamicPlayRate/DynamicPlayRate.Build.cs
@@ -6,10 +6,6 @@ public class DynamicPlayRate : ModuleRules
 {
 	public DynamicPlayRate(ReadOnlyTargetRules Target) : base(Target)
 	{
-		PCHUsage = ModuleRules.PCHUsageMode.UseExplicitOrSharedPCHs;
-		PrivateIncludePaths.Add("Runtime/Experimental/Animation/DynamicPlayRate/Private");
-		PublicIncludePaths.Add(ModuleDirectory + "/Public");
-
 		PublicDependencyModuleNames.AddRange(
 			new string[]
 			{
diff --git a/Engine/Source/Runtime/Experimental/ChaosSolverEngine/ChaosSolverEngine.Build.cs b/Engine/Source/Runtime/Experimental/ChaosSolverEngine/ChaosSolverEngine.Build.cs
--- a/Engine/Source/Runtime/Experimental/ChaosSolverEngine/ChaosSolverEngine.Build.cs
+++ b/Engine/Source/Runtime/Experimental/ChaosSolverEngine/ChaosSolverEngine.Build.cs
@@ -6,9 +6,6 @@ namespace UnrealBuildTool.Rules
 	{
         public ChaosSolverEngine(ReadOnlyTargetRules Target) : base(Target)
 		{
-            PrivateIncludePaths.Add("Runtime/Experimental/ChaosSolverEngine/Private");
-            PublicIncludePaths.Add(ModuleDirectory + "/Public");
-
 			PublicDependencyModuleNames.AddRange(
 				new string[]
 				{
diff --git a/Engine/Source/Runtime/Experimental/GeometryCollectionEngine/GeometryCollectionEngine.Build.cs b/Engine/Source/Runtime/Experimental/GeometryCollectionEngine/GeometryCollectionEngine.Build.cs
--- a/Engine/Source/Runtime/Experimental/GeometryCollectionEngine/GeometryCollectionEngine.Build.cs
+++ b/Engine/Source/Runtime/Experimental/GeometryCollectionEngine/GeometryCollectionEngine.Build.cs
@@ -6,9 +6,6 @@ namespace UnrealBuildTool.Rules
 	{
         public GeometryCollectionEngine(ReadOnlyTargetRules Target) : base(Target)
 		{
-			PrivateIncludePaths.Add("Runtime/Experimental/GeometryCollectionEngine/Private");
-            PublicIncludePaths.Add(ModuleDirectory + "/Public");
-
 			SetupModulePhysicsSupport(Target);
 
 			PublicDependencyModuleNames.AddRange(
diff --git a/Engine/Source/Runtime/Experimental/Voronoi/Voronoi.Build.cs b/Engine/Source/Runtime/Experimental/Voronoi/Voronoi.Build.cs
--- a/Engine/Source/Runtime/Experimental/Voronoi/Voronoi.Build.cs
+++ b/Engine/Source/Runtime/Experimental/Voronoi/Voronoi.Build.cs
@@ -6,9 +6,6 @@ namespace UnrealBuildTool.Rules
 	{
         public Voronoi(ReadOnlyTargetRules Target) : base(Target)
 		{
-			PrivateIncludePaths.Add("Runtime/Experimental/Voronoi/Private");
-            PublicIncludePaths.Add(ModuleDirectory + "/Public");
-
 			PublicDependencyModuleNames.AddRange(
 				new string[]
 				{
diff --git a/Engine/Source/Runtime/Foliage/Foliage.Build.cs b/Engine/Source/Runtime/Foliage/Foliage.Build.cs
--- a/Engine/Source/Runtime/Foliage/Foliage.Build.cs
+++ b/Engine/Source/Runtime/Foliage/Foliage.Build.cs
@@ -16,11 +16,5 @@ public class Foliage: ModuleRules
 				"RHI"
 			}
 		);
-
-        PrivateIncludePaths.AddRange(
-            new string[] {
-                "Runtime/Foliage/Private"
-            })
-		;
 	}
 }
diff --git a/Engine/Source/Runtime/FriendsAndChat/FriendsAndChat.Build.cs b/Engine/Source/Runtime/FriendsAndChat/FriendsAndChat.Build.cs
--- a/Engine/Source/Runtime/FriendsAndChat/FriendsAndChat.Build.cs
+++ b/Engine/Source/Runtime/FriendsAndChat/FriendsAndChat.Build.cs
@@ -21,19 +21,5 @@ public class FriendsAndChat : ModuleRules
 				"SlateCore",
 			}
 		);
-
-		PrivateIncludePaths.AddRange(
-			new string[]
-			{
-				"Runtime/FriendsAndChat/Private",
-			}
-		);
-
-		PublicIncludePaths.AddRange(
-			new string[]
-			{
-				"Runtime/FriendsAndChat/Public",
-			}
-		);
 	}
 }
diff --git a/Engine/Source/Runtime/GameMenuBuilder/GameMenuBuilder.Build.cs b/Engine/Source/Runtime/GameMenuBuilder/GameMenuBuilder.Build.cs
--- a/Engine/Source/Runtime/GameMenuBuilder/GameMenuBuilder.Build.cs
+++ b/Engine/Source/Runtime/GameMenuBuilder/GameMenuBuilder.Build.cs
@@ -15,12 +15,6 @@ public class GameMenuBuilder : ModuleRules
 					"Slate",
                     "SlateCore",
 			}
-		);
-
-		PrivateIncludePaths.AddRange(
-			new string[] {
-				"Runtime/GameMenuBuilder/Private",
-			}
-		);		
+		);	
 	}
 }
diff --git a/Engine/Source/Runtime/GameplayTags/GameplayTags.Build.cs b/Engine/Source/Runtime/GameplayTags/GameplayTags.Build.cs
--- a/Engine/Source/Runtime/GameplayTags/GameplayTags.Build.cs
+++ b/Engine/Source/Runtime/GameplayTags/GameplayTags.Build.cs
@@ -6,12 +6,6 @@ namespace UnrealBuildTool.Rules
 	{
 		public GameplayTags(ReadOnlyTargetRules Target) : base(Target)
 		{
-			PrivateIncludePaths.AddRange(
-				new string[] {
-					"Runtime/GameplayTags/Private",
-				}
-				);
-
 			PublicDependencyModuleNames.AddRange(
 				new string[]
 				{
diff --git a/Engine/Source/Runtime/GameplayTasks/GameplayTasks.Build.cs b/Engine/Source/Runtime/GameplayTasks/GameplayTasks.Build.cs
--- a/Engine/Source/Runtime/GameplayTasks/GameplayTasks.Build.cs
+++ b/Engine/Source/Runtime/GameplayTasks/GameplayTasks.Build.cs
@@ -6,20 +6,6 @@ namespace UnrealBuildTool.Rules
 	{
         public GameplayTasks(ReadOnlyTargetRules Target) : base(Target)
 		{
-            PrivateIncludePaths.Add("Runtime/GameplayTasks/Private");
-            
-			PublicIncludePaths.AddRange(
-				new string[] {
-					// ... add public include paths required here ...
-				}
-				);
-
-			PrivateIncludePaths.AddRange(
-				new string[] {
-					// ... add other private include paths required here ...
-				}
-				);
-
 			PublicDependencyModuleNames.AddRange(
 				new string[]
 				{
@@ -32,20 +18,6 @@ namespace UnrealBuildTool.Rules
 				}
 				);
 
-			PrivateDependencyModuleNames.AddRange(
-				new string[]
-				{
-					// ... add private dependencies that you statically link with here ...
-				}
-				);
-
-			DynamicallyLoadedModuleNames.AddRange(
-				new string[]
-				{
-					// ... add any modules that your module loads dynamically here ...
-				}
-				);
-
 			if (Target.bBuildEditor == true)
 			{
 				PrivateDependencyModuleNames.Add("EditorFramework");
diff --git a/Engine/Source/Runtime/HardwareSurvey/HardwareSurvey.Build.cs b/Engine/Source/Runtime/HardwareSurvey/HardwareSurvey.Build.cs
--- a/Engine/Source/Runtime/HardwareSurvey/HardwareSurvey.Build.cs
+++ b/Engine/Source/Runtime/HardwareSurvey/HardwareSurvey.Build.cs
@@ -14,20 +14,6 @@ public class HardwareSurvey : ModuleRules
 			}
 		);
 
-		PrivateIncludePaths.AddRange(
-			new string[]
-			{
-				"Runtime/HardwareSurvey/Private",
-			}
-		);
-
-		PublicIncludePaths.AddRange(
-			new string[]
-			{
-				"Runtime/HardwareSurvey/Public",
-			}
-		);
-
 		DynamicallyLoadedModuleNames.AddRange(
 		new string[] {
 				"Analytics",
diff --git a/Engine/Source/Runtime/HeadMountedDisplay/HeadMountedDisplay.Build.cs b/Engine/Source/Runtime/HeadMountedDisplay/HeadMountedDisplay.Build.cs
--- a/Engine/Source/Runtime/HeadMountedDisplay/HeadMountedDisplay.Build.cs
+++ b/Engine/Source/Runtime/HeadMountedDisplay/HeadMountedDisplay.Build.cs
@@ -1,17 +1,12 @@
 // Copyright Epic Games, Inc. All Rights Reserved.
 
+using System.IO;
 using UnrealBuildTool;
 
 public class HeadMountedDisplay : ModuleRules
 {
     public HeadMountedDisplay(ReadOnlyTargetRules Target) : base(Target)
 	{
-		PublicIncludePaths.AddRange(
-			new string[] {
-				"Runtime/HeadMountedDisplay/Public"
-			}
-		);
-
         PrivateIncludePaths.AddRange(
             new string[] {
                 "Runtime/Renderer/Private"
@@ -40,13 +35,6 @@ public class HeadMountedDisplay : ModuleRules
             }
         );
 
-        PublicIncludePathModuleNames.AddRange(
-            new string[]
-            {
-            }
-       );
-
-
         if (Target.bBuildEditor == true)
 		{
 			PrivateDependencyModuleNames.AddRange(
diff --git a/Engine/Source/Runtime/InteractiveToolsFramework/InteractiveToolsFramework.Build.cs b/Engine/Source/Runtime/InteractiveToolsFramework/InteractiveToolsFramework.Build.cs
--- a/Engine/Source/Runtime/InteractiveToolsFramework/InteractiveToolsFramework.Build.cs
+++ b/Engine/Source/Runtime/InteractiveToolsFramework/InteractiveToolsFramework.Build.cs
@@ -6,21 +6,7 @@ public class InteractiveToolsFramework : ModuleRules
 {
 	public InteractiveToolsFramework(ReadOnlyTargetRules Target) : base(Target)
 	{
-        PCHUsage = ModuleRules.PCHUsageMode.UseExplicitOrSharedPCHs;
-		
-		PublicIncludePaths.AddRange(
-			new string[] {
-				// ... add public include paths required here ...
-			}
-			);
-				
-		
-		PrivateIncludePaths.AddRange(
-			new string[] {
-				// ... add other private include paths required here ...
-			}
-			);
-			
+        PCHUsage = ModuleRules.PCHUsageMode.UseExplicitOrSharedPCHs;			
 		
 		PublicDependencyModuleNames.AddRange(
 			new string[]
@@ -32,8 +18,7 @@ public class InteractiveToolsFramework : ModuleRules
 				"MeshDescription"
 				// ... add other public dependencies that you statically link with here ...
 			}
-            );
-			
+            );			
 		
 		PrivateDependencyModuleNames.AddRange(
 			new string[]
@@ -46,13 +31,5 @@ public class InteractiveToolsFramework : ModuleRules
 				// ... add private dependencies that you statically link with here ...	
 			}
 			);
-		
-		
-		DynamicallyLoadedModuleNames.AddRange(
-			new string[]
-			{
-				// ... add any modules that your module loads dynamically here ...
-			}
-			);
 	}
 }
diff --git a/Engine/Source/Runtime/Json/Json.Build.cs b/Engine/Source/Runtime/Json/Json.Build.cs
--- a/Engine/Source/Runtime/Json/Json.Build.cs
+++ b/Engine/Source/Runtime/Json/Json.Build.cs
@@ -13,12 +13,6 @@ namespace UnrealBuildTool.Rules
 				}
 			); 
 
-			PrivateIncludePaths.AddRange(
-				new string[] {
-					"Runtime/Json/Private",
-				}
-			);
-
 			UnsafeTypeCastWarningLevel = WarningLevel.Error;
 		}
 	}
diff --git a/Engine/Source/Runtime/Launch/Launch.Build.cs b/Engine/Source/Runtime/Launch/Launch.Build.cs
--- a/Engine/Source/Runtime/Launch/Launch.Build.cs
+++ b/Engine/Source/Runtime/Launch/Launch.Build.cs
@@ -7,8 +7,6 @@ public class Launch : ModuleRules
 {
 	public Launch(ReadOnlyTargetRules Target) : base(Target)
 	{
-		PrivateIncludePaths.Add("Runtime/Launch/Private");
-
 		PrivateIncludePathModuleNames.AddRange(new string[] {
 				"AutomationController",
 				"TaskGraph",
@@ -89,7 +87,7 @@ public class Launch : ModuleRules
 			}
 			else if (Target.Platform == UnrealTargetPlatform.Mac)
 			{
-				DynamicallyLoadedModuleNames.AddRange(new string [] {
+				DynamicallyLoadedModuleNames.AddRange(new string[] {
 					"AudioMixerCoreAudio",
 				});
 			}
@@ -255,11 +253,11 @@ public class Launch : ModuleRules
 			}
 		}
 
-		if(Target.LinkType == TargetLinkType.Monolithic && !Target.bFormalBuild)
+		if (Target.LinkType == TargetLinkType.Monolithic && !Target.bFormalBuild)
 		{
 			PrivateDefinitions.Add(string.Format("COMPILED_IN_CL={0}", Target.Version.Changelist));
 			PrivateDefinitions.Add(string.Format("COMPILED_IN_COMPATIBLE_CL={0}", Target.Version.EffectiveCompatibleChangelist));
-			PrivateDefinitions.Add(string.Format("COMPILED_IN_BRANCH_NAME={0}", (Target.Version.BranchName == null || Target.Version.BranchName.Length == 0)? "UE" : Target.Version.BranchName));
+			PrivateDefinitions.Add(string.Format("COMPILED_IN_BRANCH_NAME={0}", (Target.Version.BranchName == null || Target.Version.BranchName.Length == 0) ? "UE" : Target.Version.BranchName));
 		}
 	}
 }
diff --git a/Engine/Source/Runtime/LevelSequence/LevelSequence.Build.cs b/Engine/Source/Runtime/LevelSequence/LevelSequence.Build.cs
--- a/Engine/Source/Runtime/LevelSequence/LevelSequence.Build.cs
+++ b/Engine/Source/Runtime/LevelSequence/LevelSequence.Build.cs
@@ -6,8 +6,6 @@ public class LevelSequence : ModuleRules
 {
 	public LevelSequence(ReadOnlyTargetRules Target) : base(Target)
 	{
-		PrivateIncludePaths.Add("Runtime/LevelSequence/Private");
-
 		PublicDependencyModuleNames.AddRange(
 			new string[] {
 				"Core",
diff --git a/Engine/Source/Runtime/Media/Media.Build.cs b/Engine/Source/Runtime/Media/Media.Build.cs
--- a/Engine/Source/Runtime/Media/Media.Build.cs
+++ b/Engine/Source/Runtime/Media/Media.Build.cs
@@ -11,11 +11,6 @@ namespace UnrealBuildTool.Rules
 				{
 					"Core",
 				});
-
-			PrivateIncludePaths.AddRange(
-				new string[] {
-					"Runtime/Media/Private",
-				});
 		}
 	}
 }
diff --git a/Engine/Source/Runtime/MediaUtils/MediaUtils.Build.cs b/Engine/Source/Runtime/MediaUtils/MediaUtils.Build.cs
--- a/Engine/Source/Runtime/MediaUtils/MediaUtils.Build.cs
+++ b/Engine/Source/Runtime/MediaUtils/MediaUtils.Build.cs
@@ -20,11 +20,6 @@ namespace UnrealBuildTool.Rules
 					"CoreUObject",
 					"ImageWriteQueue",
 				});
-
-			PrivateIncludePaths.AddRange(
-				new string[] {
-					"Runtime/MediaUtils/Private",
-				});
 		}
 	}
 }
diff --git a/Engine/Source/Runtime/MeshDescription/MeshDescription.Build.cs b/Engine/Source/Runtime/MeshDescription/MeshDescription.Build.cs
--- a/Engine/Source/Runtime/MeshDescription/MeshDescription.Build.cs
+++ b/Engine/Source/Runtime/MeshDescription/MeshDescription.Build.cs
@@ -6,9 +6,6 @@ namespace UnrealBuildTool.Rules
 	{
 		public MeshDescription(ReadOnlyTargetRules Target) : base(Target)
 		{
-            PrivateIncludePaths.Add("Runtime/MeshDescription/Private");
-            PublicIncludePaths.Add("Runtime/MeshDescription/Public");
-
 			PublicDependencyModuleNames.AddRange(
 				new string[]
 				{
diff --git a/Engine/Source/Runtime/MessagingRpc/MessagingRpc.Build.cs b/Engine/Source/Runtime/MessagingRpc/MessagingRpc.Build.cs
--- a/Engine/Source/Runtime/MessagingRpc/MessagingRpc.Build.cs
+++ b/Engine/Source/Runtime/MessagingRpc/MessagingRpc.Build.cs
@@ -25,11 +25,6 @@ namespace UnrealBuildTool.Rules
 				{
 					"Messaging",
 				});
-
-			PrivateIncludePaths.AddRange(
-				new string[] {
-					"Runtime/MessagingRpc/Private",
-				});
 		}
 	}
 }
diff --git a/Engine/Source/Runtime/MoviePlayer/MoviePlayer.Build.cs b/Engine/Source/Runtime/MoviePlayer/MoviePlayer.Build.cs
--- a/Engine/Source/Runtime/MoviePlayer/MoviePlayer.Build.cs
+++ b/Engine/Source/Runtime/MoviePlayer/MoviePlayer.Build.cs
@@ -6,8 +6,6 @@ public class MoviePlayer : ModuleRules
 {
 	public MoviePlayer(ReadOnlyTargetRules Target) : base(Target)
 	{
-        PrivateIncludePaths.Add("Runtime/MoviePlayer/Private");
-
 		PublicDependencyModuleNames.AddRange(
 			new string[] {
 					"Engine",
diff --git a/Engine/Source/Runtime/MoviePlayerProxy/MoviePlayerProxy.Build.cs b/Engine/Source/Runtime/MoviePlayerProxy/MoviePlayerProxy.Build.cs
--- a/Engine/Source/Runtime/MoviePlayerProxy/MoviePlayerProxy.Build.cs
+++ b/Engine/Source/Runtime/MoviePlayerProxy/MoviePlayerProxy.Build.cs
@@ -6,8 +6,6 @@ public class MoviePlayerProxy : ModuleRules
 {
 	public MoviePlayerProxy(ReadOnlyTargetRules Target) : base(Target)
 	{
-        PrivateIncludePaths.Add("Runtime/MoviePlayerProxy/Private");
-
 		PrivateDependencyModuleNames.AddRange(
             new string[]
 			{
diff --git a/Engine/Source/Runtime/MovieScene/MovieScene.Build.cs b/Engine/Source/Runtime/MovieScene/MovieScene.Build.cs
--- a/Engine/Source/Runtime/MovieScene/MovieScene.Build.cs
+++ b/Engine/Source/Runtime/MovieScene/MovieScene.Build.cs
@@ -6,8 +6,6 @@ public class MovieScene : ModuleRules
 {
 	public MovieScene(ReadOnlyTargetRules Target) : base(Target)
 	{
-		PrivateIncludePaths.Add("Runtime/MovieScene/Private");
-
         PrivateIncludePathModuleNames.AddRange(
             new string[] {
                 "TargetPlatform"
diff --git a/Engine/Source/Runtime/MovieSceneCapture/MovieSceneCapture.Build.cs b/Engine/Source/Runtime/MovieSceneCapture/MovieSceneCapture.Build.cs
--- a/Engine/Source/Runtime/MovieSceneCapture/MovieSceneCapture.Build.cs
+++ b/Engine/Source/Runtime/MovieSceneCapture/MovieSceneCapture.Build.cs
@@ -6,12 +6,6 @@ public class MovieSceneCapture : ModuleRules
 {
 	public MovieSceneCapture(ReadOnlyTargetRules Target) : base(Target)
 	{
-		PrivateIncludePaths.AddRange(
-			new string[] {
-				"Runtime/MovieSceneCapture/Private"
-			}
-		);
-
 		if (Target.Type == TargetType.Editor || Target.Type == TargetType.Program)
 		{
 			PublicIncludePathModuleNames.Add("ImageWrapper");
diff --git a/Engine/Source/Runtime/Navmesh/Navmesh.Build.cs b/Engine/Source/Runtime/Navmesh/Navmesh.Build.cs
--- a/Engine/Source/Runtime/Navmesh/Navmesh.Build.cs
+++ b/Engine/Source/Runtime/Navmesh/Navmesh.Build.cs
@@ -6,9 +6,6 @@ namespace UnrealBuildTool.Rules
 	{
 		public Navmesh(ReadOnlyTargetRules Target) : base(Target)
 		{
-            PublicIncludePaths.Add("Runtime/Navmesh/Public");
-			PrivateIncludePaths.Add("Runtime/Navmesh/Private");
-
             PrivateDependencyModuleNames.AddRange(new string[] { "Core" });
 
 			// This is an unsupported feature and has not been finished to production quality.
diff --git a/Engine/Source/Runtime/Networking/Networking.Build.cs b/Engine/Source/Runtime/Networking/Networking.Build.cs
--- a/Engine/Source/Runtime/Networking/Networking.Build.cs
+++ b/Engine/Source/Runtime/Networking/Networking.Build.cs
@@ -6,8 +6,6 @@ namespace UnrealBuildTool.Rules
 	{
 		public Networking(ReadOnlyTargetRules Target) : base(Target)
 		{
-			PrivateIncludePaths.Add("Runtime/Networking/Private");
-
 			PublicDependencyModuleNames.AddRange(
 				new string[]
 				{
diff --git a/Engine/Source/Runtime/NullDrv/NullDrv.Build.cs b/Engine/Source/Runtime/NullDrv/NullDrv.Build.cs
--- a/Engine/Source/Runtime/NullDrv/NullDrv.Build.cs
+++ b/Engine/Source/Runtime/NullDrv/NullDrv.Build.cs
@@ -6,8 +6,6 @@ public class NullDrv : ModuleRules
 {
 	public NullDrv(ReadOnlyTargetRules Target) : base(Target)
 	{
-        PrivateIncludePaths.Add("Runtime/NullDrv/Private");
-
         PrivateDependencyModuleNames.Add("Core");
         PrivateDependencyModuleNames.Add("RHI");
         PrivateDependencyModuleNames.Add("RenderCore");
diff --git a/Engine/Source/Runtime/Online/BuildPatchServices/BuildPatchServices.Build.cs b/Engine/Source/Runtime/Online/BuildPatchServices/BuildPatchServices.Build.cs
--- a/Engine/Source/Runtime/Online/BuildPatchServices/BuildPatchServices.Build.cs
+++ b/Engine/Source/Runtime/Online/BuildPatchServices/BuildPatchServices.Build.cs
@@ -20,13 +20,6 @@ public class BuildPatchServices : ModuleRules
 			}
 		);
 
-		PrivateIncludePaths.AddRange(
-			new string[]
-			{
-				"Runtime/Online/BuildPatchServices/Private",
-			}
-		);
-
 		PrivateDependencyModuleNames.AddRange(
 		new string[] {
 				"Analytics",
diff --git a/Engine/Source/Runtime/Online/HTTP/HTTP.Build.cs b/Engine/Source/Runtime/Online/HTTP/HTTP.Build.cs
--- a/Engine/Source/Runtime/Online/HTTP/HTTP.Build.cs
+++ b/Engine/Source/Runtime/Online/HTTP/HTTP.Build.cs
@@ -39,12 +39,6 @@ public class HTTP : ModuleRules
 	{
 		PublicDefinitions.Add("HTTP_PACKAGE=1");
 
-		PrivateIncludePaths.AddRange(
-			new string[] {
-				"Runtime/Online/HTTP/Private",
-			}
-			);
-
 		PublicDependencyModuleNames.AddRange(
 			new string[] {
 				"Core"
diff --git a/Engine/Source/Runtime/Online/HTTPServer/HttpServer.Build.cs b/Engine/Source/Runtime/Online/HTTPServer/HttpServer.Build.cs
--- a/Engine/Source/Runtime/Online/HTTPServer/HttpServer.Build.cs
+++ b/Engine/Source/Runtime/Online/HTTPServer/HttpServer.Build.cs
@@ -6,18 +6,6 @@ public class HTTPServer : ModuleRules
 {
     public HTTPServer(ReadOnlyTargetRules Target) : base(Target)
     {
-        PublicIncludePaths.AddRange(
-			new string[] {
-				"Runtime/Online/HTTPServer/Public",
-            }
-        );
-
-        PrivateIncludePaths.AddRange(
-            new string[] {
-                "Runtime/Online/HTTPServer/Private",
-            }
-        );
-
         PrivateDependencyModuleNames.AddRange(
             new string[] {
                 "Core",
diff --git a/Engine/Source/Runtime/Online/ICMP/Icmp.Build.cs b/Engine/Source/Runtime/Online/ICMP/Icmp.Build.cs
--- a/Engine/Source/Runtime/Online/ICMP/Icmp.Build.cs
+++ b/Engine/Source/Runtime/Online/ICMP/Icmp.Build.cs
@@ -8,12 +8,6 @@ public class Icmp : ModuleRules
 	{
 		PublicDefinitions.Add("ICMP_PACKAGE=1");
 
-		PrivateIncludePaths.AddRange(
-			new string[] {
-				"Runtime/Online/ICMP/Private",
-			}
-			);
-
 		PrivateDependencyModuleNames.AddRange(
 			new string[] { 
 				"Core", 
diff --git a/Engine/Source/Runtime/Online/SSL/SSL.Build.cs b/Engine/Source/Runtime/Online/SSL/SSL.Build.cs
--- a/Engine/Source/Runtime/Online/SSL/SSL.Build.cs
+++ b/Engine/Source/Runtime/Online/SSL/SSL.Build.cs
@@ -42,12 +42,6 @@ public class SSL : ModuleRules
 			PublicDefinitions.Add("WITH_SSL=1");
 			PrivateDefinitions.Add("USE_DEFAULT_SSLCERT=" + (bUseDefaultSSLCert ? "1" : "0"));
 
-			PrivateIncludePaths.AddRange(
-				new string[] {
-					"Runtime/Online/SSL/Private",
-				}
-			);
-
 			AddEngineThirdPartyPrivateStaticDependencies(Target, "OpenSSL");
 
 			if (Target.Platform == UnrealTargetPlatform.Win64)
diff --git a/Engine/Source/Runtime/Online/Stomp/Stomp.Build.cs b/Engine/Source/Runtime/Online/Stomp/Stomp.Build.cs
--- a/Engine/Source/Runtime/Online/Stomp/Stomp.Build.cs
+++ b/Engine/Source/Runtime/Online/Stomp/Stomp.Build.cs
@@ -27,13 +27,6 @@ public class Stomp : ModuleRules
 		{
 			PublicDefinitions.Add("WITH_STOMP=1");
 
-			PrivateIncludePaths.AddRange(
-				new string[]
-				{
-					"Runtime/Online/Stomp/Private",
-				}
-			);
-
 			PrivateDependencyModuleNames.AddRange(
 				new string[] {
 					"WebSockets"
diff --git a/Engine/Source/Runtime/Online/XMPP/XMPP.Build.cs b/Engine/Source/Runtime/Online/XMPP/XMPP.Build.cs
--- a/Engine/Source/Runtime/Online/XMPP/XMPP.Build.cs
+++ b/Engine/Source/Runtime/Online/XMPP/XMPP.Build.cs
@@ -22,13 +22,6 @@ public class XMPP : ModuleRules
 	{
 		PublicDefinitions.Add("XMPP_PACKAGE=1");
 
-		PrivateIncludePaths.AddRange(
-			new string[] 
-			{
-				"Runtime/Online/XMPP/Private"
-			}
-		);
-
 		PrivateDependencyModuleNames.AddRange(
 			new string[] 
 			{ 
diff --git a/Engine/Source/Runtime/OpenGLDrv/OpenGLDrv.Build.cs b/Engine/Source/Runtime/OpenGLDrv/OpenGLDrv.Build.cs
--- a/Engine/Source/Runtime/OpenGLDrv/OpenGLDrv.Build.cs
+++ b/Engine/Source/Runtime/OpenGLDrv/OpenGLDrv.Build.cs
@@ -7,8 +7,6 @@ public class OpenGLDrv : ModuleRules
 {
 	public OpenGLDrv(ReadOnlyTargetRules Target) : base(Target)
 	{
-		PrivateIncludePaths.Add("Runtime/OpenGLDrv/Private");
-
 		PublicDependencyModuleNames.AddRange(new string[] {
 			"Core",
 			"RHI",
diff --git a/Engine/Source/Runtime/Portal/LauncherPlatform/LauncherPlatform.Build.cs b/Engine/Source/Runtime/Portal/LauncherPlatform/LauncherPlatform.Build.cs
--- a/Engine/Source/Runtime/Portal/LauncherPlatform/LauncherPlatform.Build.cs
+++ b/Engine/Source/Runtime/Portal/LauncherPlatform/LauncherPlatform.Build.cs
@@ -6,8 +6,6 @@ public class LauncherPlatform : ModuleRules
 {
     public LauncherPlatform(ReadOnlyTargetRules Target) : base(Target)
     {
-        PrivateIncludePaths.Add("Runtime/Portal/LauncherPlatform/Private");
-
         PrivateDependencyModuleNames.AddRange(
             new string[] {
                 "Core",
diff --git a/Engine/Source/Runtime/Portal/Messages/PortalMessages.Build.cs b/Engine/Source/Runtime/Portal/Messages/PortalMessages.Build.cs
--- a/Engine/Source/Runtime/Portal/Messages/PortalMessages.Build.cs
+++ b/Engine/Source/Runtime/Portal/Messages/PortalMessages.Build.cs
@@ -14,11 +14,6 @@ namespace UnrealBuildTool.Rules
 					"MessagingRpc",
 					"PortalServices",
 				});
-
-			PrivateIncludePaths.AddRange(
-				new string[] {
-					"Runtime/Portal/Messages/Private",
-				});
 		}
 	}
 }
diff --git a/Engine/Source/Runtime/Portal/Rpc/PortalRpc.Build.cs b/Engine/Source/Runtime/Portal/Rpc/PortalRpc.Build.cs
--- a/Engine/Source/Runtime/Portal/Rpc/PortalRpc.Build.cs
+++ b/Engine/Source/Runtime/Portal/Rpc/PortalRpc.Build.cs
@@ -28,11 +28,6 @@ namespace UnrealBuildTool.Rules
                     "PortalServices",
                 }
             );
-
-			PrivateIncludePaths.AddRange(
-				new string[] {
-				}
-			);
 		}
 	}
 }
diff --git a/Engine/Source/Runtime/Projects/Projects.Build.cs b/Engine/Source/Runtime/Projects/Projects.Build.cs
--- a/Engine/Source/Runtime/Projects/Projects.Build.cs
+++ b/Engine/Source/Runtime/Projects/Projects.Build.cs
@@ -17,13 +17,6 @@ namespace UnrealBuildTool.Rules
 				}
 			);
 
-			PrivateIncludePaths.AddRange(
-				new string[]
-				{
-					"Runtime/Projects/Private",
-				}
-			);
-
 			// Monolithic and non-editor builds compile in the enabled/disabled plugins. Non-monolithic editor builds save them in the plugin receipt to avoid invalidating the shared build environment.
 			if ((Target.Type == TargetType.Editor) && (Target.BuildEnvironment == TargetBuildEnvironment.Shared) && (Target.LinkType != TargetLinkType.Monolithic))
 			{
diff --git a/Engine/Source/Runtime/RHI/RHI.Build.cs b/Engine/Source/Runtime/RHI/RHI.Build.cs
--- a/Engine/Source/Runtime/RHI/RHI.Build.cs
+++ b/Engine/Source/Runtime/RHI/RHI.Build.cs
@@ -57,7 +57,5 @@ public class RHI : ModuleRules
 		{
 			PrivateIncludePathModuleNames.AddRange(new string[] { "TaskGraph" });
 		}
-
-		PrivateIncludePaths.Add("Runtime/RHI/Private");
     }
 }
diff --git a/Engine/Source/Runtime/Serialization/Serialization.Build.cs b/Engine/Source/Runtime/Serialization/Serialization.Build.cs
--- a/Engine/Source/Runtime/Serialization/Serialization.Build.cs
+++ b/Engine/Source/Runtime/Serialization/Serialization.Build.cs
@@ -17,11 +17,6 @@ namespace UnrealBuildTool.Rules
 				new string[] {
 					"CoreUObject",
 				});
-
-			PrivateIncludePaths.AddRange(
-				new string[] {
-					"Runtime/Serialization/Private",
-				});
 		}
 	}
 }
diff --git a/Engine/Source/Runtime/SessionMessages/SessionMessages.Build.cs b/Engine/Source/Runtime/SessionMessages/SessionMessages.Build.cs
--- a/Engine/Source/Runtime/SessionMessages/SessionMessages.Build.cs
+++ b/Engine/Source/Runtime/SessionMessages/SessionMessages.Build.cs
@@ -17,12 +17,6 @@ namespace UnrealBuildTool.Rules
 				{
 					"Core",
 				});
-
-			PrivateIncludePaths.AddRange(
-				new string[]
-				{
-					"Runtime/SessionMessages/Private",
-				});
 		}
 	}
 }
diff --git a/Engine/Source/Runtime/SkeletalMeshDescription/SkeletalMeshDescription.Build.cs b/Engine/Source/Runtime/SkeletalMeshDescription/SkeletalMeshDescription.Build.cs
--- a/Engine/Source/Runtime/SkeletalMeshDescription/SkeletalMeshDescription.Build.cs
+++ b/Engine/Source/Runtime/SkeletalMeshDescription/SkeletalMeshDescription.Build.cs
@@ -6,9 +6,6 @@ namespace UnrealBuildTool.Rules
 	{
 		public SkeletalMeshDescription(ReadOnlyTargetRules Target) : base(Target)
 		{
-            PrivateIncludePaths.Add("Runtime/SkeletalMeshDescription/Private");
-            PublicIncludePaths.Add("Runtime/SkeletalMeshDescription/Public");
-
             // For GPUSkinPublicDefs.h
             PublicIncludePaths.Add("Runtime/Engine/Public");
             
