Move UV Editor 2D viewport behavior targets to EditorInteractiveToolsFramework for reuse in other editors

diff --git a/Engine/Plugins/Experimental/UVEditor/Source/UVEditor/Private/UVEditor2DViewportBehaviorTargets.cpp b/Engine/Plugins/Experimental/UVEditor/Source/UVEditor/Private/UVEditor2DViewportBehaviorTargets.cpp
deleted file mode 100644
diff --git a/Engine/Plugins/Experimental/UVEditor/Source/UVEditor/Private/UVEditor2DViewportClient.cpp b/Engine/Plugins/Experimental/UVEditor/Source/UVEditor/Private/UVEditor2DViewportClient.cpp
--- a/Engine/Plugins/Experimental/UVEditor/Source/UVEditor/Private/UVEditor2DViewportClient.cpp
+++ b/Engine/Plugins/Experimental/UVEditor/Source/UVEditor/Private/UVEditor2DViewportClient.cpp
@@ -172,14 +172,14 @@ FUVEditor2DViewportClient::FUVEditor2DViewportClient(FEditorModeTools* InModeToo
 	// numerically) than both the gizmo default and the tool default.
 	static constexpr int DEFAULT_VIEWPORT_BEHAVIOR_PRIORITY = 150;
 
-	ScrollBehaviorTarget = MakeUnique<FUVEditor2DScrollBehaviorTarget>(this);
+	ScrollBehaviorTarget = MakeUnique<FEditor2DScrollBehaviorTarget>(this);
 	UClickDragInputBehavior* ScrollBehavior = NewObject<UClickDragInputBehavior>();
 	ScrollBehavior->Initialize(ScrollBehaviorTarget.Get());
 	ScrollBehavior->SetDefaultPriority(DEFAULT_VIEWPORT_BEHAVIOR_PRIORITY);
 	ScrollBehavior->SetUseRightMouseButton();
 	BehaviorSet->Add(ScrollBehavior);
 
-	ZoomBehaviorTarget = MakeUnique<FUVEditor2DMouseWheelZoomBehaviorTarget>(this);
+	ZoomBehaviorTarget = MakeUnique<FEditor2DMouseWheelZoomBehaviorTarget>(this);
 	ZoomBehaviorTarget->SetCameraFarPlaneWorldZ(FUVEditorUXSettings::CameraFarPlaneWorldZ);
 	ZoomBehaviorTarget->SetCameraNearPlaneProportionZ(FUVEditorUXSettings::CameraNearPlaneProportionZ);
 	ZoomBehaviorTarget->SetZoomLimits(0.001, 100000);
diff --git a/Engine/Plugins/Experimental/UVEditor/Source/UVEditor/Public/UVEditor2DViewportBehaviorTargets.h b/Engine/Plugins/Experimental/UVEditor/Source/UVEditor/Public/UVEditor2DViewportBehaviorTargets.h
deleted file mode 100644
diff --git a/Engine/Plugins/Experimental/UVEditor/Source/UVEditor/Public/UVEditor2DViewportClient.h b/Engine/Plugins/Experimental/UVEditor/Source/UVEditor/Public/UVEditor2DViewportClient.h
--- a/Engine/Plugins/Experimental/UVEditor/Source/UVEditor/Public/UVEditor2DViewportClient.h
+++ b/Engine/Plugins/Experimental/UVEditor/Source/UVEditor/Public/UVEditor2DViewportClient.h
@@ -6,7 +6,7 @@
 
 #include "EditorViewportClient.h"
 #include "InputBehaviorSet.h"
-#include "UVEditor2DViewportBehaviorTargets.h" // FUVEditor2DScrollBehaviorTarget, FUVEditor2DMouseWheelZoomBehaviorTarget
+#include "Behaviors/2DViewportBehaviorTargets.h" // FEditor2DScrollBehaviorTarget, FEditor2DMouseWheelZoomBehaviorTarget
 #include "ContextObjects/UVToolViewportButtonsAPI.h" // UUVToolViewportButtonsAPI::ESelectionMode
 
 class UCanvas;
@@ -79,6 +79,6 @@ protected:
 	// Note that it's generally less hassle if the unique ptr types are complete here,
 	// not forward declared, else we get compile errors if their destruction shows up
 	// anywhere in the header.
-	TUniquePtr<FUVEditor2DScrollBehaviorTarget> ScrollBehaviorTarget;
-	TUniquePtr<FUVEditor2DMouseWheelZoomBehaviorTarget> ZoomBehaviorTarget;
+	TUniquePtr<FEditor2DScrollBehaviorTarget> ScrollBehaviorTarget;
+	TUniquePtr<FEditor2DMouseWheelZoomBehaviorTarget> ZoomBehaviorTarget;
 };
\ No newline at end of file
diff --git a/Engine/Source/Editor/Experimental/EditorInteractiveToolsFramework/Private/Behaviors/2DViewportBehaviorTargets.cpp b/Engine/Source/Editor/Experimental/EditorInteractiveToolsFramework/Private/Behaviors/2DViewportBehaviorTargets.cpp
new file mode 100644
--- /dev/null
+++ b/Engine/Source/Editor/Experimental/EditorInteractiveToolsFramework/Private/Behaviors/2DViewportBehaviorTargets.cpp
@@ -0,0 +1,141 @@
+// Copyright Epic Games, Inc. All Rights Reserved.
+
+#include "Behaviors/2DViewportBehaviorTargets.h"
+#include "EditorViewportClient.h"
+
+FEditor2DScrollBehaviorTarget::FEditor2DScrollBehaviorTarget(FEditorViewportClient* ViewportClientIn)
+	: ViewportClient(ViewportClientIn)
+{
+}
+
+FInputRayHit FEditor2DScrollBehaviorTarget::CanBeginClickDragSequence(const FInputDeviceRay& PressPos)
+{
+	// Verify that ray is facing the proper direction
+	if (PressPos.WorldRay.Direction.Z * PressPos.WorldRay.Origin.Z < 0)
+	{
+		return FInputRayHit(-PressPos.WorldRay.Origin.Z / PressPos.WorldRay.Direction.Z);
+	}
+	return FInputRayHit();
+}
+
+void FEditor2DScrollBehaviorTarget::OnClickPress(const FInputDeviceRay& PressPos)
+{
+	if (ensure(PressPos.WorldRay.Direction.Z * PressPos.WorldRay.Origin.Z < 0))
+	{
+		// Intersect with XY plane
+		double T = -PressPos.WorldRay.Origin.Z / PressPos.WorldRay.Direction.Z;
+		DragStart = FVector3d(
+			PressPos.WorldRay.Origin.X + T * PressPos.WorldRay.Direction.X,
+			PressPos.WorldRay.Origin.Y + T * PressPos.WorldRay.Direction.Y,
+			0);
+
+		OriginalCameraLocation = (FVector3d)ViewportClient->GetViewLocation();
+	}
+}
+
+void FEditor2DScrollBehaviorTarget::OnClickDrag(const FInputDeviceRay& DragPos)
+{
+	if (ensure(DragPos.WorldRay.Direction.Z * DragPos.WorldRay.Origin.Z < 0))
+	{
+		// Intersect a ray starting from the original position and using the new
+		// ray direction. I.e., pretend the camera is not moving.
+
+		double T = -OriginalCameraLocation.Z / DragPos.WorldRay.Direction.Z;
+		FVector3d DragEnd = FVector3d(
+			OriginalCameraLocation.X + T * DragPos.WorldRay.Direction.X,
+			OriginalCameraLocation.Y + T * DragPos.WorldRay.Direction.Y,
+			0);
+
+		// We want to make it look like we are sliding the plane such that DragStart
+		// ends up on DragEnd. For that, our camera will be moving the opposite direction.
+		FVector3d CameraDisplacement = DragStart - DragEnd;
+		check(CameraDisplacement.Z == 0);
+		ViewportClient->SetViewLocation((FVector)(OriginalCameraLocation + CameraDisplacement));
+	}
+}
+
+void FEditor2DScrollBehaviorTarget::OnClickRelease(const FInputDeviceRay& ReleasePos)
+{
+}
+
+void FEditor2DScrollBehaviorTarget::OnTerminateDragSequence()
+{
+}
+
+
+
+// FEditor2DMouseWheelZoomBehaviorTarget
+
+FEditor2DMouseWheelZoomBehaviorTarget::FEditor2DMouseWheelZoomBehaviorTarget(FEditorViewportClient* ViewportClientIn)
+	: ViewportClient(ViewportClientIn)
+{
+	SetZoomAmount(DEFAULT_ZOOM_AMOUNT);
+}
+
+FInputRayHit FEditor2DMouseWheelZoomBehaviorTarget::ShouldRespondToMouseWheel(const FInputDeviceRay& CurrentPos)
+{
+	// Always allowed to zoom with mouse wheel
+	FInputRayHit ToReturn;
+	ToReturn.bHit = true;
+	return ToReturn;
+}
+
+void FEditor2DMouseWheelZoomBehaviorTarget::OnMouseWheelScrollUp(const FInputDeviceRay& CurrentPos)
+{
+	FVector OriginalLocation = ViewportClient->GetViewLocation();
+	double TToPlane = -OriginalLocation.Z / CurrentPos.WorldRay.Direction.Z;
+
+	FVector NewLocation = OriginalLocation + (ZoomInProportion * TToPlane * CurrentPos.WorldRay.Direction);
+
+	ViewportClient->OverrideFarClipPlane(NewLocation.Z - CameraFarPlaneWorldZ);
+	ViewportClient->OverrideNearClipPlane(NewLocation.Z * (1.0-CameraNearPlaneProportionZ));
+
+	// Don't zoom in so far that the XY plane lies in front of our near clipping plane, or else everything
+	// will suddenly disappear.
+	if (NewLocation.Z > ViewportClient->GetNearClipPlane() && NewLocation.Z > ZoomInLimit)
+	{
+		ViewportClient->SetViewLocation(NewLocation);
+	}
+}
+
+void FEditor2DMouseWheelZoomBehaviorTarget::OnMouseWheelScrollDown(const FInputDeviceRay& CurrentPos)
+{
+	FVector OriginalLocation = ViewportClient->GetViewLocation();
+	double TToPlane = -OriginalLocation.Z / CurrentPos.WorldRay.Direction.Z;
+	FVector NewLocation = OriginalLocation - (ZoomOutProportion * TToPlane * CurrentPos.WorldRay.Direction);
+
+	ViewportClient->OverrideFarClipPlane(NewLocation.Z - CameraFarPlaneWorldZ);
+	ViewportClient->OverrideNearClipPlane(NewLocation.Z * (1.0 - CameraNearPlaneProportionZ));
+
+	if (NewLocation.Z < ZoomOutLimit)
+	{
+		ViewportClient->SetViewLocation(NewLocation);
+	}
+}
+
+void FEditor2DMouseWheelZoomBehaviorTarget::SetZoomAmount(double PercentZoomIn)
+{
+	check(PercentZoomIn < 100 && PercentZoomIn >= 0);
+
+	ZoomInProportion = PercentZoomIn / 100;
+
+	// Set the zoom out proportion such that (1 + ZoomOutProportion)(1 - ZoomInProportion) = 1
+	// so that zooming in and then zooming out will return to the same zoom level.
+	ZoomOutProportion = ZoomInProportion / (1 - ZoomInProportion);
+}
+
+void FEditor2DMouseWheelZoomBehaviorTarget::SetZoomLimits(double ZoomInLimitIn, double ZoomOutLimitIn)
+{
+	ZoomInLimit = ZoomInLimitIn;
+	ZoomOutLimit = ZoomOutLimitIn;
+}
+
+void FEditor2DMouseWheelZoomBehaviorTarget::SetCameraFarPlaneWorldZ(double CameraFarPlaneWorldZIn)
+{
+	CameraFarPlaneWorldZ = CameraFarPlaneWorldZIn;
+}
+
+void FEditor2DMouseWheelZoomBehaviorTarget::SetCameraNearPlaneProportionZ(double CameraNearPlaneProportionZIn)
+{
+	CameraNearPlaneProportionZ = CameraNearPlaneProportionZIn;
+}
\ No newline at end of file
diff --git a/Engine/Source/Editor/Experimental/EditorInteractiveToolsFramework/Public/Behaviors/2DViewportBehaviorTargets.h b/Engine/Source/Editor/Experimental/EditorInteractiveToolsFramework/Public/Behaviors/2DViewportBehaviorTargets.h
new file mode 100644
--- /dev/null
+++ b/Engine/Source/Editor/Experimental/EditorInteractiveToolsFramework/Public/Behaviors/2DViewportBehaviorTargets.h
@@ -0,0 +1,76 @@
+// Copyright Epic Games, Inc. All Rights Reserved.
+
+#pragma once
+
+#include "CoreMinimal.h"
+
+#include "BaseBehaviors/BehaviorTargetInterfaces.h"
+#include "VectorTypes.h"
+
+// TODO: Replace with a context object to avoid dependency on FEditorViewportClient (JIRA UE-163677)
+class FEditorViewportClient;
+
+/**
+ * Allows click-dragging to move the camera in the XY plane.
+ */
+class EDITORINTERACTIVETOOLSFRAMEWORK_API FEditor2DScrollBehaviorTarget : public IClickDragBehaviorTarget
+{
+public:
+
+	FEditor2DScrollBehaviorTarget(FEditorViewportClient* ViewportClientIn);
+
+	// IClickDragBehaviorTarget
+	virtual FInputRayHit CanBeginClickDragSequence(const FInputDeviceRay& PressPos) override;
+	virtual void OnClickPress(const FInputDeviceRay& PressPos) override;
+	virtual void OnClickDrag(const FInputDeviceRay& DragPos) override;
+	virtual void OnClickRelease(const FInputDeviceRay& ReleasePos) override;
+	virtual void OnTerminateDragSequence() override;
+
+protected:
+
+	// TODO: Replace with a context object to avoid dependency on FEditorViewportClient (JIRA UE-163677)
+	FEditorViewportClient* ViewportClient = nullptr;
+	FVector3d OriginalCameraLocation;
+	FVector3d DragStart;
+};
+
+/**
+ * Allows the mouse wheel to move the camera forwards/backwards relative to the XY plane,
+ * in the direction pointed to by the mouse.
+ */
+class EDITORINTERACTIVETOOLSFRAMEWORK_API FEditor2DMouseWheelZoomBehaviorTarget : public IMouseWheelBehaviorTarget
+{
+public:
+
+	FEditor2DMouseWheelZoomBehaviorTarget(FEditorViewportClient* ViewportClientIn);
+
+	// IMouseWheelBehaviorTarget
+	virtual FInputRayHit ShouldRespondToMouseWheel(const FInputDeviceRay& CurrentPos);
+	virtual void OnMouseWheelScrollUp(const FInputDeviceRay& CurrentPos);
+	virtual void OnMouseWheelScrollDown(const FInputDeviceRay& CurrentPos);
+	
+	/**
+	 * @param PercentZoomIn How much to move forward on each mouse wheel forward scroll. For instance,
+	 *  passing 20 here will decrease the distance to the XY plane by 20% each time. The zoom out amount
+	 *  will be set in such a way that it undoes the same effect. For instance if 20 was passed here,
+	 *  zoom out amount will be 25% since zooming in by 20% and then out by 25% of the result gets you
+	 *  back to where you were.
+	 */
+	virtual void SetZoomAmount(double PercentZoomIn);
+	void SetZoomLimits(double ZoomInLimitIn, double ZoomOutLimitIn);
+
+	void SetCameraFarPlaneWorldZ(double CameraFarPlaneWorldZIn);
+	void SetCameraNearPlaneProportionZ(double CameraFarPlaneProportionZIn );
+
+	inline static const double DEFAULT_ZOOM_AMOUNT = 20;
+protected:
+
+	// TODO: Replace with a context object to avoid dependency on (JIRA UE-163677)
+	FEditorViewportClient* ViewportClient = nullptr;
+	double ZoomInProportion;
+	double ZoomOutProportion;
+	double ZoomInLimit;
+	double ZoomOutLimit;
+	double CameraFarPlaneWorldZ;
+	double CameraNearPlaneProportionZ;
+};
\ No newline at end of file
