Fix a crash in modes and mode manager when iterating over all objects and attempting to get the world.

# Although there aren't references left to the object, iterating over all the things will return the thing. Flagging as garbage seems a lot testier of a change for fixing the 5.1 crash with varied lifettime of the mode and mode manager contexts, so merely tracking it and reporting the correct world here.

#rnx

diff --git a/Engine/Source/Editor/Experimental/EditorInteractiveToolsFramework/Private/EdModeInteractiveToolsContext.cpp b/Engine/Source/Editor/Experimental/EditorInteractiveToolsFramework/Private/EdModeInteractiveToolsContext.cpp
--- a/Engine/Source/Editor/Experimental/EditorInteractiveToolsFramework/Private/EdModeInteractiveToolsContext.cpp
+++ b/Engine/Source/Editor/Experimental/EditorInteractiveToolsFramework/Private/EdModeInteractiveToolsContext.cpp
@@ -296,12 +296,15 @@ void UEditorInteractiveToolsContext::Initialize(IToolsContextQueriesAPI* Queries
 
 void UEditorInteractiveToolsContext::Shutdown()
 {
+	bIsActive = false;
+
 	// auto-accept any in-progress tools
 	DeactivateAllActiveTools(EToolShutdownType::Accept);
 
 	UInteractiveToolsContext::Shutdown();
 }
 
+
 void UEditorInteractiveToolsContext::InitializeContextWithEditorModeManager(FEditorModeTools* InEditorModeManager, UInputRouter* UseInputRouter)
 {
 	check(InEditorModeManager);
@@ -392,7 +395,7 @@ void UEditorInteractiveToolsContext::PostInvalidation()
 
 UWorld* UEditorInteractiveToolsContext::GetWorld() const
 {
-	if (EditorModeManager)
+	if (bIsActive && EditorModeManager)
 	{
 		return EditorModeManager->GetWorld();
 	}
@@ -726,6 +729,15 @@ void UEditorInteractiveToolsContext::EndTool(EToolShutdownType ShutdownType)
 	PostInvalidation();
 }
 
+void UEditorInteractiveToolsContext::Activate()
+{
+	bIsActive = true;
+}
+
+void UEditorInteractiveToolsContext::Deactivate()
+{
+	bIsActive = false;
+}
 
 
 void UEditorInteractiveToolsContext::DeactivateActiveTool(EToolSide WhichSide, EToolShutdownType ShutdownType)
@@ -1172,6 +1184,7 @@ bool UModeManagerInteractiveToolsContext::OnChildEdModeActivated(UEdModeInteract
 	}
 
 	EdModeToolsContexts.Add(ChildToolsContext);
+	ChildToolsContext->Activate();
 	return true;
 }
 
@@ -1183,6 +1196,7 @@ bool UModeManagerInteractiveToolsContext::OnChildEdModeDeactivated(UEdModeIntera
 		if (EdModeToolsContexts[k] == ChildToolsContext)
 		{
 			ChildToolsContext->DeactivateAllActiveTools(EToolShutdownType::Cancel);
+			ChildToolsContext->Deactivate();
 			EdModeToolsContexts.RemoveAt(k);
 			return true;
 		}
diff --git a/Engine/Source/Editor/Experimental/EditorInteractiveToolsFramework/Public/EdModeInteractiveToolsContext.h b/Engine/Source/Editor/Experimental/EditorInteractiveToolsFramework/Public/EdModeInteractiveToolsContext.h
--- a/Engine/Source/Editor/Experimental/EditorInteractiveToolsFramework/Public/EdModeInteractiveToolsContext.h
+++ b/Engine/Source/Editor/Experimental/EditorInteractiveToolsFramework/Public/EdModeInteractiveToolsContext.h
@@ -124,6 +124,8 @@ public:
 	virtual bool CanCompleteActiveTool() const;
 	virtual void StartTool(const FString ToolTypeIdentifier);
 	virtual void EndTool(EToolShutdownType ShutdownType);
+	void Activate();
+	void Deactivate();
 
 
 	/** @return Ray into 3D scene at last mouse event */
@@ -210,6 +212,8 @@ private:
 
 	// currently defaulting to enabled as FEdModes generally assume this, and in most cases hitproxy pass is not expensive.
 	bool bEnableRenderingDuringHitProxyPass = true;
+
+	bool bIsActive = false;
 };
 
 
diff --git a/Engine/Source/Editor/UnrealEd/Private/EditorModeManager.cpp b/Engine/Source/Editor/UnrealEd/Private/EditorModeManager.cpp
--- a/Engine/Source/Editor/UnrealEd/Private/EditorModeManager.cpp
+++ b/Engine/Source/Editor/UnrealEd/Private/EditorModeManager.cpp
@@ -78,6 +78,7 @@ FEditorModeTools::FEditorModeTools()
 
 	InteractiveToolsContext = NewObject<UModeManagerInteractiveToolsContext>(GetTransientPackage(), UModeManagerInteractiveToolsContext::StaticClass(), NAME_None, RF_Transient);
 	InteractiveToolsContext->InitializeContextWithEditorModeManager(this);
+	InteractiveToolsContext->Activate();
 
 	// Load the last used settings
 	LoadConfig();
@@ -113,6 +114,7 @@ FEditorModeTools::~FEditorModeTools()
 	// which would mean that this instances will be garbage
 	if (UObjectInitialized())
 	{
+		InteractiveToolsContext->Deactivate();
 		InteractiveToolsContext->ShutdownContext();
 		InteractiveToolsContext = nullptr;
 	}
diff --git a/Engine/Source/Editor/UnrealEd/Private/Tools/AssetEditorContextObject.cpp b/Engine/Source/Editor/UnrealEd/Private/Tools/AssetEditorContextObject.cpp
--- a/Engine/Source/Editor/UnrealEd/Private/Tools/AssetEditorContextObject.cpp
+++ b/Engine/Source/Editor/UnrealEd/Private/Tools/AssetEditorContextObject.cpp
@@ -39,12 +39,7 @@ UTypedElementCommonActions* UAssetEditorContextObject::GetCommonActions()
 
 UWorld* UAssetEditorContextObject::GetEditingWorld() const
 {
-	if (GetToolkitHost())
-	{
-		return GetToolkitHost()->GetWorld();
-	}
-
-	return nullptr;
+	return GetWorld();
 }
 
 const IToolkitHost* UAssetEditorContextObject::GetToolkitHost() const
@@ -61,10 +56,3 @@ void UAssetEditorContextObject::SetToolkitHost(IToolkitHost* InToolkitHost)
 {
 	ToolkitHost = InToolkitHost;
 }
-
-UWorld* UAssetEditorContextObject::GetWorld() const
-{
-	return GetEditingWorld();
-}
-
-
diff --git a/Engine/Source/Editor/UnrealEd/Private/Tools/AssetEditorContextObject.h b/Engine/Source/Editor/UnrealEd/Private/Tools/AssetEditorContextObject.h
--- a/Engine/Source/Editor/UnrealEd/Private/Tools/AssetEditorContextObject.h
+++ b/Engine/Source/Editor/UnrealEd/Private/Tools/AssetEditorContextObject.h
@@ -27,10 +27,6 @@ public:
 	 */
 	void SetToolkitHost(IToolkitHost* InToolkitHost);
 
-	// Begin UObject interface
-	UWorld* GetWorld() const override;
-	// End UObject interface
-
 private:
 	/** The toolkit host associated with our asset editor. */
 	IToolkitHost* ToolkitHost = nullptr;
diff --git a/Engine/Source/Editor/UnrealEd/Private/Tools/UEdMode.cpp b/Engine/Source/Editor/UnrealEd/Private/Tools/UEdMode.cpp
--- a/Engine/Source/Editor/UnrealEd/Private/Tools/UEdMode.cpp
+++ b/Engine/Source/Editor/UnrealEd/Private/Tools/UEdMode.cpp
@@ -172,17 +172,21 @@ bool UEdMode::UsesToolkits() const
 
 UWorld* UEdMode::GetWorld() const
 {
-	return Owner->GetWorld();
+	return EditorToolsContext.IsValid() ? EditorToolsContext->GetWorld() : nullptr;
 }
 
-class FEditorModeTools* UEdMode::GetModeManager() const
+FEditorModeTools* UEdMode::GetModeManager() const
 {
-	return Owner;
+	return EditorToolsContext.IsValid() ? EditorToolsContext->GetParentEditorModeManager() : nullptr;
 }
 
 AActor* UEdMode::GetFirstSelectedActorInstance() const
 {
-	return Owner->GetEditorSelectionSet()->GetTopSelectedObject<AActor>();
+	if (FEditorModeTools* ModeManager = GetModeManager())
+	{
+		return ModeManager->GetEditorSelectionSet()->GetTopSelectedObject<AActor>();
+	}
+	return nullptr;
 }
 
 UInteractiveToolManager* UEdMode::GetToolManager(EToolsContextScope ToolScope) const
@@ -240,8 +244,7 @@ void UEdMode::OnModeActivated(const FEditorModeID& InID, bool bIsActive)
 	{
 		if (bIsActive)
 		{
-			Owner->GetInteractiveToolsContext()->OnChildEdModeActivated(ModeToolsContext);
-
+			EditorToolsContext->OnChildEdModeActivated(ModeToolsContext);
 			EditorToolsContext->ToolManager->OnToolStarted.AddUObject(this, &UEdMode::OnToolStarted);
 			EditorToolsContext->ToolManager->OnToolEnded.AddUObject(this, &UEdMode::OnToolEnded);
 		}
@@ -249,8 +252,7 @@ void UEdMode::OnModeActivated(const FEditorModeID& InID, bool bIsActive)
 		{
 			EditorToolsContext->ToolManager->OnToolStarted.RemoveAll(this);
 			EditorToolsContext->ToolManager->OnToolEnded.RemoveAll(this);
-
-			Owner->GetInteractiveToolsContext()->OnChildEdModeDeactivated(ModeToolsContext);
+			EditorToolsContext->OnChildEdModeDeactivated(ModeToolsContext);
 		}
 	}
 }
