Skeletal Mesh Gizmo customization
    - added support for material and size customization

diff --git a/Engine/Plugins/Experimental/Animation/SkeletalMeshModelingTools/Source/SkeletalMeshModelingTools/Private/SkeletalMeshGizmoUtils.cpp b/Engine/Plugins/Experimental/Animation/SkeletalMeshModelingTools/Source/SkeletalMeshModelingTools/Private/SkeletalMeshGizmoUtils.cpp
--- a/Engine/Plugins/Experimental/Animation/SkeletalMeshModelingTools/Source/SkeletalMeshModelingTools/Private/SkeletalMeshGizmoUtils.cpp
+++ b/Engine/Plugins/Experimental/Animation/SkeletalMeshModelingTools/Source/SkeletalMeshModelingTools/Private/SkeletalMeshGizmoUtils.cpp
@@ -9,7 +9,9 @@
 
 #include "EditorGizmos/TransformGizmo.h"
 #include "Editor/Experimental/EditorInteractiveToolsFramework/Public/EditorInteractiveGizmoManager.h"
-// #include "EditorGizmos/EditorTransformGizmoBuilder.h"
+#include "EditorGizmos/EditorTransformGizmoBuilder.h"
+
+#include "Materials/Material.h"
 
 #include UE_INLINE_GENERATED_CPP_BY_NAME(SkeletalMeshGizmoUtils)
 
@@ -94,7 +96,9 @@ void USkeletalMeshGizmoContextObject::RegisterGizmosWithManager(UInteractiveTool
 	UEditorInteractiveGizmoManager* GizmoManager = Cast<UEditorInteractiveGizmoManager>(InToolManager->GetPairedGizmoManager());
 	if (ensure(GizmoManager))
 	{
-		GizmoManager->RegisterGizmoType(TransformBuilderIdentifier(), GizmoManager->GetTransformGizmoBuilder());
+		UEditorTransformGizmoBuilder* GizmoBuilder = NewObject<UEditorTransformGizmoBuilder>(this);
+		GizmoBuilder->CustomizationFunction = [] { return GetGizmoCustomization(); };
+		GizmoManager->RegisterGizmoType(TransformBuilderIdentifier(), GizmoBuilder);
 	}
 	
 	bGizmosRegistered = true;
@@ -145,6 +149,20 @@ USkeletalMeshGizmoWrapperBase* USkeletalMeshGizmoContextObject::GetNewWrapper(UI
 	return GizmoWrapper;
 }
 
+const FGizmoCustomization& USkeletalMeshGizmoContextObject::GetGizmoCustomization()
+{
+	static const FString MaterialName = TEXT("/SkeletalMeshModelingTools/SkeletalMeshGizmoMaterial.SkeletalMeshGizmoMaterial");
+	UMaterial* Material = LoadObject<UMaterial>(nullptr, *MaterialName, nullptr, LOAD_None, nullptr);
+
+	check(Material);
+	
+	static FGizmoCustomization Customization;
+		Customization.Material = Material;
+		Customization.SizeCoefficient = 1.5f;
+
+	return Customization;
+}
+
 bool USkeletalMeshGizmoWrapper::IsGizmoHit(const FInputDeviceRay& PressPos) const
 {
 	if (!TransformGizmo || !TransformGizmo->bVisible)
diff --git a/Engine/Plugins/Experimental/Animation/SkeletalMeshModelingTools/Source/SkeletalMeshModelingTools/Private/SkeletalMeshGizmoUtils.h b/Engine/Plugins/Experimental/Animation/SkeletalMeshModelingTools/Source/SkeletalMeshModelingTools/Private/SkeletalMeshGizmoUtils.h
--- a/Engine/Plugins/Experimental/Animation/SkeletalMeshModelingTools/Source/SkeletalMeshModelingTools/Private/SkeletalMeshGizmoUtils.h
+++ b/Engine/Plugins/Experimental/Animation/SkeletalMeshModelingTools/Source/SkeletalMeshModelingTools/Private/SkeletalMeshGizmoUtils.h
@@ -12,6 +12,7 @@ class UInteractiveToolsContext;
 class UInteractiveToolManager;
 class UTransformGizmo;
 class USkeletonTransformProxy;
+struct FGizmoCustomization;
 
 namespace UE
 {
@@ -87,5 +88,8 @@ public:
 	virtual USkeletalMeshGizmoWrapperBase* GetNewWrapper(UInteractiveToolManager* InToolManager, UObject* Outer = nullptr, IGizmoStateTarget* InStateTarget = nullptr) override;
 
 private:
+	
+	static const FGizmoCustomization& GetGizmoCustomization();
+	
 	bool bGizmosRegistered = false;
 };
diff --git a/Engine/Source/Editor/Experimental/EditorInteractiveToolsFramework/Private/EditorGizmos/EditorTransformGizmoBuilder.cpp b/Engine/Source/Editor/Experimental/EditorInteractiveToolsFramework/Private/EditorGizmos/EditorTransformGizmoBuilder.cpp
--- a/Engine/Source/Editor/Experimental/EditorInteractiveToolsFramework/Private/EditorGizmos/EditorTransformGizmoBuilder.cpp
+++ b/Engine/Source/Editor/Experimental/EditorInteractiveToolsFramework/Private/EditorGizmos/EditorTransformGizmoBuilder.cpp
@@ -10,7 +10,6 @@
 #include "EditorGizmos/EditorTransformGizmoSource.h"
 #include "EditorGizmos/EditorTransformProxy.h"
 #include "EditorGizmos/TransformGizmo.h"
-#include "EditorGizmos/TransformGizmoInterfaces.h"
 #include "InteractiveGizmo.h"
 #include "InteractiveGizmoManager.h"
 #include "Templates/Casts.h"
@@ -25,6 +24,7 @@ UInteractiveGizmo* UEditorTransformGizmoBuilder::BuildGizmo(const FToolBuilderSt
 	check(GizmoViewContext && GizmoViewContext->IsValidLowLevel());
 
 	UEditorTransformGizmo* TransformGizmo = NewObject<UEditorTransformGizmo>(SceneState.GizmoManager);
+	TransformGizmo->SetCustomizationFunction(CustomizationFunction);
 	TransformGizmo->Setup();
 	TransformGizmo->TransformGizmoSource = UEditorTransformGizmoSource::Construct(TransformGizmo);
 	TransformGizmo->GizmoViewContext = GizmoViewContext;
diff --git a/Engine/Source/Editor/Experimental/EditorInteractiveToolsFramework/Private/EditorGizmos/TransformGizmo.cpp b/Engine/Source/Editor/Experimental/EditorInteractiveToolsFramework/Private/EditorGizmos/TransformGizmo.cpp
--- a/Engine/Source/Editor/Experimental/EditorInteractiveToolsFramework/Private/EditorGizmos/TransformGizmo.cpp
+++ b/Engine/Source/Editor/Experimental/EditorInteractiveToolsFramework/Private/EditorGizmos/TransformGizmo.cpp
@@ -71,7 +71,17 @@ void UTransformGizmo::SetupBehaviors()
 
 void UTransformGizmo::SetupMaterials()
 {
-	UMaterial* AxisMaterialBase = GEngine->ArrowMaterial;
+	auto GetBaseMaterial = [this]()
+	{
+		if (CustomizationFunction)
+		{
+			const FGizmoCustomization& GizmoCustomization = CustomizationFunction();
+			return GizmoCustomization.Material ? GizmoCustomization.Material : GEngine->ArrowMaterial; 
+		}
+		return GEngine->ArrowMaterial;
+	};
+	
+	UMaterial* AxisMaterialBase = GetBaseMaterial(); 
 
 	AxisMaterialX = UMaterialInstanceDynamic::Create(AxisMaterialBase, NULL);
 	AxisMaterialX->SetVectorParameterValue("GizmoColor", AxisColorX);
@@ -656,18 +666,25 @@ void UTransformGizmo::SetVisibility(bool bVisibleIn)
 	bVisible = bVisibleIn;
 }
 
+void UTransformGizmo::SetCustomizationFunction(const TFunction<const FGizmoCustomization()>& InFunction)
+{
+	CustomizationFunction = InFunction;
+}
+
 UGizmoElementArrow* UTransformGizmo::MakeTranslateAxis(ETransformGizmoPartIdentifier InPartId, const FVector& InAxisDir, const FVector& InSideDir, UMaterialInterface* InMaterial)
 {
+	const float SizeCoeff = CustomizationFunction ? CustomizationFunction().SizeCoefficient : 1.f;
+	
 	UGizmoElementArrow* ArrowElement = NewObject<UGizmoElementArrow>();
 	ArrowElement->SetPartIdentifier(static_cast<uint32>(InPartId));
 	ArrowElement->SetHeadType(EGizmoElementArrowHeadType::Cone);
 	ArrowElement->SetBase(InAxisDir * AxisLengthOffset);
 	ArrowElement->SetDirection(InAxisDir);
 	ArrowElement->SetSideDirection(InSideDir);
-	ArrowElement->SetBodyLength(TranslateAxisLength);
+	ArrowElement->SetBodyLength(TranslateAxisLength * SizeCoeff);
 	ArrowElement->SetBodyRadius(AxisRadius);
-	ArrowElement->SetHeadLength(TranslateAxisConeHeight);
-	ArrowElement->SetHeadRadius(TranslateAxisConeRadius);
+	ArrowElement->SetHeadLength(TranslateAxisConeHeight * SizeCoeff);
+	ArrowElement->SetHeadRadius(TranslateAxisConeRadius * SizeCoeff);
 	ArrowElement->SetNumSides(32);
 	ArrowElement->SetMaterial(InMaterial);
 	ArrowElement->SetViewDependentType(EGizmoElementViewDependentType::Axis);
@@ -677,15 +694,17 @@ UGizmoElementArrow* UTransformGizmo::MakeTranslateAxis(ETransformGizmoPartIdenti
 
 UGizmoElementArrow* UTransformGizmo::MakeScaleAxis(ETransformGizmoPartIdentifier InPartId, const FVector& InAxisDir, const FVector& InSideDir, UMaterialInterface* InMaterial)
 {
+	const float SizeCoeff = CustomizationFunction ? CustomizationFunction().SizeCoefficient : 1.f;
+
 	UGizmoElementArrow* ArrowElement = NewObject<UGizmoElementArrow>();
 	ArrowElement->SetPartIdentifier(static_cast<uint32>(InPartId));
 	ArrowElement->SetHeadType(EGizmoElementArrowHeadType::Cube);
 	ArrowElement->SetBase(InAxisDir * AxisLengthOffset);
 	ArrowElement->SetDirection(InAxisDir);
 	ArrowElement->SetSideDirection(InSideDir);
-	ArrowElement->SetBodyLength(ScaleAxisLength);
+	ArrowElement->SetBodyLength(ScaleAxisLength * SizeCoeff);
 	ArrowElement->SetBodyRadius(AxisRadius);
-	ArrowElement->SetHeadLength(ScaleAxisCubeDim);
+	ArrowElement->SetHeadLength(ScaleAxisCubeDim * SizeCoeff);
 	ArrowElement->SetNumSides(32);
 	ArrowElement->SetMaterial(InMaterial);
 	ArrowElement->SetViewDependentType(EGizmoElementViewDependentType::Axis);
@@ -695,12 +714,14 @@ UGizmoElementArrow* UTransformGizmo::MakeScaleAxis(ETransformGizmoPartIdentifier
 
 UGizmoElementBox* UTransformGizmo::MakeUniformScaleHandle()
 {
+	const float SizeCoeff = CustomizationFunction ? CustomizationFunction().SizeCoefficient : 1.f;
+	
 	UGizmoElementBox* BoxElement = NewObject<UGizmoElementBox>();
 	BoxElement->SetPartIdentifier(static_cast<uint32>(ETransformGizmoPartIdentifier::ScaleUniform));
 	BoxElement->SetCenter(FVector::ZeroVector);
 	BoxElement->SetUpDirection(FVector::UpVector);
 	BoxElement->SetSideDirection(FVector::RightVector);
-	BoxElement->SetDimensions(FVector(ScaleAxisCubeDim, ScaleAxisCubeDim, ScaleAxisCubeDim));
+	BoxElement->SetDimensions(FVector(ScaleAxisCubeDim, ScaleAxisCubeDim, ScaleAxisCubeDim) * SizeCoeff);
 	BoxElement->SetMaterial(GreyMaterial);
 	return BoxElement;
 }
@@ -708,7 +729,9 @@ UGizmoElementBox* UTransformGizmo::MakeUniformScaleHandle()
 UGizmoElementRectangle* UTransformGizmo::MakePlanarHandle(ETransformGizmoPartIdentifier InPartId, const FVector& InUpDirection, const FVector& InSideDirection, const FVector& InPlaneNormal,
 	UMaterialInterface* InMaterial, const FLinearColor& InVertexColor)
 {
-	FVector PlanarHandleCenter = (InUpDirection + InSideDirection) * PlanarHandleOffset;
+	const float SizeCoeff = CustomizationFunction ? CustomizationFunction().SizeCoefficient : 1.f;
+	
+	FVector PlanarHandleCenter = (InUpDirection + InSideDirection) * PlanarHandleOffset * SizeCoeff;
 
 	FLinearColor LineColor = InVertexColor;
 	FLinearColor VertexColor = LineColor;
@@ -719,8 +742,8 @@ UGizmoElementRectangle* UTransformGizmo::MakePlanarHandle(ETransformGizmoPartIde
 	RectangleElement->SetUpDirection(InUpDirection);
 	RectangleElement->SetSideDirection(InSideDirection);
 	RectangleElement->SetCenter(PlanarHandleCenter);
-	RectangleElement->SetHeight(PlanarHandleSize);
-	RectangleElement->SetWidth(PlanarHandleSize);
+	RectangleElement->SetHeight(PlanarHandleSize * SizeCoeff);
+	RectangleElement->SetWidth(PlanarHandleSize * SizeCoeff);
 	RectangleElement->SetMaterial(InMaterial);
 	RectangleElement->SetVertexColor(VertexColor);
 	RectangleElement->SetLineColor(LineColor);
@@ -734,13 +757,15 @@ UGizmoElementRectangle* UTransformGizmo::MakePlanarHandle(ETransformGizmoPartIde
 
 UGizmoElementRectangle* UTransformGizmo::MakeTranslateScreenSpaceHandle()
 {
+	const float SizeCoeff = CustomizationFunction ? CustomizationFunction().SizeCoefficient : 1.f;
+	
 	UGizmoElementRectangle* RectangleElement = NewObject<UGizmoElementRectangle>();
 	RectangleElement->SetPartIdentifier(static_cast<uint32>(ETransformGizmoPartIdentifier::TranslateScreenSpace));
 	RectangleElement->SetUpDirection(FVector::UpVector);
 	RectangleElement->SetSideDirection(FVector::RightVector);
 	RectangleElement->SetCenter(FVector::ZeroVector);
-	RectangleElement->SetHeight(TranslateScreenSpaceHandleSize);
-	RectangleElement->SetWidth(TranslateScreenSpaceHandleSize);
+	RectangleElement->SetHeight(TranslateScreenSpaceHandleSize * SizeCoeff);
+	RectangleElement->SetWidth(TranslateScreenSpaceHandleSize * SizeCoeff);
 	RectangleElement->SetViewAlignType(EGizmoElementViewAlignType::PointScreen);
 	RectangleElement->SetViewAlignAxis(FVector::UpVector);
 	RectangleElement->SetViewAlignNormal(-FVector::ForwardVector);
@@ -757,12 +782,14 @@ UGizmoElementRectangle* UTransformGizmo::MakeTranslateScreenSpaceHandle()
 UGizmoElementTorus* UTransformGizmo::MakeRotateAxis(ETransformGizmoPartIdentifier InPartId, const FVector& TorusAxis0, const FVector& TorusAxis1,
 	UMaterialInterface* InMaterial, UMaterialInterface* InCurrentMaterial)
 {
+	const float SizeCoeff = CustomizationFunction ? CustomizationFunction().SizeCoefficient : 1.f;
+	
 	UGizmoElementTorus* RotateAxisElement = NewObject<UGizmoElementTorus>();
 	RotateAxisElement->SetPartIdentifier(static_cast<uint32>(InPartId));
 	RotateAxisElement->SetCenter(FVector::ZeroVector);
-	RotateAxisElement->SetRadius(UTransformGizmo::RotateAxisOuterRadius);
+	RotateAxisElement->SetRadius(UTransformGizmo::RotateAxisOuterRadius * SizeCoeff);
 	RotateAxisElement->SetNumSegments(UTransformGizmo::RotateAxisNumSegments);
-	RotateAxisElement->SetInnerRadius(UTransformGizmo::RotateAxisInnerRadius);
+	RotateAxisElement->SetInnerRadius(UTransformGizmo::RotateAxisInnerRadius * SizeCoeff);
 	RotateAxisElement->SetNumInnerSlices(UTransformGizmo::RotateAxisInnerSlices);
 	RotateAxisElement->SetAxis0(TorusAxis0);
 	RotateAxisElement->SetAxis1(TorusAxis1);
@@ -780,10 +807,12 @@ UGizmoElementTorus* UTransformGizmo::MakeRotateAxis(ETransformGizmoPartIdentifie
 
 UGizmoElementCircle* UTransformGizmo::MakeRotateCircleHandle(ETransformGizmoPartIdentifier InPartId, float InRadius, const FLinearColor& InColor, float bFill)
 {
+	const float SizeCoeff = CustomizationFunction ? CustomizationFunction().SizeCoefficient : 1.f;
+	
 	UGizmoElementCircle* CircleElement = NewObject<UGizmoElementCircle>();
 	CircleElement->SetPartIdentifier(static_cast<uint32>(InPartId));
 	CircleElement->SetCenter(FVector::ZeroVector);
-	CircleElement->SetRadius(InRadius);
+	CircleElement->SetRadius(InRadius * SizeCoeff);
 	CircleElement->SetAxis0(FVector::UpVector);
 	CircleElement->SetAxis1(-FVector::RightVector);
 	CircleElement->SetLineColor(InColor);
diff --git a/Engine/Source/Editor/Experimental/EditorInteractiveToolsFramework/Public/EditorGizmos/EditorTransformGizmoBuilder.h b/Engine/Source/Editor/Experimental/EditorInteractiveToolsFramework/Public/EditorGizmos/EditorTransformGizmoBuilder.h
--- a/Engine/Source/Editor/Experimental/EditorInteractiveToolsFramework/Public/EditorGizmos/EditorTransformGizmoBuilder.h
+++ b/Engine/Source/Editor/Experimental/EditorInteractiveToolsFramework/Public/EditorGizmos/EditorTransformGizmoBuilder.h
@@ -8,6 +8,7 @@
 #include "InteractiveGizmoBuilder.h"
 #include "UObject/ObjectMacros.h"
 #include "UObject/UObjectGlobals.h"
+#include "EditorGizmos/TransformGizmoInterfaces.h"
 
 #include "EditorTransformGizmoBuilder.generated.h"
 
@@ -25,6 +26,9 @@ public:
 	// UEditorInteractiveGizmoSelectionBuilder interface 
 	virtual UInteractiveGizmo* BuildGizmo(const FToolBuilderState& SceneState) const override;
 	virtual void UpdateGizmoForSelection(UInteractiveGizmo* Gizmo, const FToolBuilderState& SceneState) override;
+
+	// If set, this function will be passed to UTransformGizmo instances to override the default material and display size.
+	TFunction<const FGizmoCustomization()> CustomizationFunction;
 };
 
 
diff --git a/Engine/Source/Editor/Experimental/EditorInteractiveToolsFramework/Public/EditorGizmos/TransformGizmo.h b/Engine/Source/Editor/Experimental/EditorInteractiveToolsFramework/Public/EditorGizmos/TransformGizmo.h
--- a/Engine/Source/Editor/Experimental/EditorInteractiveToolsFramework/Public/EditorGizmos/TransformGizmo.h
+++ b/Engine/Source/Editor/Experimental/EditorInteractiveToolsFramework/Public/EditorGizmos/TransformGizmo.h
@@ -205,6 +205,11 @@ public:
 	 */
 	virtual void SetVisibility(bool bVisible);
 
+	/**
+	 * Set customization function for this Gizmo
+	 */
+	void SetCustomizationFunction(const TFunction<const FGizmoCustomization()>& InFunction);
+	
 public:
 
 	/** The active target object for the Gizmo */
@@ -671,6 +676,9 @@ protected:
 	/** Array of function pointers, indexed by gizmo part id, to handle update interacting state */
 	TArray<TFunction<void(UTransformGizmo* TransformGizmo, bool bInInteracting, uint32 InHitPart)> > OnUpdateInteractingFunctions;
 
+	/** Customization function (to override default material or increment gizmo size for example) */
+	TFunction<const FGizmoCustomization()> CustomizationFunction;
+
 	/** Percentage-based scale multiplier */
 	UPROPERTY()
 	double ScaleMultiplier = 0.01;
diff --git a/Engine/Source/Editor/Experimental/EditorInteractiveToolsFramework/Public/EditorGizmos/TransformGizmoInterfaces.h b/Engine/Source/Editor/Experimental/EditorInteractiveToolsFramework/Public/EditorGizmos/TransformGizmoInterfaces.h
--- a/Engine/Source/Editor/Experimental/EditorInteractiveToolsFramework/Public/EditorGizmos/TransformGizmoInterfaces.h
+++ b/Engine/Source/Editor/Experimental/EditorInteractiveToolsFramework/Public/EditorGizmos/TransformGizmoInterfaces.h
@@ -7,6 +7,8 @@
 #include "UnrealWidgetFwd.h"
 #include "TransformGizmoInterfaces.generated.h"
 
+class UMaterial;
+
 //
 // Interface for the Transform gizmo.
 //
@@ -71,4 +73,12 @@ public:
 	virtual EGizmoTransformScaleType GetScaleType() const PURE_VIRTUAL(ITransformGizmoSource::GetScaleType, return EGizmoTransformScaleType::Default;);
 };
 
+/**
+ * FGizmoCustomization is a struct that can be used to make some display overrides (currently material and size) 
+ */
 
+struct FGizmoCustomization
+{
+	TObjectPtr<UMaterial> Material = nullptr;
+	float SizeCoefficient = 1.f;
+};
