Editor2DMouseWheelZoomBehaviorTarget: make the clip plane math work out if the camera has negative z coordinates

ViewportClient's NearPlane value is assumed to be positive, which makes things janky if your camera has negative z value. This change uses the camera's absolute Z coordinate for computing the temporary clip planes, and also testing if the XY-plane is on the wrong size of our near clipping plane.

diff --git a/Engine/Source/Editor/Experimental/EditorInteractiveToolsFramework/Private/Behaviors/2DViewportBehaviorTargets.cpp b/Engine/Source/Editor/Experimental/EditorInteractiveToolsFramework/Private/Behaviors/2DViewportBehaviorTargets.cpp
--- a/Engine/Source/Editor/Experimental/EditorInteractiveToolsFramework/Private/Behaviors/2DViewportBehaviorTargets.cpp
+++ b/Engine/Source/Editor/Experimental/EditorInteractiveToolsFramework/Private/Behaviors/2DViewportBehaviorTargets.cpp
@@ -82,17 +82,18 @@ FInputRayHit FEditor2DMouseWheelZoomBehaviorTarget::ShouldRespondToMouseWheel(co
 
 void FEditor2DMouseWheelZoomBehaviorTarget::OnMouseWheelScrollUp(const FInputDeviceRay& CurrentPos)
 {
-	FVector OriginalLocation = ViewportClient->GetViewLocation();
-	double TToPlane = -OriginalLocation.Z / CurrentPos.WorldRay.Direction.Z;
+	const FVector OriginalLocation = ViewportClient->GetViewLocation();
+	// TODO: These behaviors should support other planes, not just XY. Using Abs(Z) allows us to at least position the camera on either size of the XY plane.
+	const double TToPlane = -OriginalLocation.Z / CurrentPos.WorldRay.Direction.Z;
+	const FVector NewLocation = OriginalLocation + (ZoomInProportion * TToPlane * CurrentPos.WorldRay.Direction);
+	const double AbsZ = FMath::Abs(NewLocation.Z);
 
-	FVector NewLocation = OriginalLocation + (ZoomInProportion * TToPlane * CurrentPos.WorldRay.Direction);
-
-	ViewportClient->OverrideFarClipPlane(static_cast<float>(NewLocation.Z - CameraFarPlaneWorldZ));
-	ViewportClient->OverrideNearClipPlane(static_cast<float>(NewLocation.Z * (1.0-CameraNearPlaneProportionZ)));
+	ViewportClient->OverrideFarClipPlane(static_cast<float>(AbsZ - CameraFarPlaneWorldZ));
+	ViewportClient->OverrideNearClipPlane(static_cast<float>(AbsZ * (1.0-CameraNearPlaneProportionZ)));
 
 	// Don't zoom in so far that the XY plane lies in front of our near clipping plane, or else everything
 	// will suddenly disappear.
-	if (NewLocation.Z > ViewportClient->GetNearClipPlane() && NewLocation.Z > ZoomInLimit)
+	if (AbsZ > ViewportClient->GetNearClipPlane() && AbsZ > ZoomInLimit)
 	{
 		ViewportClient->SetViewLocation(NewLocation);
 	}
@@ -100,14 +101,16 @@ void FEditor2DMouseWheelZoomBehaviorTarget::OnMouseWheelScrollUp(const FInputDev
 
 void FEditor2DMouseWheelZoomBehaviorTarget::OnMouseWheelScrollDown(const FInputDeviceRay& CurrentPos)
 {
-	FVector OriginalLocation = ViewportClient->GetViewLocation();
-	double TToPlane = -OriginalLocation.Z / CurrentPos.WorldRay.Direction.Z;
-	FVector NewLocation = OriginalLocation - (ZoomOutProportion * TToPlane * CurrentPos.WorldRay.Direction);
+	const FVector OriginalLocation = ViewportClient->GetViewLocation();
+	// TODO: These behaviors should support other planes, not just XY. Using Abs(Z) allows us to at least position the camera on either size of the XY plane.
+	const double TToPlane = -OriginalLocation.Z / CurrentPos.WorldRay.Direction.Z;
+	const FVector NewLocation = OriginalLocation - (ZoomOutProportion * TToPlane * CurrentPos.WorldRay.Direction);
+	const double AbsZ = FMath::Abs(NewLocation.Z);
 
-	ViewportClient->OverrideFarClipPlane(static_cast<float>(NewLocation.Z - CameraFarPlaneWorldZ));
-	ViewportClient->OverrideNearClipPlane(static_cast<float>(NewLocation.Z * (1.0 - CameraNearPlaneProportionZ)));
+	ViewportClient->OverrideFarClipPlane(static_cast<float>(AbsZ - CameraFarPlaneWorldZ));
+	ViewportClient->OverrideNearClipPlane(static_cast<float>(AbsZ * (1.0 - CameraNearPlaneProportionZ)));
 
-	if (NewLocation.Z < ZoomOutLimit)
+	if (AbsZ < ZoomOutLimit)
 	{
 		ViewportClient->SetViewLocation(NewLocation);
 	}
