New TRS Gizmo updates
    - arcball rotation: the rotation is computed using a spherical+hyperbolic projection (see code for more explanation)
    - indirect manipulation via MMB: hit parts are stored per mode then are used to drive the proper drag function
    - defered drag function: in order to avoid notifications storm, the drag functions can be applied on tick instead using a pending function
    - a single default TRS gizmo is now created when needed (instead of creating several thru selection changed or when the gizmo manager was ticking) and its visibility is set based on rules

TODO
    - UEditorInteractiveGizmoManager is now closer to UInteractiveGizmoManager so I don't think we need ActiveEditorGizmos, CachedGizmoMap, etc. that are not used anymore. This should be cleaned in another CL
    - UEditorTransformGizmo::OnGizmoTransformBegin should set the current axis (needed for some modes, including IKRig) but this is disabled for now has there are some side effects with dragging (this will be treated in another CL)

diff --git a/Engine/Plugins/Animation/ControlRig/Source/ControlRigEditor/Private/EditMode/ControlRigEditMode.cpp b/Engine/Plugins/Animation/ControlRig/Source/ControlRigEditor/Private/EditMode/ControlRigEditMode.cpp
--- a/Engine/Plugins/Animation/ControlRig/Source/ControlRigEditor/Private/EditMode/ControlRigEditMode.cpp
+++ b/Engine/Plugins/Animation/ControlRig/Source/ControlRigEditor/Private/EditMode/ControlRigEditMode.cpp
@@ -1769,7 +1769,15 @@ bool FControlRigEditMode::InputDelta(FEditorViewportClient* InViewportClient, FV
 	const bool bShiftDown = InViewport->KeyState(EKeys::LeftShift) || InViewport->KeyState(EKeys::RightShift);
 
 	//button down if left and ctrl and right is down, needed for indirect posting
-	const bool bMouseButtonDown = InViewport->KeyState(EKeys::LeftMouseButton) || (bCtrlDown && InViewport->KeyState(EKeys::RightMouseButton));
+
+	// enable MMB with the new TRS gizmos
+	static IConsoleVariable* const UseLegacyWidgetCVar = IConsoleManager::Get().FindConsoleVariable(TEXT("Gizmos.UseLegacyWidget"));
+	const bool bEnableMMB = UseLegacyWidgetCVar ? !UseLegacyWidgetCVar->GetBool() : false;
+	
+	const bool bMouseButtonDown =
+		InViewport->KeyState(EKeys::LeftMouseButton) ||
+		(bCtrlDown && InViewport->KeyState(EKeys::RightMouseButton)) ||
+		bEnableMMB;
 
 	const UE::Widget::EWidgetMode WidgetMode = InViewportClient->GetWidgetMode();
 	const EAxisList::Type CurrentAxis = InViewportClient->GetCurrentWidgetAxis();
diff --git a/Engine/Source/Editor/Experimental/EditorInteractiveToolsFramework/Private/EdModeInteractiveToolsContext.cpp b/Engine/Source/Editor/Experimental/EditorInteractiveToolsFramework/Private/EdModeInteractiveToolsContext.cpp
--- a/Engine/Source/Editor/Experimental/EditorInteractiveToolsFramework/Private/EdModeInteractiveToolsContext.cpp
+++ b/Engine/Source/Editor/Experimental/EditorInteractiveToolsFramework/Private/EdModeInteractiveToolsContext.cpp
@@ -321,6 +321,11 @@ void UEditorInteractiveToolsContext::Shutdown()
 {
 	bIsActive = false;
 
+	if (UContextObjectStore* ContextStore = GizmoManager ? GizmoManager->GetContextObjectStore() : nullptr)
+	{
+		ContextStore->RemoveContextObject(this);
+	}
+	
 	// auto-accept any in-progress tools
 	DeactivateAllActiveTools(EToolShutdownType::Accept);
 
@@ -339,6 +344,10 @@ void UEditorInteractiveToolsContext::InitializeContextWithEditorModeManager(FEdi
 	SetCreateGizmoManagerFunc([this](const FContextInitInfo& ContextInfo)
 	{
 		UEditorInteractiveGizmoManager* NewGizmoManager = NewObject<UEditorInteractiveGizmoManager>(ContextInfo.ToolsContext);
+		if (UContextObjectStore* ContextStore = NewGizmoManager->GetContextObjectStore())
+		{
+			ContextStore->AddContextObject(this);
+		}
 		NewGizmoManager->InitializeWithEditorModeManager(ContextInfo.QueriesAPI, ContextInfo.TransactionsAPI, ContextInfo.InputRouter, EditorModeManager);
 		NewGizmoManager->RegisterDefaultGizmos();
 		return NewGizmoManager;
@@ -364,16 +373,6 @@ void UEditorInteractiveToolsContext::InitializeContextWithEditorModeManager(FEdi
 
 	// set up standard materials
 	StandardVertexColorMaterial = GEngine->VertexColorMaterial;
-
-	if (UTypedElementSelectionSet* TypedElementSelectionSet = EditorModeManager->GetEditorSelectionSet())
-	{
-		TypedElementSelectionSet->OnChanged().AddUObject(this, &UEditorInteractiveToolsContext::OnEditorSelectionSetChanged);
-	}
-	else
-	{
-		FLevelEditorModule& LevelEditor = FModuleManager::Get().LoadModuleChecked<FLevelEditorModule>("LevelEditor");
-		LevelEditor.OnLevelEditorCreated().AddUObject(this, &UEditorInteractiveToolsContext::OnLevelEditorCreated);
-	}
 }
 
 
@@ -852,25 +851,6 @@ void UEditorInteractiveToolsContext::SetAbsoluteWorldSnappingEnabled(bool bEnabl
 	bEnableAbsoluteWorldSnapping = bEnabled;
 }
 
-void UEditorInteractiveToolsContext::OnLevelEditorCreated(TSharedPtr<ILevelEditor> InLevelEditor)
-{
-	if (UTypedElementSelectionSet* TypedElementSelectionSet = EditorModeManager->GetEditorSelectionSet())
-	{
-		TypedElementSelectionSet->OnChanged().AddUObject(this, &UEdModeInteractiveToolsContext::OnEditorSelectionSetChanged);
-	}
-}
-
-void UEditorInteractiveToolsContext::OnEditorSelectionSetChanged(const UTypedElementSelectionSet* InSelectionSet)
-{
-	if (UEditorInteractiveGizmoManager* EditorGizmoManager = Cast<UEditorInteractiveGizmoManager>(GizmoManager))
-	{
-		EditorGizmoManager->HandleEditorSelectionSetChanged(InSelectionSet);
-	}
-}
-
-
-
-
 void UModeManagerInteractiveToolsContext::Tick(FEditorViewportClient* ViewportClient, float DeltaTime)
 {
 	UEditorInteractiveToolsContext::Tick(ViewportClient, DeltaTime);
diff --git a/Engine/Source/Editor/Experimental/EditorInteractiveToolsFramework/Private/EditorGizmos/EditorTransformGizmo.cpp b/Engine/Source/Editor/Experimental/EditorInteractiveToolsFramework/Private/EditorGizmos/EditorTransformGizmo.cpp
--- a/Engine/Source/Editor/Experimental/EditorInteractiveToolsFramework/Private/EditorGizmos/EditorTransformGizmo.cpp
+++ b/Engine/Source/Editor/Experimental/EditorInteractiveToolsFramework/Private/EditorGizmos/EditorTransformGizmo.cpp
@@ -14,6 +14,17 @@
 
 DEFINE_LOG_CATEGORY_STATIC(LogEditorTransformGizmo, Log, All);
 
+void UEditorTransformGizmo::Render(IToolsContextRenderAPI* RenderAPI)
+{
+	const FSceneView* SceneView = RenderAPI ? RenderAPI->GetSceneView() : nullptr;
+	const bool bEngineShowFlagsModeWidget = SceneView && SceneView->Family &&
+											SceneView->Family->EngineShowFlags.ModeWidgets;
+	if (bEngineShowFlagsModeWidget)
+	{
+		Super::Render(RenderAPI);
+	}
+}
+
 void UEditorTransformGizmo::ApplyTranslateDelta(const FVector& InTranslateDelta)
 {
 	check(ActiveTarget);
@@ -70,4 +81,34 @@ void UEditorTransformGizmo::ApplyScaleDelta(const FVector& InScaleDelta)
 	}
 }
 
+void UEditorTransformGizmo::SetActiveTarget(UTransformProxy* Target, IToolContextTransactionProvider* TransactionProvider, IGizmoStateTarget* InStateTarget)
+{
+	Super::SetActiveTarget(Target, TransactionProvider, InStateTarget);
+
+	Target->OnBeginTransformEdit.AddUObject(this, &UEditorTransformGizmo::OnGizmoTransformBegin);
+	Target->OnEndTransformEdit.AddUObject(this, &UEditorTransformGizmo::OnGizmoTransformEnd);
+}
+
+void UEditorTransformGizmo::OnGizmoTransformBegin(UTransformProxy* InTransformProxy) const
+{
+	const UEditorTransformProxy* EditorTransformProxy = Cast<UEditorTransformProxy>(ActiveTarget);
+	if (InTransformProxy && InTransformProxy == EditorTransformProxy)
+	{
+		// FIXME this is need for some modes (including IKRig) but has side effects for now
+		// has it sets dragging to true. I'll fix this in a next CL
+		
+		// Set legacy widget axis temporarily because FEditorViewportClient overrides may expect it
+		// EditorTransformProxy->SetCurrentAxis(InteractionAxisList);
+	}
+}
+
+void UEditorTransformGizmo::OnGizmoTransformEnd(UTransformProxy* InTransformProxy) const
+{
+	const UEditorTransformProxy* EditorTransformProxy = Cast<UEditorTransformProxy>(ActiveTarget);
+	if (InTransformProxy && InTransformProxy == EditorTransformProxy)
+	{
+		EditorTransformProxy->SetCurrentAxis(EAxisList::None);
+	}
+}
+
 #undef LOCTEXT_NAMESPACE
diff --git a/Engine/Source/Editor/Experimental/EditorInteractiveToolsFramework/Private/EditorGizmos/EditorTransformGizmoBuilder.cpp b/Engine/Source/Editor/Experimental/EditorInteractiveToolsFramework/Private/EditorGizmos/EditorTransformGizmoBuilder.cpp
--- a/Engine/Source/Editor/Experimental/EditorInteractiveToolsFramework/Private/EditorGizmos/EditorTransformGizmoBuilder.cpp
+++ b/Engine/Source/Editor/Experimental/EditorInteractiveToolsFramework/Private/EditorGizmos/EditorTransformGizmoBuilder.cpp
@@ -6,6 +6,8 @@
 #include "BaseGizmos/GizmoElementHitTargets.h"
 #include "BaseGizmos/GizmoViewContext.h"
 #include "ContextObjectStore.h"
+#include "EditorModeManager.h"
+#include "EdModeInteractiveToolsContext.h"
 #include "EditorGizmos/EditorTransformGizmo.h"
 #include "EditorGizmos/EditorTransformGizmoSource.h"
 #include "EditorGizmos/EditorTransformProxy.h"
@@ -14,6 +16,7 @@
 #include "InteractiveGizmoManager.h"
 #include "Templates/Casts.h"
 #include "ToolContextInterfaces.h"
+#include "EditorGizmos/EditorTransformGizmoUtil.h"
 #include "UObject/Object.h"
 #include "UObject/ObjectPtr.h"
 #include "UObject/ScriptInterface.h"
@@ -26,7 +29,7 @@ UInteractiveGizmo* UEditorTransformGizmoBuilder::BuildGizmo(const FToolBuilderSt
 	UEditorTransformGizmo* TransformGizmo = NewObject<UEditorTransformGizmo>(SceneState.GizmoManager);
 	TransformGizmo->SetCustomizationFunction(CustomizationFunction);
 	TransformGizmo->Setup();
-	TransformGizmo->TransformGizmoSource = UEditorTransformGizmoSource::Construct(TransformGizmo);
+	TransformGizmo->TransformGizmoSource = UEditorTransformGizmoSource::CreateNew(TransformGizmo, GetTransformGizmoContext(SceneState));
 	TransformGizmo->GizmoViewContext = GizmoViewContext;
 
 	// @todo: Gizmo element construction to be moved here from UTransformGizmo.
@@ -41,7 +44,7 @@ void UEditorTransformGizmoBuilder::UpdateGizmoForSelection(UInteractiveGizmo* Gi
 {
 	if (UTransformGizmo* TransformGizmo = Cast<UTransformGizmo>(Gizmo))
 	{
-		UEditorTransformProxy* TransformProxy = NewObject<UEditorTransformProxy>();
+		UEditorTransformProxy* TransformProxy = UEditorTransformProxy::CreateNew(GetTransformGizmoContext(SceneState));
 		TransformGizmo->SetActiveTarget(TransformProxy);
 		TransformGizmo->SetVisibility(true);
 		
@@ -51,3 +54,9 @@ void UEditorTransformGizmoBuilder::UpdateGizmoForSelection(UInteractiveGizmo* Gi
 		}
 	}
 }
+
+const UEditorTransformGizmoContextObject* UEditorTransformGizmoBuilder::GetTransformGizmoContext(const FToolBuilderState& InSceneState)
+{
+	const UContextObjectStore* ContextStore = InSceneState.GizmoManager ? InSceneState.GizmoManager->GetContextObjectStore() : nullptr;
+	return ContextStore ? ContextStore->FindContext<UEditorTransformGizmoContextObject>() : nullptr;
+}
diff --git a/Engine/Source/Editor/Experimental/EditorInteractiveToolsFramework/Private/EditorGizmos/EditorTransformGizmoSource.cpp b/Engine/Source/Editor/Experimental/EditorInteractiveToolsFramework/Private/EditorGizmos/EditorTransformGizmoSource.cpp
--- a/Engine/Source/Editor/Experimental/EditorInteractiveToolsFramework/Private/EditorGizmos/EditorTransformGizmoSource.cpp
+++ b/Engine/Source/Editor/Experimental/EditorInteractiveToolsFramework/Private/EditorGizmos/EditorTransformGizmoSource.cpp
@@ -5,6 +5,7 @@
 #include "Editor.h"
 #include "EditorModeManager.h"
 #include "EditorViewportClient.h"
+#include "EditorGizmos/EditorTransformGizmoUtil.h"
 #include "HAL/IConsoleManager.h"
 #include "HAL/Platform.h"
 #include "Misc/AssertionMacros.h"
@@ -16,6 +17,7 @@ EGizmoTransformMode FEditorTransformGizmoUtil::GetGizmoMode(UE::Widget::EWidgetM
 		case UE::Widget::EWidgetMode::WM_Translate: return EGizmoTransformMode::Translate;
 		case UE::Widget::EWidgetMode::WM_Rotate: return EGizmoTransformMode::Rotate;
 		case UE::Widget::EWidgetMode::WM_Scale: return EGizmoTransformMode::Scale;
+		default: return EGizmoTransformMode::None;
 	}
 	return EGizmoTransformMode::None;
 }
@@ -27,15 +29,16 @@ UE::Widget::EWidgetMode FEditorTransformGizmoUtil::GetWidgetMode(EGizmoTransform
 		case EGizmoTransformMode::Translate: return UE::Widget::EWidgetMode::WM_Translate;
 		case EGizmoTransformMode::Rotate: return UE::Widget::EWidgetMode::WM_Rotate;
 		case EGizmoTransformMode::Scale: return UE::Widget::EWidgetMode::WM_Scale;
+		default: return UE::Widget::EWidgetMode::WM_None;
 	}
 	return UE::Widget::EWidgetMode::WM_None;
 }
 
 EGizmoTransformMode UEditorTransformGizmoSource::GetGizmoMode() const
 {
-	if (FEditorViewportClient* ViewportClient = GetViewportClient())
+	if (const FEditorViewportClient* ViewportClient = GetViewportClient())
 	{
-		UE::Widget::EWidgetMode WidgetMode = ViewportClient->GetWidgetMode();
+		const UE::Widget::EWidgetMode WidgetMode = ViewportClient->GetWidgetMode();
 		return FEditorTransformGizmoUtil::GetGizmoMode(WidgetMode);
 	}
 	return EGizmoTransformMode::None;
@@ -43,9 +46,9 @@ EGizmoTransformMode UEditorTransformGizmoSource::GetGizmoMode() const
 
 EAxisList::Type UEditorTransformGizmoSource::GetGizmoAxisToDraw(EGizmoTransformMode InGizmoMode) const
 { 
-	if (FEditorViewportClient* ViewportClient = GetViewportClient())
+	if (const FEditorViewportClient* ViewportClient = GetViewportClient())
 	{
-		UE::Widget::EWidgetMode WidgetMode = ViewportClient->GetWidgetMode();
+		const UE::Widget::EWidgetMode WidgetMode = ViewportClient->GetWidgetMode();
 		return GetModeTools().GetWidgetAxisToDraw(WidgetMode);
 	}
 	return EAxisList::None;
@@ -53,15 +56,13 @@ EAxisList::Type UEditorTransformGizmoSource::GetGizmoAxisToDraw(EGizmoTransformM
 
 EToolContextCoordinateSystem UEditorTransformGizmoSource::GetGizmoCoordSystemSpace() const
 {
-	FEditorViewportClient* ViewportClient = GetViewportClient();
+	const FEditorViewportClient* ViewportClient = GetViewportClient();
 	if (ViewportClient && ViewportClient->GetWidgetCoordSystemSpace() == ECoordSystem::COORD_Local)
 	{
 		return EToolContextCoordinateSystem::Local;
 	}
-	else
-	{
-		return EToolContextCoordinateSystem::World;
-	}
+	
+	return EToolContextCoordinateSystem::World;
 }
 
 float UEditorTransformGizmoSource::GetGizmoScale() const
@@ -69,14 +70,14 @@ float UEditorTransformGizmoSource::GetGizmoScale() const
 	return GetModeTools().GetWidgetScale();
 }
 
-
 bool UEditorTransformGizmoSource::GetVisible() const
 {
-	if (FEditorViewportClient* ViewportClient = GetViewportClient()) 
+	if (const FEditorViewportClient* ViewportClient = GetViewportClient()) 
 	{
-		if (GetModeTools().GetShowWidget() && GetModeTools().UsesTransformWidget())
+		const FEditorModeTools& ModeTools = GetModeTools();
+		if (ModeTools.GetShowWidget() && ModeTools.UsesTransformWidget())
 		{
-			UE::Widget::EWidgetMode WidgetMode = ViewportClient->GetWidgetMode();
+			const UE::Widget::EWidgetMode WidgetMode = ViewportClient->GetWidgetMode();
 			bool bUseLegacyWidget = (WidgetMode == UE::Widget::WM_TranslateRotateZ || WidgetMode == UE::Widget::WM_2D);
 			if (!bUseLegacyWidget)
 			{
@@ -104,12 +105,23 @@ EGizmoTransformScaleType UEditorTransformGizmoSource::GetScaleType() const
 	return EGizmoTransformScaleType::Default;
 }
 
-FEditorModeTools& UEditorTransformGizmoSource::GetModeTools() const
+UEditorTransformGizmoSource* UEditorTransformGizmoSource::CreateNew(UObject* Outer, const UEditorTransformGizmoContextObject* InContext)
+{
+	UEditorTransformGizmoSource* NewSource = NewObject<UEditorTransformGizmoSource>(Outer);
+	NewSource->WeakContext = InContext;
+	return NewSource;
+}
+
+const FEditorModeTools& UEditorTransformGizmoSource::GetModeTools() const
 {
+	if (WeakContext.IsValid())
+	{
+		return *WeakContext->GetModeTools();
+	}
 	return GLevelEditorModeTools();
 }
 
-FEditorViewportClient* UEditorTransformGizmoSource::GetViewportClient() const
+const FEditorViewportClient* UEditorTransformGizmoSource::GetViewportClient() const
 {
 	return GetModeTools().GetFocusedViewportClient();
 }
\ No newline at end of file
diff --git a/Engine/Source/Editor/Experimental/EditorInteractiveToolsFramework/Private/EditorGizmos/EditorTransformGizmoUtil.cpp b/Engine/Source/Editor/Experimental/EditorInteractiveToolsFramework/Private/EditorGizmos/EditorTransformGizmoUtil.cpp
new file mode 100644
--- /dev/null
+++ b/Engine/Source/Editor/Experimental/EditorInteractiveToolsFramework/Private/EditorGizmos/EditorTransformGizmoUtil.cpp
@@ -0,0 +1,327 @@
+﻿// Copyright Epic Games, Inc. All Rights Reserved.
+
+#include "EditorGizmos/EditorTransformGizmoUtil.h"
+
+#include "ContextObjectStore.h"
+#include "EditorInteractiveGizmoSubsystem.h"
+#include "EditorInteractiveGizmoManager.h"
+#include "InteractiveToolsContext.h"
+#include "EditorModeManager.h"
+#include "EdModeInteractiveToolsContext.h"
+#include "EditorGizmos/TransformGizmo.h"
+#include "EditorGizmos/EditorTransformGizmoBuilder.h"
+#include "EditorViewportClient.h"
+
+#include UE_INLINE_GENERATED_CPP_BY_NAME(EditorTransformGizmoUtil)
+
+bool UE::EditorTransformGizmoUtil::RegisterTransformGizmoContextObject(FEditorModeTools* InModeTools)
+{
+	const UModeManagerInteractiveToolsContext* ToolsContext = InModeTools->GetInteractiveToolsContext();
+	if (ensure(ToolsContext))
+	{
+		const UEditorTransformGizmoContextObject* Found = ToolsContext->ContextObjectStore->FindContext<UEditorTransformGizmoContextObject>();
+		if (Found)
+		{
+			return true;
+		}
+
+		UEditorTransformGizmoContextObject* GizmoContextObject = NewObject<UEditorTransformGizmoContextObject>(ToolsContext->ToolManager);
+		if (ensure(GizmoContextObject))
+		{
+			GizmoContextObject->RegisterGizmosWithManager(ToolsContext->ToolManager);
+			GizmoContextObject->Initialize(InModeTools);
+			return true;
+		}
+	}
+	return false;
+}
+
+bool UE::EditorTransformGizmoUtil::UnregisterTransformGizmoContextObject(FEditorModeTools* InModeTools)
+{
+	const UModeManagerInteractiveToolsContext* ToolsContext = InModeTools->GetInteractiveToolsContext();
+	if (ensure(ToolsContext))
+	{
+		UEditorTransformGizmoContextObject* Found = ToolsContext->ContextObjectStore->FindContext<UEditorTransformGizmoContextObject>();
+		if (Found != nullptr)
+		{
+			Found->Shutdown();
+			Found->UnregisterGizmosWithManager(ToolsContext->ToolManager);
+		}
+		return true;
+	}
+	return false;
+}
+
+UTransformGizmo* UE::EditorTransformGizmoUtil::CreateTransformGizmo(
+	UInteractiveToolManager* InToolManager, const FString& InInstanceIdentifier, void* InOwner)
+{
+	if (!InToolManager)
+	{
+		return nullptr;
+	}
+	
+	UEditorInteractiveGizmoManager* GizmoManager = Cast<UEditorInteractiveGizmoManager>(InToolManager->GetPairedGizmoManager());
+	if (!GizmoManager)
+	{
+		return nullptr;
+	}
+
+	const UEditorTransformGizmoContextObject* Found = GizmoManager->GetContextObjectStore()->FindContext<UEditorTransformGizmoContextObject>();
+	if (!Found)
+	{
+		return nullptr;
+	}
+
+	UInteractiveGizmo* NewGizmo = GizmoManager->CreateGizmo(
+		UEditorInteractiveGizmoManager::TransformBuilderIdentifier(), InInstanceIdentifier, InOwner);
+
+	ensure(NewGizmo);
+	
+	return Cast<UTransformGizmo>(NewGizmo);
+}
+
+UTransformGizmo* UE::EditorTransformGizmoUtil::GetDefaultTransformGizmo(UInteractiveToolManager* InToolManager)
+{
+	return CreateTransformGizmo(InToolManager, UEditorInteractiveGizmoManager::TransformInstanceIdentifier(), nullptr);
+}
+
+UTransformGizmo* UE::EditorTransformGizmoUtil::FindDefaultTransformGizmo(UInteractiveToolManager* InToolManager)
+{
+	if (!InToolManager)
+	{
+		return nullptr;
+	}
+	
+	const UEditorInteractiveGizmoManager* GizmoManager = Cast<UEditorInteractiveGizmoManager>(InToolManager->GetPairedGizmoManager());
+	if (!GizmoManager)
+	{
+		return nullptr;
+	}
+
+	UInteractiveGizmo* Found = GizmoManager->FindGizmoByInstanceIdentifier(UEditorInteractiveGizmoManager::TransformInstanceIdentifier());
+	return Cast<UTransformGizmo>(Found);
+}
+
+void UE::EditorTransformGizmoUtil::RemoveDefaultTransformGizmo(UInteractiveToolManager* InToolManager)
+{
+	if (!InToolManager)
+	{
+		return;
+	}
+	
+	UEditorInteractiveGizmoManager* GizmoManager = Cast<UEditorInteractiveGizmoManager>(InToolManager->GetPairedGizmoManager());
+	if (!GizmoManager)
+	{
+		return;
+	}
+
+	UInteractiveGizmo* Found = GizmoManager->FindGizmoByInstanceIdentifier(UEditorInteractiveGizmoManager::TransformInstanceIdentifier());
+	if (Found)
+	{
+		GizmoManager->DestroyGizmo(Found);
+	}
+}
+
+void UEditorTransformGizmoContextObject::RegisterGizmosWithManager(UInteractiveToolManager* InToolManager)
+{
+	if (ensure(!bGizmosRegistered) == false)
+	{
+		return;
+	}
+
+	InToolManager->GetContextObjectStore()->AddContextObject(this);
+	
+	UEditorInteractiveGizmoManager* GizmoManager = Cast<UEditorInteractiveGizmoManager>(InToolManager->GetPairedGizmoManager());
+	if (ensure(GizmoManager))
+	{
+		if (UEditorInteractiveGizmoSubsystem* GizmoSubsystem = GEditor->GetEditorSubsystem<UEditorInteractiveGizmoSubsystem>())
+		{
+			GizmoManager->RegisterGizmoType(UEditorInteractiveGizmoManager::TransformBuilderIdentifier(), GizmoSubsystem->GetTransformGizmoBuilder());
+		}
+		else
+		{
+			UEditorTransformGizmoBuilder* GizmoBuilder = NewObject<UEditorTransformGizmoBuilder>(this);
+			GizmoManager->RegisterGizmoType(UEditorInteractiveGizmoManager::TransformBuilderIdentifier(), GizmoBuilder);	
+		}
+	}
+	
+	bGizmosRegistered = true;
+}
+
+void UEditorTransformGizmoContextObject::UnregisterGizmosWithManager(UInteractiveToolManager* InToolManager)
+{
+	ensure(bGizmosRegistered);
+	
+	InToolManager->GetContextObjectStore()->RemoveContextObject(this);
+
+	UInteractiveGizmoManager* GizmoManager = InToolManager->GetPairedGizmoManager();
+	GizmoManager->DeregisterGizmoType(UEditorInteractiveGizmoManager::TransformBuilderIdentifier());
+	
+	bGizmosRegistered = false;
+}
+
+void UEditorTransformGizmoContextObject::Initialize(FEditorModeTools* InModeTools)
+{
+	if (ensure(InModeTools && !ModeTools))
+	{
+		ModeTools = InModeTools;
+		InitializeCVarBinding();
+	}
+}
+
+void UEditorTransformGizmoContextObject::Shutdown()
+{
+	if (ensure(ModeTools))
+	{
+		UpdateGizmo({});
+		RemoveViewportsBinding();
+		RemoveCVarBinding();
+		
+		ModeTools = nullptr;
+	}
+}
+
+void UEditorTransformGizmoContextObject::UpdateGizmo(const TArray<FEditorViewportClient*>& InViewportClients) const
+{
+	if (!ModeTools)
+	{
+		return;
+	}
+
+	// get viewports pointing at ModeTools
+	const TArray<FEditorViewportClient*> ViewportClients = InViewportClients.FilterByPredicate([this](const FEditorViewportClient* ViewportClient)
+	{
+		return ViewportClient && ViewportClient->GetModeTools() == ModeTools;
+	});
+
+	// remove useless gizmo
+	if (ViewportClients.IsEmpty())
+	{
+		const UModeManagerInteractiveToolsContext* ToolsContext = ModeTools->GetInteractiveToolsContext();
+		const TObjectPtr<UInteractiveToolManager> ToolManager = ToolsContext->ToolManager;
+		UE::EditorTransformGizmoUtil::RemoveDefaultTransformGizmo(ToolManager);
+		return;
+	}
+
+	// if the focus viewport has been deleted, switch to the first visible one if any
+	// FIXME, use the hover one instead ?
+	const bool bIsFocusStillAvailable = ViewportClients.Contains(ModeTools->GetFocusedViewportClient());
+	if (!bIsFocusStillAvailable)
+	{
+		const int32 Index = ViewportClients.IndexOfByPredicate([](const FEditorViewportClient* ViewportClient)
+		{
+			return ViewportClient->IsVisible();
+		});
+		if (Index != INDEX_NONE)
+		{
+			ModeTools->ReceivedFocus(ViewportClients[Index], ViewportClients[Index]->Viewport);
+		}
+	}
+	
+	const UModeManagerInteractiveToolsContext* ToolsContext = ModeTools->GetInteractiveToolsContext();
+	const TObjectPtr<UInteractiveToolManager> ToolManager = ToolsContext->ToolManager;
+	
+	if (UE::EditorTransformGizmoUtil::FindDefaultTransformGizmo(ToolManager))
+	{
+		return;
+	}
+				
+	// create a new one
+	if (UTransformGizmo* TransformGizmo = UE::EditorTransformGizmoUtil::GetDefaultTransformGizmo(ToolManager))
+	{
+		TransformGizmo->SetVisibility(false);
+	}
+}
+
+void UEditorTransformGizmoContextObject::InitializeCVarBinding()
+{
+	static IConsoleVariable* UseLegacyWidgetCVar = IConsoleManager::Get().FindConsoleVariable(TEXT("Gizmos.UseLegacyWidget"));
+	if (ensure(UseLegacyWidgetCVar))
+	{
+		auto OnGizmoVariableChanged = [this](IConsoleVariable* InLegacyWidgetCVar)
+		{
+			const bool bUseLegacyGizmo = InLegacyWidgetCVar ? UseLegacyWidgetCVar->GetInt() > 0 : false;
+
+			if (bUseLegacyGizmo)
+			{
+				// remove viewports' binding + gizmos as they are useless
+				
+				const UModeManagerInteractiveToolsContext* ToolsContext = ModeTools->GetInteractiveToolsContext();
+				const TObjectPtr<UInteractiveToolManager> ToolManager = ToolsContext->ToolManager;
+				UE::EditorTransformGizmoUtil::RemoveDefaultTransformGizmo(ToolManager);
+				
+				RemoveViewportsBinding();
+				
+				return;
+			}
+
+			InitializeViewportsBinding();
+		};
+
+		// bind cvar change delegate
+		if (!UseLegacyChangedHandled.IsValid())
+		{
+			UseLegacyChangedHandled = UseLegacyWidgetCVar->OnChangedDelegate().AddLambda(OnGizmoVariableChanged);
+		}
+		
+		// initialize default viewport list change binding
+		OnGizmoVariableChanged(UseLegacyWidgetCVar);
+	}
+}
+
+void UEditorTransformGizmoContextObject::InitializeViewportsBinding()
+{
+	if (GEditor)
+	{
+		// check added/removed viewports to ensure that ModeTools' FocusedViewportClient is still safe to use
+		auto OnViewportClientsChanged = [this]()
+		{
+			const TArray<FEditorViewportClient*> ViewportClients =
+				GEditor->GetAllViewportClients().FilterByPredicate([this](const FEditorViewportClient* ViewportClient)
+				{
+					return ViewportClient && ViewportClient->GetModeTools() == ModeTools;
+				});
+
+			UpdateGizmo(ViewportClients);
+		};
+
+		// bind if needed
+        if (!ViewportClientsChangedHandle.IsValid())
+        {
+        	ViewportClientsChangedHandle = GEditor->OnViewportClientListChanged().AddLambda(OnViewportClientsChanged);
+        }
+    
+		// initialize
+        OnViewportClientsChanged();
+	}
+}
+
+void UEditorTransformGizmoContextObject::RemoveCVarBinding()
+{
+	static IConsoleVariable* UseLegacyWidgetCVar = IConsoleManager::Get().FindConsoleVariable(TEXT("Gizmos.UseLegacyWidget"));
+	if (ensure(UseLegacyWidgetCVar))
+	{
+		if (UseLegacyChangedHandled.IsValid())
+		{
+			UseLegacyWidgetCVar->OnChangedDelegate().Remove(UseLegacyChangedHandled);
+			UseLegacyChangedHandled.Reset();
+		}
+	}
+}
+
+void UEditorTransformGizmoContextObject::RemoveViewportsBinding()
+{
+	if (ViewportClientsChangedHandle.IsValid())
+	{
+		if (ensure(GEditor))
+		{
+			GEditor->OnViewportClientListChanged().Remove(ViewportClientsChangedHandle);
+		}
+		ViewportClientsChangedHandle.Reset();
+	}
+}
+
+const FEditorModeTools* UEditorTransformGizmoContextObject::GetModeTools() const
+{
+	return ModeTools;
+}
\ No newline at end of file
diff --git a/Engine/Source/Editor/Experimental/EditorInteractiveToolsFramework/Private/EditorGizmos/EditorTransformProxy.cpp b/Engine/Source/Editor/Experimental/EditorInteractiveToolsFramework/Private/EditorGizmos/EditorTransformProxy.cpp
--- a/Engine/Source/Editor/Experimental/EditorInteractiveToolsFramework/Private/EditorGizmos/EditorTransformProxy.cpp
+++ b/Engine/Source/Editor/Experimental/EditorInteractiveToolsFramework/Private/EditorGizmos/EditorTransformProxy.cpp
@@ -5,6 +5,7 @@
 #include "Editor.h"
 #include "EditorModeManager.h"
 #include "EditorViewportClient.h"
+#include "EditorGizmos/EditorTransformGizmoUtil.h"
 #include "Math/Matrix.h"
 #include "Math/Quat.h"
 #include "Math/Vector.h"
@@ -14,21 +15,19 @@
 
 FTransform UEditorTransformProxy::GetTransform() const
 {
-	if (FEditorViewportClient* ViewportClient = GLevelEditorModeTools().GetFocusedViewportClient())
+	if (const FEditorViewportClient* ViewportClient = GetViewportClient())
 	{
-		FVector Location = ViewportClient->GetWidgetLocation();
-		FMatrix RotMatrix = ViewportClient->GetWidgetCoordSystem();
-		return FTransform(FQuat(RotMatrix), Location, FVector::OneVector);
-	}
-	else
-	{
-		return FTransform::Identity;
+		const FVector Location = ViewportClient->GetWidgetLocation();
+		const FQuat Rotation(ViewportClient->GetWidgetCoordSystem());
+		return FTransform(Rotation, Location, FVector::OneVector);
 	}
+	
+	return FTransform::Identity;
 }
 
 void UEditorTransformProxy::InputTranslateDelta(const FVector& InDeltaTranslate, EAxisList::Type InAxisList)
 {
-	if (FEditorViewportClient* ViewportClient = GLevelEditorModeTools().GetFocusedViewportClient())
+	if (FEditorViewportClient* ViewportClient = GetViewportClient())
 	{
 		FVector Translate = InDeltaTranslate;
 		FRotator Rot = FRotator::ZeroRotator;
@@ -43,7 +42,7 @@ void UEditorTransformProxy::InputTranslateDelta(const FVector& InDeltaTranslate,
 
 void UEditorTransformProxy::InputScaleDelta(const FVector& InDeltaScale, EAxisList::Type InAxisList)
 {
-	if (FEditorViewportClient* ViewportClient = GLevelEditorModeTools().GetFocusedViewportClient())
+	if (FEditorViewportClient* ViewportClient = GetViewportClient())
 	{
 		FVector Translate = FVector::ZeroVector;
 		FRotator Rot = FRotator::ZeroRotator;
@@ -58,7 +57,7 @@ void UEditorTransformProxy::InputScaleDelta(const FVector& InDeltaScale, EAxisLi
 
 void UEditorTransformProxy::InputRotateDelta(const FRotator& InDeltaRotate, EAxisList::Type InAxisList)
 {
-	if (FEditorViewportClient* ViewportClient = GLevelEditorModeTools().GetFocusedViewportClient())
+	if (FEditorViewportClient* ViewportClient = GetViewportClient())
 	{
 		FVector Translate = FVector::ZeroVector;
 		FRotator Rot = InDeltaRotate;
@@ -71,4 +70,28 @@ void UEditorTransformProxy::InputRotateDelta(const FRotator& InDeltaRotate, EAxi
 	}
 }
 
+UEditorTransformProxy* UEditorTransformProxy::CreateNew(const UEditorTransformGizmoContextObject* InContext)
+{
+	UEditorTransformProxy* NewProxy = NewObject<UEditorTransformProxy>();
+	NewProxy->WeakContext = InContext;
+	return NewProxy;
+}
+
+FEditorViewportClient* UEditorTransformProxy::GetViewportClient() const
+{
+	if (WeakContext.IsValid() && WeakContext->GetModeTools())
+	{
+		return WeakContext->GetModeTools()->GetFocusedViewportClient();	
+	}
+	return GLevelEditorModeTools().GetFocusedViewportClient();
+}
+
+void UEditorTransformProxy::SetCurrentAxis(const EAxisList::Type InAxisList) const
+{
+	if (FEditorViewportClient* ViewportClient = GetViewportClient())
+	{
+		ViewportClient->SetCurrentWidgetAxis(InAxisList);
+	}
+}
+
 #undef LOCTEXT_NAMESPACE
diff --git a/Engine/Source/Editor/Experimental/EditorInteractiveToolsFramework/Private/EditorGizmos/TransformGizmo.cpp b/Engine/Source/Editor/Experimental/EditorInteractiveToolsFramework/Private/EditorGizmos/TransformGizmo.cpp
--- a/Engine/Source/Editor/Experimental/EditorInteractiveToolsFramework/Private/EditorGizmos/TransformGizmo.cpp
+++ b/Engine/Source/Editor/Experimental/EditorInteractiveToolsFramework/Private/EditorGizmos/TransformGizmo.cpp
@@ -8,17 +8,12 @@
 #include "BaseGizmos/GizmoElementShapes.h"
 #include "BaseGizmos/GizmoMath.h"
 #include "BaseGizmos/GizmoRenderingUtil.h"
-#include "BaseGizmos/ParameterSourcesFloat.h"
 #include "BaseGizmos/StateTargets.h"
-#include "Elements/Interfaces/TypedElementObjectInterface.h"
-#include "Elements/Interfaces/TypedElementWorldInterface.h"
 #include "Engine/CollisionProfile.h"
-#include "Engine/Selection.h"
 #include "Engine/World.h"
-#include "EditorModeTools.h"
 #include "Materials/Material.h"
-#include "UnrealEdGlobals.h"
 #include "UnrealEngine.h"
+#include "Intersection/IntersectionUtil.h"
 
 #define LOCTEXT_NAMESPACE "UTransformGizmo"
 
@@ -52,6 +47,12 @@ void UTransformGizmo::Setup()
 	GizmoElementRoot->SetInteractLineColor(CurrentColor);
 
 	bInInteraction = false;
+
+	SetModeLastHitPart(EGizmoTransformMode::None, ETransformGizmoPartIdentifier::Default);
+	SetModeLastHitPart(EGizmoTransformMode::Translate, ETransformGizmoPartIdentifier::TranslateScreenSpace);
+	SetModeLastHitPart(EGizmoTransformMode::Rotate, ETransformGizmoPartIdentifier::RotateArcball);
+	SetModeLastHitPart(EGizmoTransformMode::Scale, ETransformGizmoPartIdentifier::ScaleUniform);
+	SetModeLastHitPart(EGizmoTransformMode::Max, ETransformGizmoPartIdentifier::Default);
 }
 
 void UTransformGizmo::SetupBehaviors()
@@ -67,18 +68,59 @@ void UTransformGizmo::SetupBehaviors()
 	MouseBehavior->Initialize(this);
 	MouseBehavior->SetDefaultPriority(FInputCapturePriority(FInputCapturePriority::DEFAULT_GIZMO_PRIORITY));
 	AddInputBehavior(MouseBehavior);
+
+	// Add middle mouse input behavior for indirect manipulation
+	ULocalClickDragInputBehavior* MiddleClickDragBehavior = NewObject<ULocalClickDragInputBehavior>();
+	MiddleClickDragBehavior->Initialize();
+	MiddleClickDragBehavior->SetUseMiddleMouseButton();
+	MiddleClickDragBehavior->CanBeginClickDragFunc = [this](const FInputDeviceRay& InPressPos)
+	{
+		bIndirectManipulation = true;
+		return CanBeginClickDragSequence(InPressPos);
+	};
+	MiddleClickDragBehavior->OnClickPressFunc = [this](const FInputDeviceRay& InPressPos)
+	{
+		return OnClickPress(InPressPos);
+	};
+	MiddleClickDragBehavior->OnClickDragFunc = [this](const FInputDeviceRay& InDragPos)
+	{
+		return OnClickDrag(InDragPos);
+	};
+	MiddleClickDragBehavior->OnClickReleaseFunc = [this](const FInputDeviceRay& InReleasePos)
+	{
+		bIndirectManipulation = false;
+		return OnClickRelease(InReleasePos);
+	};
+	MiddleClickDragBehavior->OnTerminateFunc = [this]()
+	{
+		return OnTerminateDragSequence();
+	};
+	// add this to enable indirect manipulation using CTRL + MMB
+	// MiddleClickDragBehavior->ModifierCheckFunc = FInputDeviceState::IsCtrlKeyDown;
+	AddInputBehavior(MiddleClickDragBehavior);
 }
 
 void UTransformGizmo::SetupMaterials()
 {
-	auto GetBaseMaterial = [this]()
+	auto GetBaseMaterial = [this]() -> UMaterial*
 	{
 		if (CustomizationFunction)
 		{
 			const FGizmoCustomization& GizmoCustomization = CustomizationFunction();
-			return GizmoCustomization.Material ? GizmoCustomization.Material : GEngine->ArrowMaterial; 
+			if (IsValid(GizmoCustomization.Material))
+			{
+				return GizmoCustomization.Material.Get();
+			}
 		}
-		return GEngine->ArrowMaterial;
+
+		static const FString MaterialName = TEXT("/Engine/EditorMaterials/TransformGizmoMaterial.TransformGizmoMaterial");
+		UMaterial* Material = FindObject<UMaterial>(nullptr, *MaterialName);
+		if (!Material)
+		{
+			Material = LoadObject<UMaterial>(nullptr, *MaterialName, nullptr, LOAD_None, nullptr);
+		}
+		
+		return Material ? Material : GEngine->ArrowMaterial.Get();
 	};
 	
 	UMaterial* AxisMaterialBase = GetBaseMaterial(); 
@@ -176,9 +218,18 @@ bool UTransformGizmo::OnUpdateHover(const FInputDeviceRay& DevicePos)
 
 void UTransformGizmo::OnEndHover()
 {
-	if (HitTarget && LastHitPart != ETransformGizmoPartIdentifier::Default)
+	if (HitTarget)
 	{
-		HitTarget->UpdateHoverState(false, static_cast<uint32>(LastHitPart));
+		if (LastHitPart != ETransformGizmoPartIdentifier::Default)
+		{
+			UpdateHoverState(false, LastHitPart);
+		}
+
+		const ETransformGizmoPartIdentifier ModeHitPart = GetCurrentModeLastHitPart();
+		if (ModeHitPart != ETransformGizmoPartIdentifier::Default)
+		{
+			UpdateInteractingState(true, ModeHitPart, true);
+		}
 	}
 }
 
@@ -216,6 +267,12 @@ FInputRayHit UTransformGizmo::UpdateHoveredPart(const FInputDeviceRay& PressPos)
 		LastHitPart = HitPart;
 	}
 
+	const ETransformGizmoPartIdentifier ModeHitPart = GetCurrentModeLastHitPart();
+	if (ModeHitPart != ETransformGizmoPartIdentifier::Default)
+	{
+		UpdateInteractingState(true, ModeHitPart, true);
+	}
+
 	return RayHit;
 }
 
@@ -236,10 +293,59 @@ bool UTransformGizmo::VerifyPartIdentifier(uint32 InPartIdentifier) const
 	return true;
 }
 
+void UTransformGizmo::SetModeLastHitPart(const EGizmoTransformMode InMode, const ETransformGizmoPartIdentifier InIdentifier)
+{
+	if (InMode < EGizmoTransformMode::None || InMode > EGizmoTransformMode::Max)
+	{
+		return;
+	}
+	LastHitPartPerMode[static_cast<uint8>(InMode)] = InIdentifier;
+}
+
+ETransformGizmoPartIdentifier UTransformGizmo::GetCurrentModeLastHitPart() const
+{
+	auto GetTransformMode = [&]()
+	{
+		if (TransformGizmoSource)
+		{
+			return TransformGizmoSource->GetGizmoMode();
+		}
+		
+		const EToolContextTransformGizmoMode ActiveGizmoMode = GetGizmoManager()->GetContextQueriesAPI()->GetCurrentTransformGizmoMode();
+		switch (ActiveGizmoMode)
+		{
+		case EToolContextTransformGizmoMode::NoGizmo:
+			return EGizmoTransformMode::None;
+		case EToolContextTransformGizmoMode::Translation:
+			return EGizmoTransformMode::Translate;
+		case EToolContextTransformGizmoMode::Rotation:
+			return EGizmoTransformMode::Rotate;
+		case EToolContextTransformGizmoMode::Scale:
+			return EGizmoTransformMode::Scale;
+		case EToolContextTransformGizmoMode::Combined:
+			return EGizmoTransformMode::None;
+		}
+		
+		return EGizmoTransformMode::None;
+	};
 
+	const EGizmoTransformMode GizmoTransformMode = GetTransformMode();
+	return (GizmoTransformMode < EGizmoTransformMode::Max) ?
+		LastHitPartPerMode[static_cast<uint8>(GizmoTransformMode)] :
+		ETransformGizmoPartIdentifier::Default; 
+}
 
 FInputRayHit UTransformGizmo::CanBeginClickDragSequence(const FInputDeviceRay& PressPos)
 {
+	if (bIndirectManipulation)
+	{
+		if (LastHitPart == ETransformGizmoPartIdentifier::Default)
+		{
+			LastHitPart = GetCurrentModeLastHitPart();
+		}
+		return FInputRayHit(0.0);
+	}
+	
 	FInputRayHit RayHit;
 
 	if (HitTarget)
@@ -489,16 +595,10 @@ void UTransformGizmo::EnableRotate(EAxisList::Type InAxisListToDraw)
 			GizmoElementRoot->Add(RotateOuterCircleElement);
 		}
 
-		if (RotateArcballOuterElement == nullptr)
+		if (RotateArcballElement == nullptr)
 		{
-			RotateArcballOuterElement = MakeRotateCircleHandle(ETransformGizmoPartIdentifier::RotateArcball, RotateArcballOuterRadius, RotateArcballCircleColor, false);
-			GizmoElementRoot->Add(RotateArcballOuterElement);
-		}
-
-		if (RotateArcballInnerElement == nullptr)
-		{
-			RotateArcballInnerElement = MakeRotateCircleHandle(ETransformGizmoPartIdentifier::RotateArcballInnerCircle, RotateArcballInnerRadius, RotateArcballCircleColor, true);
-			GizmoElementRoot->Add(RotateArcballInnerElement);
+			RotateArcballElement = MakeRotateCircleHandle(ETransformGizmoPartIdentifier::RotateArcball, RotateArcballSphereRadius, RotateArcballCircleColor, true);
+			GizmoElementRoot->Add(RotateArcballElement);
 		}
 	}
 
@@ -527,14 +627,9 @@ void UTransformGizmo::EnableRotate(EAxisList::Type InAxisListToDraw)
 		RotateOuterCircleElement->SetEnabled(bEnableAll);
 	}
 
-	if (RotateArcballOuterElement)
-	{
-		RotateArcballOuterElement->SetEnabled(bEnableAll);
-	}
-
-	if (RotateArcballInnerElement)
+	if (RotateArcballElement)
 	{ 
-		RotateArcballInnerElement->SetEnabled(bEnableAll);
+		RotateArcballElement->SetEnabled(bEnableAll);
 	}
 }
 
@@ -608,6 +703,12 @@ void UTransformGizmo::UpdateCameraAxisSource()
 
 void UTransformGizmo::Tick(float DeltaTime)
 {
+	if (PendingDragFunction)
+	{
+		PendingDragFunction();
+		PendingDragFunction.Reset();
+	}
+	
 	UpdateMode();
 
 	UpdateCameraAxisSource();
@@ -673,7 +774,7 @@ void UTransformGizmo::SetCustomizationFunction(const TFunction<const FGizmoCusto
 
 UGizmoElementArrow* UTransformGizmo::MakeTranslateAxis(ETransformGizmoPartIdentifier InPartId, const FVector& InAxisDir, const FVector& InSideDir, UMaterialInterface* InMaterial)
 {
-	const float SizeCoeff = CustomizationFunction ? CustomizationFunction().SizeCoefficient : 1.f;
+	const float SizeCoeff = GetSizeCoefficient();
 	
 	UGizmoElementArrow* ArrowElement = NewObject<UGizmoElementArrow>();
 	ArrowElement->SetPartIdentifier(static_cast<uint32>(InPartId));
@@ -694,7 +795,7 @@ UGizmoElementArrow* UTransformGizmo::MakeTranslateAxis(ETransformGizmoPartIdenti
 
 UGizmoElementArrow* UTransformGizmo::MakeScaleAxis(ETransformGizmoPartIdentifier InPartId, const FVector& InAxisDir, const FVector& InSideDir, UMaterialInterface* InMaterial)
 {
-	const float SizeCoeff = CustomizationFunction ? CustomizationFunction().SizeCoefficient : 1.f;
+	const float SizeCoeff = GetSizeCoefficient();
 
 	UGizmoElementArrow* ArrowElement = NewObject<UGizmoElementArrow>();
 	ArrowElement->SetPartIdentifier(static_cast<uint32>(InPartId));
@@ -714,7 +815,7 @@ UGizmoElementArrow* UTransformGizmo::MakeScaleAxis(ETransformGizmoPartIdentifier
 
 UGizmoElementBox* UTransformGizmo::MakeUniformScaleHandle()
 {
-	const float SizeCoeff = CustomizationFunction ? CustomizationFunction().SizeCoefficient : 1.f;
+	const float SizeCoeff = GetSizeCoefficient();
 	
 	UGizmoElementBox* BoxElement = NewObject<UGizmoElementBox>();
 	BoxElement->SetPartIdentifier(static_cast<uint32>(ETransformGizmoPartIdentifier::ScaleUniform));
@@ -729,7 +830,7 @@ UGizmoElementBox* UTransformGizmo::MakeUniformScaleHandle()
 UGizmoElementRectangle* UTransformGizmo::MakePlanarHandle(ETransformGizmoPartIdentifier InPartId, const FVector& InUpDirection, const FVector& InSideDirection, const FVector& InPlaneNormal,
 	UMaterialInterface* InMaterial, const FLinearColor& InVertexColor)
 {
-	const float SizeCoeff = CustomizationFunction ? CustomizationFunction().SizeCoefficient : 1.f;
+	const float SizeCoeff = GetSizeCoefficient();
 	
 	FVector PlanarHandleCenter = (InUpDirection + InSideDirection) * PlanarHandleOffset * SizeCoeff;
 
@@ -757,7 +858,7 @@ UGizmoElementRectangle* UTransformGizmo::MakePlanarHandle(ETransformGizmoPartIde
 
 UGizmoElementRectangle* UTransformGizmo::MakeTranslateScreenSpaceHandle()
 {
-	const float SizeCoeff = CustomizationFunction ? CustomizationFunction().SizeCoefficient : 1.f;
+	const float SizeCoeff = GetSizeCoefficient();
 	
 	UGizmoElementRectangle* RectangleElement = NewObject<UGizmoElementRectangle>();
 	RectangleElement->SetPartIdentifier(static_cast<uint32>(ETransformGizmoPartIdentifier::TranslateScreenSpace));
@@ -782,7 +883,7 @@ UGizmoElementRectangle* UTransformGizmo::MakeTranslateScreenSpaceHandle()
 UGizmoElementTorus* UTransformGizmo::MakeRotateAxis(ETransformGizmoPartIdentifier InPartId, const FVector& TorusAxis0, const FVector& TorusAxis1,
 	UMaterialInterface* InMaterial, UMaterialInterface* InCurrentMaterial)
 {
-	const float SizeCoeff = CustomizationFunction ? CustomizationFunction().SizeCoefficient : 1.f;
+	const float SizeCoeff = GetSizeCoefficient();
 	
 	UGizmoElementTorus* RotateAxisElement = NewObject<UGizmoElementTorus>();
 	RotateAxisElement->SetPartIdentifier(static_cast<uint32>(InPartId));
@@ -807,7 +908,7 @@ UGizmoElementTorus* UTransformGizmo::MakeRotateAxis(ETransformGizmoPartIdentifie
 
 UGizmoElementCircle* UTransformGizmo::MakeRotateCircleHandle(ETransformGizmoPartIdentifier InPartId, float InRadius, const FLinearColor& InColor, float bFill)
 {
-	const float SizeCoeff = CustomizationFunction ? CustomizationFunction().SizeCoefficient : 1.f;
+	const float SizeCoeff = GetSizeCoefficient();
 	
 	UGizmoElementCircle* CircleElement = NewObject<UGizmoElementCircle>();
 	CircleElement->SetPartIdentifier(static_cast<uint32>(InPartId));
@@ -821,8 +922,15 @@ UGizmoElementCircle* UTransformGizmo::MakeRotateCircleHandle(ETransformGizmoPart
 
 	if (bFill)
 	{
-		CircleElement->SetVertexColor(InColor);
-		CircleElement->SetMaterial(WhiteMaterial);
+		const FLinearColor LightColor(InColor.R, InColor.G, InColor.B, InColor.A * 0.5);
+		CircleElement->SetVertexColor(LightColor);
+		CircleElement->SetMaterial(TransparentVertexColorMaterial);
+
+		CircleElement->SetHoverVertexColor(InColor);
+		CircleElement->SetHoverMaterial(TransparentVertexColorMaterial);
+		
+		CircleElement->SetInteractVertexColor(InColor);
+		CircleElement->SetInteractMaterial(TransparentVertexColorMaterial);
 	}
 	else
 	{
@@ -839,7 +947,13 @@ UGizmoElementCircle* UTransformGizmo::MakeRotateCircleHandle(ETransformGizmoPart
 void UTransformGizmo::ClearActiveTarget()
 {
 	StateTarget = nullptr;
-	ActiveTarget = nullptr;
+
+	if (ActiveTarget)
+	{
+		ActiveTarget->OnBeginTransformEdit.RemoveAll(this);
+		ActiveTarget->OnEndTransformEdit.RemoveAll(this);
+		ActiveTarget = nullptr;
+	}
 }
 
 
@@ -948,6 +1062,7 @@ void UTransformGizmo::SetupOnClickFunctions()
 	OnClickPressFunctions[static_cast<int>(ETransformGizmoPartIdentifier::RotateYAxis)] = &UTransformGizmo::OnClickPressRotateYAxis;
 	OnClickPressFunctions[static_cast<int>(ETransformGizmoPartIdentifier::RotateZAxis)] = &UTransformGizmo::OnClickPressRotateZAxis;
 	OnClickPressFunctions[static_cast<int>(ETransformGizmoPartIdentifier::RotateScreenSpace)] = &UTransformGizmo::OnClickPressScreenSpaceRotate;
+	OnClickPressFunctions[static_cast<int>(ETransformGizmoPartIdentifier::RotateArcball)] = &UTransformGizmo::OnClickPressArcBallRotate;
 
 	OnClickDragFunctions[static_cast<int>(ETransformGizmoPartIdentifier::TranslateXAxis)] = &UTransformGizmo::OnClickDragTranslateAxis;
 	OnClickDragFunctions[static_cast<int>(ETransformGizmoPartIdentifier::TranslateYAxis)] = &UTransformGizmo::OnClickDragTranslateAxis;
@@ -967,6 +1082,7 @@ void UTransformGizmo::SetupOnClickFunctions()
 	OnClickDragFunctions[static_cast<int>(ETransformGizmoPartIdentifier::RotateYAxis)] = &UTransformGizmo::OnClickDragRotateAxis;
 	OnClickDragFunctions[static_cast<int>(ETransformGizmoPartIdentifier::RotateZAxis)] = &UTransformGizmo::OnClickDragRotateAxis;
 	OnClickDragFunctions[static_cast<int>(ETransformGizmoPartIdentifier::RotateScreenSpace)] = &UTransformGizmo::OnClickDragScreenSpaceRotate;
+	OnClickDragFunctions[static_cast<int>(ETransformGizmoPartIdentifier::RotateArcball)] = &UTransformGizmo::OnClickDragArcBallRotate;
 
 	OnClickReleaseFunctions[static_cast<int>(ETransformGizmoPartIdentifier::TranslateXAxis)] = &UTransformGizmo::OnClickReleaseTranslateAxis;
 	OnClickReleaseFunctions[static_cast<int>(ETransformGizmoPartIdentifier::TranslateYAxis)] = &UTransformGizmo::OnClickReleaseTranslateAxis;
@@ -986,6 +1102,7 @@ void UTransformGizmo::SetupOnClickFunctions()
 	OnClickReleaseFunctions[static_cast<int>(ETransformGizmoPartIdentifier::RotateYAxis)] = &UTransformGizmo::OnClickReleaseRotateAxis;
 	OnClickReleaseFunctions[static_cast<int>(ETransformGizmoPartIdentifier::RotateZAxis)] = &UTransformGizmo::OnClickReleaseRotateAxis;
 	OnClickReleaseFunctions[static_cast<int>(ETransformGizmoPartIdentifier::RotateScreenSpace)] = &UTransformGizmo::OnClickReleaseScreenSpaceRotate;
+	OnClickReleaseFunctions[static_cast<int>(ETransformGizmoPartIdentifier::RotateArcball)] = &UTransformGizmo::OnClickReleaseArcBallRotate;
 }
 
 
@@ -1043,42 +1160,89 @@ void UTransformGizmo::UpdateHoverState(bool bInHover, ETransformGizmoPartIdentif
 	}
 }
 
-void UTransformGizmo::UpdateInteractingState(bool bInInteracting, ETransformGizmoPartIdentifier InHitPartId)
+void UTransformGizmo::UpdateInteractingState(bool bInInteracting, ETransformGizmoPartIdentifier InHitPartId, const bool bIdOnly)
 {
 	HitTarget->UpdateInteractingState(bInInteracting, static_cast<uint32>(InHitPartId));
 
-	switch (InHitPartId)
+	if (!bIdOnly)
 	{
-	case ETransformGizmoPartIdentifier::ScaleUniform:
+		switch (InHitPartId)
+		{
+		case ETransformGizmoPartIdentifier::ScaleUniform:
+			{
+				HitTarget->UpdateInteractingState(bInInteracting, static_cast<uint32>(ETransformGizmoPartIdentifier::ScaleXAxis));
+				HitTarget->UpdateInteractingState(bInInteracting, static_cast<uint32>(ETransformGizmoPartIdentifier::ScaleYAxis));
+				HitTarget->UpdateInteractingState(bInInteracting, static_cast<uint32>(ETransformGizmoPartIdentifier::ScaleZAxis));
+				HitTarget->UpdateInteractingState(bInInteracting, static_cast<uint32>(ETransformGizmoPartIdentifier::ScaleUniform));
+				GizmoElementRoot->UpdatePartVisibleState(!bInInteracting, static_cast<uint32>(ETransformGizmoPartIdentifier::ScaleXYPlanar));
+				GizmoElementRoot->UpdatePartVisibleState(!bInInteracting, static_cast<uint32>(ETransformGizmoPartIdentifier::ScaleYZPlanar));
+				GizmoElementRoot->UpdatePartVisibleState(!bInInteracting, static_cast<uint32>(ETransformGizmoPartIdentifier::ScaleXZPlanar));
+				break;
+			}
+		case ETransformGizmoPartIdentifier::ScaleXYPlanar:
+			HitTarget->UpdateInteractingState(bInInteracting, static_cast<uint32>(ETransformGizmoPartIdentifier::ScaleXAxis));
+			HitTarget->UpdateInteractingState(bInInteracting, static_cast<uint32>(ETransformGizmoPartIdentifier::ScaleYAxis));
+			break;
+		case ETransformGizmoPartIdentifier::ScaleYZPlanar:
+			HitTarget->UpdateInteractingState(bInInteracting, static_cast<uint32>(ETransformGizmoPartIdentifier::ScaleYAxis));
+			HitTarget->UpdateInteractingState(bInInteracting, static_cast<uint32>(ETransformGizmoPartIdentifier::ScaleZAxis));
+			break;
+		case ETransformGizmoPartIdentifier::ScaleXZPlanar:
+			HitTarget->UpdateInteractingState(bInInteracting, static_cast<uint32>(ETransformGizmoPartIdentifier::ScaleXAxis));
+			HitTarget->UpdateInteractingState(bInInteracting, static_cast<uint32>(ETransformGizmoPartIdentifier::ScaleZAxis));
+			break;
+		}
+	}
+}
+
+void UTransformGizmo::BeginTransformEditSequence()
+{
+	if (ensure(StateTarget))
 	{
-		HitTarget->UpdateInteractingState(bInInteracting, static_cast<uint32>(ETransformGizmoPartIdentifier::ScaleXAxis));
-		HitTarget->UpdateInteractingState(bInInteracting, static_cast<uint32>(ETransformGizmoPartIdentifier::ScaleYAxis));
-		HitTarget->UpdateInteractingState(bInInteracting, static_cast<uint32>(ETransformGizmoPartIdentifier::ScaleZAxis));
-		HitTarget->UpdateInteractingState(bInInteracting, static_cast<uint32>(ETransformGizmoPartIdentifier::ScaleUniform));
-		GizmoElementRoot->UpdatePartVisibleState(!bInInteracting, static_cast<uint32>(ETransformGizmoPartIdentifier::ScaleXYPlanar));
-		GizmoElementRoot->UpdatePartVisibleState(!bInInteracting, static_cast<uint32>(ETransformGizmoPartIdentifier::ScaleYZPlanar));
-		GizmoElementRoot->UpdatePartVisibleState(!bInInteracting, static_cast<uint32>(ETransformGizmoPartIdentifier::ScaleXZPlanar));
-		break;
+		StateTarget->BeginUpdate();
 	}
-	case ETransformGizmoPartIdentifier::ScaleXYPlanar:
-		HitTarget->UpdateInteractingState(bInInteracting, static_cast<uint32>(ETransformGizmoPartIdentifier::ScaleXAxis));
-		HitTarget->UpdateInteractingState(bInInteracting, static_cast<uint32>(ETransformGizmoPartIdentifier::ScaleYAxis));
-		break;
-	case ETransformGizmoPartIdentifier::ScaleYZPlanar:
-		HitTarget->UpdateInteractingState(bInInteracting, static_cast<uint32>(ETransformGizmoPartIdentifier::ScaleYAxis));
-		HitTarget->UpdateInteractingState(bInInteracting, static_cast<uint32>(ETransformGizmoPartIdentifier::ScaleZAxis));
-		break;
-	case ETransformGizmoPartIdentifier::ScaleXZPlanar:
-		HitTarget->UpdateInteractingState(bInInteracting, static_cast<uint32>(ETransformGizmoPartIdentifier::ScaleXAxis));
-		HitTarget->UpdateInteractingState(bInInteracting, static_cast<uint32>(ETransformGizmoPartIdentifier::ScaleZAxis));
-		break;
+
+	if (ensure(ActiveTarget))
+	{
+		if (ActiveTarget->bSetPivotMode)
+		{
+			ActiveTarget->BeginPivotEditSequence();
+		}
+		else
+		{
+			ActiveTarget->BeginTransformEditSequence();
+		}
+	}
+}
+
+void UTransformGizmo::EndTransformEditSequence()
+{
+	if (ensure(StateTarget))
+	{
+		StateTarget->EndUpdate();
+	}
+
+	if (ensure(ActiveTarget))
+	{
+		if (ActiveTarget->bSetPivotMode)
+		{
+			ActiveTarget->EndPivotEditSequence();
+		}
+		else
+		{
+			ActiveTarget->EndTransformEditSequence();
+		}
 	}
 }
 
 void UTransformGizmo::OnClickPress(const FInputDeviceRay& PressPos)
 {
+	PendingDragFunction.Reset();
+	
 	check(OnClickPressFunctions.Num() == static_cast<int>(ETransformGizmoPartIdentifier::Max));
 
+	const ETransformGizmoPartIdentifier ModeLastHitPart = GetCurrentModeLastHitPart();
+	
 	if (OnClickPressFunctions[static_cast<int>(LastHitPart)])
 	{
 		OnClickPressFunctions[static_cast<int>(LastHitPart)](this, PressPos);
@@ -1088,13 +1252,14 @@ void UTransformGizmo::OnClickPress(const FInputDeviceRay& PressPos)
 	{
 		if (HitTarget && LastHitPart != ETransformGizmoPartIdentifier::Default)
 		{
+			if (LastHitPart != ModeLastHitPart)
+			{
+				UpdateInteractingState(false, ModeLastHitPart, true);	
+			}
 			UpdateInteractingState(true, LastHitPart);
 		}
 
-		if (StateTarget)
-		{
-			StateTarget->BeginUpdate();
-		}
+		BeginTransformEditSequence();
 	}
 }
 
@@ -1110,7 +1275,17 @@ void UTransformGizmo::OnClickDrag(const FInputDeviceRay& DragPos)
 
 	if (OnClickDragFunctions[HitPartIndex])
 	{
-		OnClickDragFunctions[HitPartIndex](this, DragPos);
+		if (bDeferDrag)
+		{ // defer drag function on next tick
+			PendingDragFunction = [this, DragPos, HitPartIndex]()
+			{
+				OnClickDragFunctions[HitPartIndex](this, DragPos);
+			};
+		}
+		else
+		{ // do drag function
+			OnClickDragFunctions[HitPartIndex](this, DragPos);
+		}
 	}
 }
 
@@ -1121,7 +1296,7 @@ void UTransformGizmo::OnClickRelease(const FInputDeviceRay& ReleasePos)
 		return;
 	}
 
-	int HitPartIndex = static_cast<int>(LastHitPart);
+	const int HitPartIndex = static_cast<int>(LastHitPart);
 	check(HitPartIndex < OnClickReleaseFunctions.Num());
 
 	if (OnClickReleaseFunctions[HitPartIndex])
@@ -1129,17 +1304,17 @@ void UTransformGizmo::OnClickRelease(const FInputDeviceRay& ReleasePos)
 		OnClickReleaseFunctions[HitPartIndex](this, ReleasePos);
 	}
 
-	if (StateTarget)
-	{
-		StateTarget->EndUpdate();
-	}
+	EndTransformEditSequence();
 
 	bInInteraction = false;
 
 	if (HitTarget && LastHitPart != ETransformGizmoPartIdentifier::Default)
 	{
 		UpdateInteractingState(false, LastHitPart);
+		UpdateInteractingState(true, GetCurrentModeLastHitPart(), true);
 	}
+
+	PendingDragFunction.Reset();
 }
 
 void UTransformGizmo::OnTerminateDragSequence()
@@ -1190,6 +1365,8 @@ void UTransformGizmo::OnClickPressAxis(const FInputDeviceRay& PressPos)
 	InteractionAxisStartParam = GetNearestRayParamToInteractionAxis(PressPos);
 	InteractionAxisCurrParam = InteractionAxisStartParam;
 	bInInteraction = true;
+
+	SetModeLastHitPart(EGizmoTransformMode::Translate, LastHitPart);
 }
 
 void UTransformGizmo::OnClickDragTranslateAxis(const FInputDeviceRay& DragPos)
@@ -1243,6 +1420,8 @@ void UTransformGizmo::OnClickPressPlanar(const FInputDeviceRay& PressPos)
 		InteractionPlanarStartPoint = PressPos.WorldRay.Origin + PressPos.WorldRay.Direction * HitDepth;
 		InteractionPlanarCurrPoint = InteractionPlanarStartPoint;
 		bInInteraction = true;
+		
+		SetModeLastHitPart(EGizmoTransformMode::Translate, LastHitPart);
 	}
 }
 
@@ -1358,6 +1537,8 @@ void UTransformGizmo::OnClickPressScale(const FInputDeviceRay& PressPos)
 	InteractionScreenAxisDirection = AxisDir.GetSafeNormal();
 	InteractionScreenStartPos = InteractionScreenEndPos = InteractionScreenCurrPos = PressPos.ScreenPosition;
 	bInInteraction = true;
+
+	SetModeLastHitPart(EGizmoTransformMode::Scale, LastHitPart);
 }
 
 void UTransformGizmo::OnClickDragScaleAxis(const FInputDeviceRay& DragPos)
@@ -1445,6 +1626,8 @@ void UTransformGizmo::OnClickPressRotateXAxis(const FInputDeviceRay& PressPos)
 	InteractionAxisList = EAxisList::X;
 	InteractionScreenStartPos = InteractionScreenCurrPos = PressPos.ScreenPosition;
 	bInInteraction = true;
+
+	SetModeLastHitPart(EGizmoTransformMode::Rotate, LastHitPart);
 }
 
 void UTransformGizmo::OnClickPressRotateYAxis(const FInputDeviceRay& PressPos)
@@ -1453,6 +1636,8 @@ void UTransformGizmo::OnClickPressRotateYAxis(const FInputDeviceRay& PressPos)
 	InteractionAxisList = EAxisList::Y;
 	InteractionScreenStartPos = InteractionScreenCurrPos = PressPos.ScreenPosition;
 	bInInteraction = true;
+	
+	SetModeLastHitPart(EGizmoTransformMode::Rotate, LastHitPart);
 }
 
 void UTransformGizmo::OnClickPressRotateZAxis(const FInputDeviceRay& PressPos)
@@ -1461,6 +1646,8 @@ void UTransformGizmo::OnClickPressRotateZAxis(const FInputDeviceRay& PressPos)
 	InteractionAxisList = EAxisList::Z;
 	InteractionScreenStartPos = InteractionScreenCurrPos = PressPos.ScreenPosition;
 	bInInteraction = true;
+
+	SetModeLastHitPart(EGizmoTransformMode::Rotate, LastHitPart);
 }
 
 FVector2D UTransformGizmo::GetScreenRotateAxisDir(const FVector& InAxis0, const FVector& InAxis1)
@@ -1549,6 +1736,8 @@ void UTransformGizmo::OnClickPressScreenSpaceRotate(const FInputDeviceRay& Press
 		InteractionCurrAngle = InteractionStartAngle;
 
 		bInInteraction = true;
+		
+		SetModeLastHitPart(EGizmoTransformMode::Rotate, LastHitPart);
 	}
 }
 
@@ -1559,11 +1748,12 @@ void UTransformGizmo::OnClickDragScreenSpaceRotate(const FInputDeviceRay& DragPo
 	float HitDepth;
 	if (GetRayParamIntersectionWithInteractionPlane(DragPos, HitDepth))
 	{
-		FVector HitPoint = DragPos.WorldRay.Origin + DragPos.WorldRay.Direction * HitDepth;
-		float HitAngle = GizmoMath::ComputeAngleInPlane(HitPoint,
+		const FVector HitPoint = DragPos.WorldRay.Origin + DragPos.WorldRay.Direction * HitDepth;
+		const float HitAngle = GizmoMath::ComputeAngleInPlane(HitPoint,
 			InteractionPlanarOrigin, InteractionPlanarNormal, InteractionPlanarAxisX, InteractionPlanarAxisY);
 
-		FQuat Delta = ComputeAngularRotateDelta(InteractionCurrAngle, HitAngle);
+		const FQuat Delta = ComputeAngularRotateDelta(InteractionCurrAngle, HitAngle);
+		
 		ApplyRotateDelta(Delta);
 		InteractionCurrAngle = HitAngle;
 	}
@@ -1571,7 +1761,7 @@ void UTransformGizmo::OnClickDragScreenSpaceRotate(const FInputDeviceRay& DragPo
 
 FQuat UTransformGizmo::ComputeAngularRotateDelta(double InStartAngle, double InEndAngle)
 {
-	float DeltaAngle = InEndAngle - InStartAngle;
+	const float DeltaAngle = InEndAngle - InStartAngle;
 	return FQuat(InteractionPlanarNormal, DeltaAngle);
 }
 
@@ -1580,7 +1770,158 @@ void UTransformGizmo::OnClickReleaseScreenSpaceRotate(const FInputDeviceRay& InR
 	bInInteraction = false;
 }
 
-FVector2D UTransformGizmo::GetScreenProjectedAxis(const UGizmoViewContext* View, const FVector& InLocalAxis, const FTransform& InLocalToWorld) const
+namespace ArcBallLocals
+{	
+	/* See Holroyd's implementation that mixes a sphere + and hyperbola to avoid popping
+	 * Knud Henriksen, Jon Sporring, and Kasper Hornbaek, “Virtual trackballs revisited”,
+	 * IEEE Transactions on Visualization and Computer Graphics, vol. 10, no. 2, pp. 206–216, 2004.
+	 */
+	bool GetSphereAndHyperbolicProjection(
+		const FVector& SphereOrigin,
+		const double SphereRadius,
+		const FVector& RayOrigin,
+		const FVector& RayDirection,
+		const UGizmoViewContext& ViewContext,
+		FVector& OutProjection)
+	{
+		const FVector CircleNormal = -ViewContext.GetViewDirection();
+		
+		// if ray is parallel to circle, no hit
+		if (FMath::IsNearlyZero(FVector::DotProduct(CircleNormal, RayDirection)))
+		{
+			return false;
+		}
+
+		const FPlane Plane(SphereOrigin, CircleNormal);
+		const double Param = FMath::RayPlaneIntersectionParam(RayOrigin, RayDirection, Plane);
+
+		if (Param < 0)
+		{
+			return false;
+		}
+
+		const FVector HitPoint = RayOrigin + RayDirection * Param;
+		FVector Offset = HitPoint - SphereOrigin;
+
+		// switch to screen space
+		FVector OffsetProjection = ViewContext.ViewMatrices.GetInvViewMatrix().InverseTransformVector(Offset);
+
+		const double OffsetSquared = Offset.SizeSquared();
+		const double CircleRadiusSquared = SphereRadius * SphereRadius;
+		
+		OffsetProjection.Z = (OffsetSquared <= CircleRadiusSquared * 0.5) ?
+			-FMath::Sqrt(CircleRadiusSquared - OffsetSquared) : // spherical projection
+			-CircleRadiusSquared * 0.5 / Offset.Length(); // hyperbolic projection
+
+		// switch back to world space
+		Offset = ViewContext.ViewMatrices.GetInvViewMatrix().TransformVector(OffsetProjection);
+		OutProjection = SphereOrigin + Offset;
+		
+		return true;
+	}
+}
+
+void UTransformGizmo::OnClickPressArcBallRotate(const FInputDeviceRay& PressPos)
+{
+	check(GizmoViewContext);
+	
+	const FVector& RayOrigin = PressPos.WorldRay.Origin;
+	const FVector& RayDir = PressPos.WorldRay.Direction;		
+	const double SphereRadius = GetArcBallWorldRadius();
+
+	StartRotation = CurrentRotation = CurrentTransform.GetRotation();
+	InteractionPlanarOrigin = CurrentTransform.GetLocation();
+	InteractionPlanarNormal = -GizmoViewContext->GetViewDirection();
+	InteractionPlanarAxisX = GizmoViewContext->GetViewUp();
+	InteractionPlanarAxisY = GizmoViewContext->GetViewRight();
+	InteractionAxisList = EAxisList::XYZ;
+
+	auto NeedsInteraction = [&]()
+	{
+		const bool bIntersect = IntersectionUtil::RaySphereTest(RayOrigin, RayDir, InteractionPlanarOrigin, SphereRadius);
+		if (bIntersect)
+		{
+			return true;
+		}
+		
+		if (bIndirectManipulation)
+		{
+			// change the arc ball center in indirect manipulation if we didn't hit the sphere
+			float HitDepth;
+			if (GetRayParamIntersectionWithInteractionPlane(PressPos, HitDepth))
+			{
+				InteractionPlanarOrigin = PressPos.WorldRay.Origin + PressPos.WorldRay.Direction * HitDepth;
+				return true;
+			}
+		}
+		
+		return false;
+	};
+
+	if (NeedsInteraction())
+	{
+		// project on sphere
+		ArcBallLocals::GetSphereAndHyperbolicProjection(
+			InteractionPlanarOrigin, SphereRadius, RayOrigin, RayDir, *GizmoViewContext,
+			InteractionArcBallStartPoint);
+	
+		bInInteraction = true;
+	
+		SetModeLastHitPart(EGizmoTransformMode::Rotate, LastHitPart);
+	}
+}
+
+void UTransformGizmo::OnClickDragArcBallRotate(const FInputDeviceRay& DragPos)
+{
+	const FVector& RayOrigin = DragPos.WorldRay.Origin;
+	const FVector& RayDir = DragPos.WorldRay.Direction;		
+	const float SphereRadius = GetArcBallWorldRadius();
+
+	// compute projection
+	ArcBallLocals::GetSphereAndHyperbolicProjection(
+		InteractionPlanarOrigin, SphereRadius, RayOrigin, RayDir, *GizmoViewContext,
+		InteractionArcBallCurrPoint);
+
+	if ((InteractionArcBallCurrPoint-InteractionArcBallStartPoint).Length() <= 0.0)
+	{
+		return;
+	}
+
+	// compute rotation
+	const FVector Axis1 = (InteractionArcBallCurrPoint - InteractionPlanarOrigin).GetSafeNormal();
+	const FVector Axis0 = (InteractionArcBallStartPoint - InteractionPlanarOrigin).GetSafeNormal();
+	
+	const FQuat DeltaQ = FQuat::FindBetweenNormals(Axis0, Axis1);
+	
+	// apply rotation
+	const FQuat FinalRotation = DeltaQ * StartRotation;
+	const FQuat InvCurrentRot = CurrentRotation.Inverse();
+	
+	const FQuat DeltaRot = (FinalRotation * InvCurrentRot).GetNormalized();
+	if (!DeltaRot.IsIdentity())
+	{
+		ApplyRotateDelta(DeltaRot);
+		CurrentRotation = DeltaRot * CurrentRotation; 
+	}
+}
+
+void UTransformGizmo::OnClickReleaseArcBallRotate(const FInputDeviceRay& ReleasePos)
+{
+	bInInteraction = false;
+}
+
+float UTransformGizmo::GetArcBallWorldRadius() const
+{
+	const float PixelToWorldScale = GizmoRenderingUtil::CalculateLocalPixelToWorldScale(GizmoViewContext, CurrentTransform.GetLocation());
+	return RotateArcballSphereRadius * GetSizeCoefficient() * PixelToWorldScale;
+}
+
+float UTransformGizmo::GetSizeCoefficient() const
+{
+	return CustomizationFunction ? CustomizationFunction().SizeCoefficient : 1.5f;
+}
+
+FVector2D UTransformGizmo::GetScreenProjectedAxis(const UGizmoViewContext* View, const FVector& InLocalAxis, const FTransform& InLocalToWorld)
 {
 	FVector2D Origin;
 	FVector2D AxisEnd;
@@ -1617,15 +1958,15 @@ void UTransformGizmo::ApplyTranslateDelta(const FVector& InTranslateDelta)
 void UTransformGizmo::ApplyRotateDelta(const FQuat& InRotateDelta)
 {
 	// Applies rot delta after the current rotation
-	FQuat NewRotation = InRotateDelta * CurrentTransform.GetRotation();
+	const FQuat NewRotation = InRotateDelta * CurrentTransform.GetRotation();
 	CurrentTransform.SetRotation(NewRotation);
 	ActiveTarget->SetTransform(CurrentTransform);
 }
 
 void UTransformGizmo::ApplyScaleDelta(const FVector& InScaleDelta)
 {
-	FVector StartScale = CurrentTransform.GetScale3D();
-	FVector NewScale = StartScale + InScaleDelta;
+	const FVector StartScale = CurrentTransform.GetScale3D();
+	const FVector NewScale = StartScale + InScaleDelta;
 	CurrentTransform.SetScale3D(NewScale);
 	ActiveTarget->SetTransform(CurrentTransform);
 }
diff --git a/Engine/Source/Editor/Experimental/EditorInteractiveToolsFramework/Private/EditorInteractiveGizmoManager.cpp b/Engine/Source/Editor/Experimental/EditorInteractiveToolsFramework/Private/EditorInteractiveGizmoManager.cpp
--- a/Engine/Source/Editor/Experimental/EditorInteractiveToolsFramework/Private/EditorInteractiveGizmoManager.cpp
+++ b/Engine/Source/Editor/Experimental/EditorInteractiveToolsFramework/Private/EditorInteractiveGizmoManager.cpp
@@ -21,6 +21,7 @@
 #include "ShowFlags.h"
 #include "Templates/Casts.h"
 #include "ToolContextInterfaces.h"
+#include "EditorGizmos/TransformGizmo.h"
 
 class FCanvas;
 
@@ -39,6 +40,7 @@ UEditorInteractiveGizmoManager::UEditorInteractiveGizmoManager() :
 	UInteractiveGizmoManager()
 {
 	Registry = NewObject<UEditorInteractiveGizmoRegistry>();
+	bShowEditorGizmos = (CVarUseLegacyWidget.GetValueOnGameThread() == 0);
 }
 
 
@@ -84,68 +86,9 @@ void UEditorInteractiveGizmoManager::GetQualifiedEditorGizmoBuilders(EEditorGizm
 	}
 }
 
-TArray<UInteractiveGizmo*> UEditorInteractiveGizmoManager::CreateGizmosForCurrentSelectionSet()
+UTransformGizmo* UEditorInteractiveGizmoManager::FindDefaultTransformGizmo() const
 {
-	if (bShowEditorGizmos)
-	{
-		FToolBuilderState CurrentSceneState;
-		QueriesAPI->GetCurrentSelectionState(CurrentSceneState);
-
-		TArray<TObjectPtr<UInteractiveGizmoBuilder>> FoundBuilders;
-		GetQualifiedEditorGizmoBuilders(EEditorGizmoCategory::Primary, CurrentSceneState, MutableView(FoundBuilders));
-
-		bool bHasPrimaryBuilder = (FoundBuilders.Num() > 0);
-
-		if (!bHasPrimaryBuilder)
-		{
-			GetQualifiedEditorGizmoBuilders(EEditorGizmoCategory::Accessory, CurrentSceneState, MutableView(FoundBuilders));
-
-			UInteractiveGizmoBuilder* TransformBuilder = GetTransformGizmoBuilder();
-
-			FoundBuilders.Add(TransformBuilder);
-		}
-
-		TArray<TObjectPtr<UInteractiveGizmo>> NewGizmos;
-
-		for (UInteractiveGizmoBuilder* FoundBuilder : FoundBuilders)
-		{
-			if (UInteractiveGizmo* NewGizmo = FoundBuilder->BuildGizmo(CurrentSceneState))
-			{
-				if (IEditorInteractiveGizmoSelectionBuilder* SelectionBuilder = Cast<IEditorInteractiveGizmoSelectionBuilder>(FoundBuilder))
-				{
-					SelectionBuilder->UpdateGizmoForSelection(NewGizmo, CurrentSceneState);
-				}
-
-				// register new active input behaviors
-				InputRouter->RegisterSource(NewGizmo);
-
-				NewGizmos.Add(NewGizmo);
-			}
-		}
-
-		PostInvalidation();
-
-		for (UInteractiveGizmo* NewGizmo : NewGizmos)
-		{
-			FActiveEditorGizmo ActiveGizmo = { NewGizmo, nullptr };
-			ActiveEditorGizmos.Add(ActiveGizmo);
-		}
-
-		return NewGizmos;
-	}
-
-	return TArray<UInteractiveGizmo*>();
-}
-
-UInteractiveGizmoBuilder* UEditorInteractiveGizmoManager::GetTransformGizmoBuilder()
-{
-	UEditorInteractiveGizmoSubsystem* GizmoSubsystem = GEditor->GetEditorSubsystem<UEditorInteractiveGizmoSubsystem>();
-	if (ensure(GizmoSubsystem))
-	{
-		return GizmoSubsystem->GetTransformGizmoBuilder();
-	}
-
-	return nullptr;
+	return Cast<UTransformGizmo>( FindGizmoByInstanceIdentifier(TransformInstanceIdentifier()) );
 }
 
 bool UEditorInteractiveGizmoManager::DestroyEditorGizmo(UInteractiveGizmo* Gizmo)
@@ -185,14 +128,35 @@ void UEditorInteractiveGizmoManager::DestroyAllEditorGizmos()
 	PostInvalidation();
 }
 
-void UEditorInteractiveGizmoManager::HandleEditorSelectionSetChanged(const UTypedElementSelectionSet* InSelectionSet)
+UInteractiveGizmo* UEditorInteractiveGizmoManager::CreateGizmo(const FString& BuilderIdentifier, const FString& InstanceIdentifier, void* Owner)
 {
-	DestroyAllEditorGizmos();
-
-	if (InSelectionSet && InSelectionSet->HasSelectedElements())
+	if (BuilderIdentifier == TransformBuilderIdentifier() && InstanceIdentifier == TransformInstanceIdentifier())
 	{
-		CreateGizmosForCurrentSelectionSet();
+		// return the default transform gizmo if it already exists.
+		if (UTransformGizmo* ExistingGizmo = FindDefaultTransformGizmo())
+		{
+			return ExistingGizmo;
+		}
+
+		// create a new one
+		UInteractiveGizmo* NewGizmo = Super::CreateGizmo(BuilderIdentifier, InstanceIdentifier, Owner);
+		if (!NewGizmo)
+		{
+			return nullptr;
+		}
+		
+		if (IEditorInteractiveGizmoSelectionBuilder* SelectionBuilder = Cast<IEditorInteractiveGizmoSelectionBuilder>(GizmoBuilders[BuilderIdentifier]))
+		{
+			FToolBuilderState CurrentSceneState;
+			QueriesAPI->GetCurrentSelectionState(CurrentSceneState);
+			
+			SelectionBuilder->UpdateGizmoForSelection(NewGizmo, CurrentSceneState);
+		}
+
+		return NewGizmo;
 	}
+	
+	return Super::CreateGizmo(BuilderIdentifier, InstanceIdentifier, Owner);
 }
 
 // @todo move this to a gizmo context object
@@ -218,11 +182,18 @@ void UEditorInteractiveGizmoManager::UpdateActiveEditorGizmos()
 	if (bShowEditorGizmos != bNewShowEditorGizmos)
 	{
 		bShowEditorGizmos = bNewShowEditorGizmos;
-		if (bShowEditorGizmos)
+
+		if (UTransformGizmo* Gizmo = FindDefaultTransformGizmo())
 		{
-			CreateGizmosForCurrentSelectionSet();
+			Gizmo->SetVisibility(bShowEditorGizmos ? bEditorModeToolsSupportsWidgetDrawing : false);
+
+			if (!bEnableEditorGizmos)
+			{
+				DestroyGizmo(Gizmo);	
+			}
 		}
-		else
+
+		if (!bShowEditorGizmos)
 		{
 			DestroyAllEditorGizmos();
 		}
@@ -231,9 +202,9 @@ void UEditorInteractiveGizmoManager::UpdateActiveEditorGizmos()
 
 void UEditorInteractiveGizmoManager::Tick(float DeltaTime)
 {
-	Super::Tick(DeltaTime);
-
 	UpdateActiveEditorGizmos();
+	
+	Super::Tick(DeltaTime);
 
 	for (FActiveEditorGizmo& ActiveEditorGizmo : ActiveEditorGizmos)
 	{
@@ -268,4 +239,16 @@ void UEditorInteractiveGizmoManager::DrawHUD(FCanvas* Canvas, IToolsContextRende
 	}
 }
 
+const FString& UEditorInteractiveGizmoManager::TransformInstanceIdentifier()
+{
+	static const FString Identifier(TEXT("EditorTransformGizmoInstance"));
+	return Identifier;	
+}
+
+const FString& UEditorInteractiveGizmoManager::TransformBuilderIdentifier()
+{
+	static const FString Identifier(TEXT("EditorTransformGizmoBuilder"));
+	return Identifier;
+}
+
 #undef LOCTEXT_NAMESPACE
diff --git a/Engine/Source/Editor/Experimental/EditorInteractiveToolsFramework/Public/EdModeInteractiveToolsContext.h b/Engine/Source/Editor/Experimental/EditorInteractiveToolsFramework/Public/EdModeInteractiveToolsContext.h
--- a/Engine/Source/Editor/Experimental/EditorInteractiveToolsFramework/Public/EdModeInteractiveToolsContext.h
+++ b/Engine/Source/Editor/Experimental/EditorInteractiveToolsFramework/Public/EdModeInteractiveToolsContext.h
@@ -173,18 +173,6 @@ public:
 	bool GetAbsoluteWorldSnappingEnabled() const { return bEnableAbsoluteWorldSnapping; }
 
 protected:
-	/**
-	 * Called upon Level Editor Created in order to bind to EditorElementSelectionPtr changed
-	 * event since global selection set is not initialized before the tools context.
-	 * @param InSelectionSet - typed element selection set which invoked this selection changed call
-	 */
-	virtual void OnLevelEditorCreated(TSharedPtr<ILevelEditor> InLevelEditor);
-
-	/**
-	 * Handle Editor selection changes
-	 * @param InSelectionSet - typed element selection set which invoked this selection changed call
-	 */
-	virtual void OnEditorSelectionSetChanged(const UTypedElementSelectionSet* InSelectionSet);
 
 	// we hide these 
 	virtual void Initialize(IToolsContextQueriesAPI* QueriesAPI, IToolsContextTransactionsAPI* TransactionsAPI) override;
diff --git a/Engine/Source/Editor/Experimental/EditorInteractiveToolsFramework/Public/EditorGizmos/EditorTransformGizmo.h b/Engine/Source/Editor/Experimental/EditorInteractiveToolsFramework/Public/EditorGizmos/EditorTransformGizmo.h
--- a/Engine/Source/Editor/Experimental/EditorInteractiveToolsFramework/Public/EditorGizmos/EditorTransformGizmo.h
+++ b/Engine/Source/Editor/Experimental/EditorInteractiveToolsFramework/Public/EditorGizmos/EditorTransformGizmo.h
@@ -21,14 +21,30 @@ class EDITORINTERACTIVETOOLSFRAMEWORK_API UEditorTransformGizmo : public UTransf
 {
 	GENERATED_BODY()
 
+	/**  UTransformGizmo override */
+	virtual void Render(IToolsContextRenderAPI* RenderAPI) override;
+	
 protected:
 
 	/** Apply translate delta to transform proxy */
-	virtual void ApplyTranslateDelta(const FVector& InTranslateDelta);
+	virtual void ApplyTranslateDelta(const FVector& InTranslateDelta) override;
 
 	/** Apply rotate delta to transform proxy */
-	virtual void ApplyRotateDelta(const FQuat& InRotateDelta);
+	virtual void ApplyRotateDelta(const FQuat& InRotateDelta) override;
 
 	/** Apply scale delta to transform proxy */
-	virtual void ApplyScaleDelta(const FVector& InScaleDelta);
+	virtual void ApplyScaleDelta(const FVector& InScaleDelta) override;
+
+	/**  UTransformGizmo override */
+	virtual void SetActiveTarget(
+		UTransformProxy* Target,
+		IToolContextTransactionProvider* TransactionProvider = nullptr,
+		IGizmoStateTarget* InStateTarget = nullptr) override;
+
+	/**
+	 * Functions to listen to gizmo transform begin/end.
+	 * They are currently used to set data on the legacy widget as some viewport clients rely on current axis to be set.
+	 */
+	void OnGizmoTransformBegin(UTransformProxy* TransformProxy) const;
+	void OnGizmoTransformEnd(UTransformProxy* TransformProxy) const;;
 };
\ No newline at end of file
diff --git a/Engine/Source/Editor/Experimental/EditorInteractiveToolsFramework/Public/EditorGizmos/EditorTransformGizmoBuilder.h b/Engine/Source/Editor/Experimental/EditorInteractiveToolsFramework/Public/EditorGizmos/EditorTransformGizmoBuilder.h
--- a/Engine/Source/Editor/Experimental/EditorInteractiveToolsFramework/Public/EditorGizmos/EditorTransformGizmoBuilder.h
+++ b/Engine/Source/Editor/Experimental/EditorInteractiveToolsFramework/Public/EditorGizmos/EditorTransformGizmoBuilder.h
@@ -2,7 +2,6 @@
 
 #pragma once
 
-#include "BaseGizmos/TransformProxy.h"
 #include "CoreMinimal.h"
 #include "EditorInteractiveGizmoSelectionBuilder.h"
 #include "InteractiveGizmoBuilder.h"
@@ -14,6 +13,7 @@
 
 class UInteractiveGizmo;
 class UObject;
+class UEditorTransformGizmoContextObject;
 struct FToolBuilderState;
 
 UCLASS()
@@ -29,6 +29,9 @@ public:
 
 	// If set, this function will be passed to UTransformGizmo instances to override the default material and display size.
 	TFunction<const FGizmoCustomization()> CustomizationFunction;
+
+private:
+	static const UEditorTransformGizmoContextObject* GetTransformGizmoContext(const FToolBuilderState& InSceneState);
 };
 
 
diff --git a/Engine/Source/Editor/Experimental/EditorInteractiveToolsFramework/Public/EditorGizmos/EditorTransformGizmoSource.h b/Engine/Source/Editor/Experimental/EditorInteractiveToolsFramework/Public/EditorGizmos/EditorTransformGizmoSource.h
--- a/Engine/Source/Editor/Experimental/EditorInteractiveToolsFramework/Public/EditorGizmos/EditorTransformGizmoSource.h
+++ b/Engine/Source/Editor/Experimental/EditorInteractiveToolsFramework/Public/EditorGizmos/EditorTransformGizmoSource.h
@@ -3,6 +3,7 @@
 #pragma once
 
 #include "CoreMinimal.h"
+
 #include "EditorGizmos/TransformGizmoInterfaces.h"
 #include "Math/Axis.h"
 #include "ToolContextInterfaces.h"
@@ -14,6 +15,7 @@
 
 #include "EditorTransformGizmoSource.generated.h"
 
+class UEditorTransformGizmoContextObject;
 class FEditorModeTools;
 class FEditorViewportClient;
 class FSceneView;
@@ -40,44 +42,43 @@ public:
 	/**
 	 * @return The current display mode for the Editor transform gizmo
 	 */
-	virtual EGizmoTransformMode GetGizmoMode() const;
+	virtual EGizmoTransformMode GetGizmoMode() const override;
 
 	/**
 	 * @return The current axes to draw for the specified mode
 	 */
-	virtual EAxisList::Type GetGizmoAxisToDraw(EGizmoTransformMode InWidgetMode) const;
+	virtual EAxisList::Type GetGizmoAxisToDraw(EGizmoTransformMode InWidgetMode) const override;
 
 	/**
 	 * @return The coordinate system space (world or local) to display the widget in
 	 */
-	virtual EToolContextCoordinateSystem GetGizmoCoordSystemSpace() const;
+	virtual EToolContextCoordinateSystem GetGizmoCoordSystemSpace() const override;
 
 	/**
 	 * Returns a scale factor for the gizmo
 	 */
-	virtual float GetGizmoScale() const;
+	virtual float GetGizmoScale() const override;
 
 	/**
 	 * Whether the gizmo is visible
 	 */
-	virtual bool GetVisible() const;
+	virtual bool GetVisible() const override;
 
 	/**
  	 * Get current scale type
 	 */
-	virtual EGizmoTransformScaleType GetScaleType() const;
+	virtual EGizmoTransformScaleType GetScaleType() const override;
 
-public:
-	static UEditorTransformGizmoSource* Construct(
-		UObject* Outer = (UObject*)GetTransientPackage())
-	{
-		return NewObject<UEditorTransformGizmoSource>(Outer);
-	}
+	static UEditorTransformGizmoSource* CreateNew(
+		UObject* Outer = (UObject*)GetTransientPackage(),
+		const UEditorTransformGizmoContextObject* InContext = nullptr);
 
 protected:
 
-	FEditorModeTools& GetModeTools() const;
+	const FEditorModeTools& GetModeTools() const;
 
-	FEditorViewportClient* GetViewportClient() const;
+	const FEditorViewportClient* GetViewportClient() const;
+	
+	TWeakObjectPtr<const UEditorTransformGizmoContextObject> WeakContext = nullptr;
 };
 
diff --git a/Engine/Source/Editor/Experimental/EditorInteractiveToolsFramework/Public/EditorGizmos/EditorTransformGizmoUtil.h b/Engine/Source/Editor/Experimental/EditorInteractiveToolsFramework/Public/EditorGizmos/EditorTransformGizmoUtil.h
new file mode 100644
--- /dev/null
+++ b/Engine/Source/Editor/Experimental/EditorInteractiveToolsFramework/Public/EditorGizmos/EditorTransformGizmoUtil.h
@@ -0,0 +1,97 @@
+﻿// Copyright Epic Games, Inc. All Rights Reserved.
+
+#pragma once
+
+#include "UObject/Object.h"
+#include "Containers/UnrealString.h"
+
+#include "EditorTransformGizmoUtil.generated.h"
+
+class UInteractiveToolsContext;
+class UInteractiveToolManager;
+class UTransformGizmo;
+class FEditorModeTools;
+class FEditorViewportClient;
+
+namespace UE
+{
+	
+namespace EditorTransformGizmoUtil
+{
+	/**
+	 * The functions below are helper functions that simplify usage of a UEditorTransformGizmoContextObject
+	 * that is registered as a ContextStoreObject in an InteractiveToolsContext
+	 */
+
+	/**
+	 * If one does not already exist, create a new instance of UEditorTransformGizmoContextObject and add it to the
+	 * ToolsContext's ContextObjectStore
+	 * @return true if the ContextObjectStore now has a UEditorTransformGizmoContextObject (whether it already existed, or was created)
+	 */
+	EDITORINTERACTIVETOOLSFRAMEWORK_API bool RegisterTransformGizmoContextObject(FEditorModeTools* InModeTools);
+
+	/**
+	 * Remove any existing UEditorTransformGizmoContextObject from the ToolsContext's ContextObjectStore
+	 * @return true if the ContextObjectStore no longer has a UEditorTransformGizmoContextObject (whether it was removed, or did not exist)
+	 */
+	EDITORINTERACTIVETOOLSFRAMEWORK_API bool UnregisterTransformGizmoContextObject(FEditorModeTools* InModeTools);
+
+	/**
+	 * Spawn an editor like Transform Gizmo. ToolManager's ToolsContext must have a UEditorTransformGizmoContextObject registered 
+	 */
+	EDITORINTERACTIVETOOLSFRAMEWORK_API UTransformGizmo* CreateTransformGizmo(
+		UInteractiveToolManager* InToolManager, const FString& InInstanceIdentifier = FString(), void* InOwner = nullptr);
+	
+	/**
+	 * Spawn an editor like Transform Gizmo. ToolManager's ToolsContext must have a UEditorTransformGizmoContextObject registered 
+	 */
+	EDITORINTERACTIVETOOLSFRAMEWORK_API UTransformGizmo* GetDefaultTransformGizmo(UInteractiveToolManager* InToolManager);
+
+	/**
+	 * Returns the existing default Transform Gizmo if any. 
+	 */
+	EDITORINTERACTIVETOOLSFRAMEWORK_API UTransformGizmo* FindDefaultTransformGizmo(UInteractiveToolManager* InToolManager);
+
+	/**
+     * Removes the existing default Transform Gizmo if any. 
+     */
+	EDITORINTERACTIVETOOLSFRAMEWORK_API void RemoveDefaultTransformGizmo(UInteractiveToolManager* InToolManager);
+}
+	
+}
+
+/**
+ * UEditorTransformGizmoContextObject is a utility object that registers a Gizmo Builder for UTransformGizmo.
+ * (see UCombinedTransformGizmoContextObject for more details)
+ */
+
+UCLASS(Transient, MinimalAPI)
+class UEditorTransformGizmoContextObject : public UObject
+{
+	GENERATED_BODY()
+
+public:
+	
+	EDITORINTERACTIVETOOLSFRAMEWORK_API void RegisterGizmosWithManager(UInteractiveToolManager* InToolManager);
+	EDITORINTERACTIVETOOLSFRAMEWORK_API void UnregisterGizmosWithManager(UInteractiveToolManager* InToolManager);
+
+	EDITORINTERACTIVETOOLSFRAMEWORK_API void Initialize(FEditorModeTools* InModeTools);
+	EDITORINTERACTIVETOOLSFRAMEWORK_API void Shutdown();
+
+	EDITORINTERACTIVETOOLSFRAMEWORK_API const FEditorModeTools* GetModeTools() const;
+	
+private:
+
+	EDITORINTERACTIVETOOLSFRAMEWORK_API void UpdateGizmo(const TArray<FEditorViewportClient*>& InViewportClients) const;
+	
+	EDITORINTERACTIVETOOLSFRAMEWORK_API void InitializeCVarBinding();
+	EDITORINTERACTIVETOOLSFRAMEWORK_API void InitializeViewportsBinding();
+	
+	EDITORINTERACTIVETOOLSFRAMEWORK_API void RemoveCVarBinding();
+	EDITORINTERACTIVETOOLSFRAMEWORK_API void RemoveViewportsBinding();
+
+	FEditorModeTools* ModeTools = nullptr;
+	FDelegateHandle ViewportClientsChangedHandle;
+	FDelegateHandle UseLegacyChangedHandled;
+	bool bGizmosRegistered = false;
+};
diff --git a/Engine/Source/Editor/Experimental/EditorInteractiveToolsFramework/Public/EditorGizmos/EditorTransformProxy.h b/Engine/Source/Editor/Experimental/EditorInteractiveToolsFramework/Public/EditorGizmos/EditorTransformProxy.h
--- a/Engine/Source/Editor/Experimental/EditorInteractiveToolsFramework/Public/EditorGizmos/EditorTransformProxy.h
+++ b/Engine/Source/Editor/Experimental/EditorInteractiveToolsFramework/Public/EditorGizmos/EditorTransformProxy.h
@@ -14,6 +14,8 @@
 
 #include "EditorTransformProxy.generated.h"
 
+class UEditorTransformGizmoContextObject;
+class FEditorViewportClient;
 class UObject;
 
 /**
@@ -21,9 +23,7 @@ class UObject;
  * returns the transform that defines the current space of the default
  * Editor transform gizmo for a given mode manager / viewport.
  * 
- * @todo Currently this defaults internally to GLevelEditorModeManager()
- * but eventually it should be possible to set and use a different mode
- * manager.
+ * This defaults internally to GLevelEditorModeManager() if none provided when constructing the proxy.
  */
 UCLASS(Transient)
 class EDITORINTERACTIVETOOLSFRAMEWORK_API UEditorTransformProxy : public UTransformProxy
@@ -52,5 +52,16 @@ public:
 
 	/** Input scale delta to be applied in local space of the current transform. */
 	virtual void InputScaleDelta(const FVector& InDeltaScale, EAxisList::Type InAxisList);
+
+	/** Set legacy widget axis temporarily because FEditorViewportClient overrides may expect it */
+	void SetCurrentAxis(const EAxisList::Type InAxisList) const;
+
+	static UEditorTransformProxy* CreateNew(const UEditorTransformGizmoContextObject* InContext = nullptr);
+	
+private:
+	
+	FEditorViewportClient* GetViewportClient() const;
+
+	TWeakObjectPtr<const UEditorTransformGizmoContextObject> WeakContext;
 };
 
diff --git a/Engine/Source/Editor/Experimental/EditorInteractiveToolsFramework/Public/EditorGizmos/TransformGizmo.h b/Engine/Source/Editor/Experimental/EditorInteractiveToolsFramework/Public/EditorGizmos/TransformGizmo.h
--- a/Engine/Source/Editor/Experimental/EditorInteractiveToolsFramework/Public/EditorGizmos/TransformGizmo.h
+++ b/Engine/Source/Editor/Experimental/EditorInteractiveToolsFramework/Public/EditorGizmos/TransformGizmo.h
@@ -84,7 +84,6 @@ enum class ETransformGizmoPartIdentifier
 	RotateZAxis,
 	RotateScreenSpace,
 	RotateArcball,
-	RotateArcballInnerCircle,
 	ScaleAll,
 	ScaleXAxis,
 	ScaleYAxis,
@@ -118,8 +117,6 @@ public:
 	static constexpr float TranslateScreenSpaceHandleSize = 14.0f;
 
 	// Rotate constants
-	static constexpr float RotateArcballInnerRadius = 8.0f;
-	static constexpr float RotateArcballOuterRadius = 10.0f;
 	static constexpr float RotateArcballSphereRadius = 70.0f;
 	static constexpr float RotateAxisOuterRadius = 73.0f;
 	static constexpr float RotateAxisInnerRadius = 1.25f;
@@ -148,7 +145,7 @@ public:
 
 	static constexpr FLinearColor RotateScreenSpaceCircleColor = WhiteColor;
 	static constexpr FLinearColor RotateOuterCircleColor = GreyColor;
-	static constexpr FLinearColor RotateArcballCircleColor = WhiteColor;
+	static constexpr FLinearColor RotateArcballCircleColor = FLinearColor(0.50f, 0.50f, 0.50f, 0.1f);
 
 	static constexpr float LargeOuterAlpha = 0.5f;
 
@@ -166,7 +163,7 @@ public:
 	// UInteractiveGizmo overrides
 	virtual void Setup() override;
 	virtual void Shutdown() override;
-	virtual void Render(IToolsContextRenderAPI* RenderAPI);
+	virtual void Render(IToolsContextRenderAPI* RenderAPI) override;
 	virtual void Tick(float DeltaTime) override;
 
 	// IHoverBehaviorTarget implementation
@@ -186,6 +183,7 @@ public:
 	 * Set the active target object for the Gizmo
 	 * @param Target active target
 	 * @param TransactionProvider optional IToolContextTransactionProvider implementation to use - by default uses GizmoManager
+	 * @param InStateTarget optional IGizmoStateTarget implementation to use - will create a new UGizmoObjectModifyStateTarget none provided
 	 */
 	virtual void SetActiveTarget(UTransformProxy* Target, IToolContextTransactionProvider* TransactionProvider = nullptr, IGizmoStateTarget* InStateTarget = nullptr);
 
@@ -320,13 +318,9 @@ protected:
 	UPROPERTY()
 	TObjectPtr<UGizmoElementCircle> RotateOuterCircleElement;
 
-	/** Rotate arcball outer circle */
-	UPROPERTY()
-	TObjectPtr<UGizmoElementCircle> RotateArcballOuterElement;
-
 	/** Rotate arcball inner circle */
 	UPROPERTY()
-	TObjectPtr<UGizmoElementCircle> RotateArcballInnerElement;
+	TObjectPtr<UGizmoElementCircle> RotateArcballElement;
 
 	/** Rotate screen space circle */
 	UPROPERTY()
@@ -451,8 +445,14 @@ protected:
 	virtual void UpdateHoverState(bool bInHover, ETransformGizmoPartIdentifier InPartId);
 
 	/** Update interacting state for given part id */
-	virtual void UpdateInteractingState(bool bInInteracting, ETransformGizmoPartIdentifier InPartId);
+	virtual void UpdateInteractingState(bool bInInteracting, ETransformGizmoPartIdentifier InPartId, const bool bIdOnly = false);
 
+	/** Called at the start of a sequence of gizmo transform edits */
+	virtual void BeginTransformEditSequence();
+
+	/** Called at the end of a sequence of gizmo transform edits. */
+	virtual void EndTransformEditSequence();
+	
 	/**
 	 * Translate axis click-drag handling methods 
 	 */ 
@@ -556,6 +556,25 @@ protected:
 
 	/** Compute rotate delta based on start/end angles */
 	virtual FQuat ComputeAngularRotateDelta(double InStartAngle, double InEndAngle);
+	
+	/**
+	 * Arc ball rotate interaction methods
+	 */
+
+	/** Handle click press for arc ball rotate */
+	virtual void OnClickPressArcBallRotate(const FInputDeviceRay& PressPos);
+
+	/** Handle click drag for arc ball rotate */
+	virtual void OnClickDragArcBallRotate(const FInputDeviceRay& DragPos);
+
+	/** Handle click release for arc ball rotate */
+	virtual void OnClickReleaseArcBallRotate(const FInputDeviceRay& ReleasePos);
+
+	/** Get the arc ball sphere world radius */
+	float GetArcBallWorldRadius() const;
+
+	/** */
+	float GetSizeCoefficient() const;
 
 	/**
 	* Scale click-drag handling methods
@@ -614,7 +633,7 @@ protected:
 	 */
 
 	/** Returns 2D vector projection of input axis onto input view plane */
-	FVector2D GetScreenProjectedAxis(const UGizmoViewContext* View, const FVector& InLocalAxis, const FTransform& InLocalToWorld = FTransform::Identity) const;
+	static FVector2D GetScreenProjectedAxis(const UGizmoViewContext* View, const FVector& InLocalAxis, const FTransform& InLocalToWorld = FTransform::Identity);
 
 	/**
 	 * Apply transform delta methods
@@ -699,6 +718,12 @@ protected:
 	UPROPERTY()
 	ETransformGizmoPartIdentifier LastHitPart = ETransformGizmoPartIdentifier::Default;
 
+	/** Last hit part per mode */
+	ETransformGizmoPartIdentifier LastHitPartPerMode[static_cast<uint8>(EGizmoTransformMode::Max)];
+	
+	ETransformGizmoPartIdentifier GetCurrentModeLastHitPart() const;
+	void SetModeLastHitPart(const EGizmoTransformMode InMode, const ETransformGizmoPartIdentifier InIdentifier);
+	
 	//
 	// The values below are used in the context of a single click-drag interaction, ie if bInInteraction = true
 	// They otherwise should be considered uninitialized
@@ -772,4 +797,29 @@ protected:
 	UPROPERTY()
 	FVector2D InteractionScreenCurrPos;
 
+	/** Active interaction arc ball start point */
+	UPROPERTY()
+	FVector InteractionArcBallStartPoint;
+
+	/** Active interaction arc ball current point */
+	UPROPERTY()
+	FVector InteractionArcBallCurrPoint;
+
+	/** Arc ball start rotation */
+	UPROPERTY()
+	FQuat StartRotation = FQuat::Identity;
+
+	/** Arc ball current rotation */
+	FQuat CurrentRotation = FQuat::Identity;
+
+	/** Indirect manipulation */
+	UPROPERTY()
+	bool bIndirectManipulation;
+
+	/** Defer drag function on tick to avoid firing too many drag moves */
+	UPROPERTY()
+	bool bDeferDrag = true;
+
+	/** Pending drag function to be called if bDeferDrag is true */
+	TFunction<void()> PendingDragFunction;
 };
diff --git a/Engine/Source/Editor/Experimental/EditorInteractiveToolsFramework/Public/EditorInteractiveGizmoManager.h b/Engine/Source/Editor/Experimental/EditorInteractiveToolsFramework/Public/EditorInteractiveGizmoManager.h
--- a/Engine/Source/Editor/Experimental/EditorInteractiveToolsFramework/Public/EditorInteractiveGizmoManager.h
+++ b/Engine/Source/Editor/Experimental/EditorInteractiveToolsFramework/Public/EditorInteractiveGizmoManager.h
@@ -25,6 +25,7 @@ class UInteractiveGizmoBuilder;
 class UObject;
 class UTypedElementSelectionSet;
 struct FToolBuilderState;
+class UTransformGizmo;
 
 USTRUCT()
 struct FActiveEditorGizmo
@@ -109,24 +110,6 @@ public:
 		return bSearchLocalBuildersOnly;
 	}
 
-	/**
-	 * Try to automatically activate a new Gizmo instance based on the current selection set state
-	 * @return array of new Gizmo instances that have been created and initialized
-	 */
-	virtual TArray<UInteractiveGizmo*> CreateGizmosForCurrentSelectionSet();
-
-	/**
-	 * Try to automatically activate a new Gizmo instance based on the current state
-	 * @return array of new Gizmo instances that have been created and initialized
-	 */
-	virtual UInteractiveGizmoBuilder* GetTransformGizmoBuilder();
-
-	/**
-	 * Handle Editor selection changes
-	 * @param InSelectionSet - typed element selection set which invoked this selection changed call
-	 */
-	void HandleEditorSelectionSetChanged(const UTypedElementSelectionSet* InSelectionSet);
-
 	/**
 	 * Shutdown and remove a selection-based Gizmo
 	 * @param Gizmo the Gizmo to shutdown and remove
@@ -139,6 +122,14 @@ public:
 	 */
 	virtual void DestroyAllEditorGizmos();
 
+	/** Try to activate a new Gizmo instance (UInteractiveGizmoManager override) */	
+	virtual UInteractiveGizmo* CreateGizmo(
+		const FString& BuilderIdentifier, const FString& InstanceIdentifier = FString(), void* Owner = nullptr) override;
+
+	/** instance/builder identifiers for transform gizmo */
+	static const FString& TransformInstanceIdentifier();
+	static const FString& TransformBuilderIdentifier();
+	
 protected:
 
 	/**
@@ -158,7 +149,6 @@ protected:
 	 */
 	void UpdateActiveEditorGizmos();
 
-
 	/** Actual registry */
 	UPROPERTY()
 	TObjectPtr<UEditorInteractiveGizmoRegistry> Registry;
@@ -179,4 +169,9 @@ protected:
 
 	/** Whether Editor gizmos are enabled. UpdateActiveEditorGizmos() determines this value each tick and updates if it has changed. */
 	bool bShowEditorGizmos = false;
+
+private:
+
+	/** Returns the existing default Gizmo instance if any. */
+	UTransformGizmo* FindDefaultTransformGizmo() const;
 };
diff --git a/Engine/Source/Editor/UnrealEd/Private/EditorModeManager.cpp b/Engine/Source/Editor/UnrealEd/Private/EditorModeManager.cpp
--- a/Engine/Source/Editor/UnrealEd/Private/EditorModeManager.cpp
+++ b/Engine/Source/Editor/UnrealEd/Private/EditorModeManager.cpp
@@ -51,6 +51,7 @@
 
 #include "Elements/Interfaces/TypedElementWorldInterface.h"
 #include "TextureResource.h"
+#include "EditorGizmos/EditorTransformGizmoUtil.h"
 #include "Toolkits/ToolkitManager.h"
 
 /*------------------------------------------------------------------------------
@@ -102,6 +103,8 @@ FEditorModeTools::FEditorModeTools()
 		GEditor->GetEditorSubsystem<UAssetEditorSubsystem>()->OnEditorModeUnregistered().AddRaw(this, &FEditorModeTools::OnModeUnregistered);
 	}
 
+	UE::EditorTransformGizmoUtil::RegisterTransformGizmoContextObject(this);
+
 	FWorldDelegates::OnWorldCleanup.AddRaw(this, &FEditorModeTools::OnWorldCleanup);
 }
 
@@ -119,6 +122,7 @@ FEditorModeTools::~FEditorModeTools()
 	// which would mean that this instances will be garbage
 	if (UObjectInitialized())
 	{
+		UE::EditorTransformGizmoUtil::UnregisterTransformGizmoContextObject(this);
 		InteractiveToolsContext->Deactivate();
 		InteractiveToolsContext->ShutdownContext();
 		InteractiveToolsContext = nullptr;
diff --git a/Engine/Source/Editor/UnrealEd/Private/LevelEditorViewport.cpp b/Engine/Source/Editor/UnrealEd/Private/LevelEditorViewport.cpp
--- a/Engine/Source/Editor/UnrealEd/Private/LevelEditorViewport.cpp
+++ b/Engine/Source/Editor/UnrealEd/Private/LevelEditorViewport.cpp
@@ -2866,7 +2866,11 @@ bool FLevelEditorViewportClient::InputWidgetDelta(FViewport* InViewport, EAxisLi
 				}
 
 				// We do not want actors updated if we are holding down the middle mouse button.
-				if(!MiddleMouseButtonDown)
+				// enable MMB for New TRS Gizmos
+				static IConsoleVariable* const UseLegacyWidgetCVar = IConsoleManager::Get().FindConsoleVariable(TEXT("Gizmos.UseLegacyWidget"));
+				const bool bEnableMMB = UseLegacyWidgetCVar ? !UseLegacyWidgetCVar->GetBool() && !bDraggingByHandle : false;
+				
+				if(!MiddleMouseButtonDown || bEnableMMB)
 				{
 					bool bSnapped = FSnappingUtils::SnapActorsToNearestActor( Drag, this );
 					bSnapped = bSnapped || FSnappingUtils::SnapDraggedActorsToNearestVertex( Drag, this );
diff --git a/Engine/Source/Runtime/InteractiveToolsFramework/Private/InteractiveGizmoManager.cpp b/Engine/Source/Runtime/InteractiveToolsFramework/Private/InteractiveGizmoManager.cpp
--- a/Engine/Source/Runtime/InteractiveToolsFramework/Private/InteractiveGizmoManager.cpp
+++ b/Engine/Source/Runtime/InteractiveToolsFramework/Private/InteractiveGizmoManager.cpp
@@ -202,7 +202,7 @@ void UInteractiveGizmoManager::DestroyAllGizmosByOwner(void* Owner)
 
 
 
-UInteractiveGizmo* UInteractiveGizmoManager::FindGizmoByInstanceIdentifier(const FString& Identifier)
+UInteractiveGizmo* UInteractiveGizmoManager::FindGizmoByInstanceIdentifier(const FString& Identifier) const
 {
 	for (int i = 0; i < ActiveGizmos.Num(); ++i)
 	{
diff --git a/Engine/Source/Runtime/InteractiveToolsFramework/Public/InteractiveGizmoManager.h b/Engine/Source/Runtime/InteractiveToolsFramework/Public/InteractiveGizmoManager.h
--- a/Engine/Source/Runtime/InteractiveToolsFramework/Public/InteractiveGizmoManager.h
+++ b/Engine/Source/Runtime/InteractiveToolsFramework/Public/InteractiveGizmoManager.h
@@ -130,7 +130,7 @@ public:
 	 * @param Identifier the InstanceIdentifier that was passed to CreateGizmo()
 	 * @return the found Gizmo, or null
 	 */
-	INTERACTIVETOOLSFRAMEWORK_API virtual UInteractiveGizmo* FindGizmoByInstanceIdentifier(const FString& Identifier);
+	INTERACTIVETOOLSFRAMEWORK_API virtual UInteractiveGizmo* FindGizmoByInstanceIdentifier(const FString& Identifier) const;
 
 
 
