New TRS Gizmo updates: follow-up to CL 27751703

#rnx

diff --git a/Engine/Source/Editor/Experimental/EditorInteractiveToolsFramework/Private/EdModeInteractiveToolsContext.cpp b/Engine/Source/Editor/Experimental/EditorInteractiveToolsFramework/Private/EdModeInteractiveToolsContext.cpp
--- a/Engine/Source/Editor/Experimental/EditorInteractiveToolsFramework/Private/EdModeInteractiveToolsContext.cpp
+++ b/Engine/Source/Editor/Experimental/EditorInteractiveToolsFramework/Private/EdModeInteractiveToolsContext.cpp
@@ -321,11 +321,6 @@ void UEditorInteractiveToolsContext::Shutdown()
 {
 	bIsActive = false;
 
-	if (UContextObjectStore* ContextStore = GizmoManager ? GizmoManager->GetContextObjectStore() : nullptr)
-	{
-		ContextStore->RemoveContextObject(this);
-	}
-	
 	// auto-accept any in-progress tools
 	DeactivateAllActiveTools(EToolShutdownType::Accept);
 
@@ -344,10 +339,6 @@ void UEditorInteractiveToolsContext::InitializeContextWithEditorModeManager(FEdi
 	SetCreateGizmoManagerFunc([this](const FContextInitInfo& ContextInfo)
 	{
 		UEditorInteractiveGizmoManager* NewGizmoManager = NewObject<UEditorInteractiveGizmoManager>(ContextInfo.ToolsContext);
-		if (UContextObjectStore* ContextStore = NewGizmoManager->GetContextObjectStore())
-		{
-			ContextStore->AddContextObject(this);
-		}
 		NewGizmoManager->InitializeWithEditorModeManager(ContextInfo.QueriesAPI, ContextInfo.TransactionsAPI, ContextInfo.InputRouter, EditorModeManager);
 		NewGizmoManager->RegisterDefaultGizmos();
 		return NewGizmoManager;
