Fixed static analysis
    - cf. https://horde.devtools.epicgames.com/job/64ff2bb0af5c3ec547a643b2

#rnx

diff --git a/Engine/Source/Editor/Experimental/EditorInteractiveToolsFramework/Private/EditorGizmos/TransformGizmo.cpp b/Engine/Source/Editor/Experimental/EditorInteractiveToolsFramework/Private/EditorGizmos/TransformGizmo.cpp
--- a/Engine/Source/Editor/Experimental/EditorInteractiveToolsFramework/Private/EditorGizmos/TransformGizmo.cpp
+++ b/Engine/Source/Editor/Experimental/EditorInteractiveToolsFramework/Private/EditorGizmos/TransformGizmo.cpp
@@ -295,11 +295,10 @@ bool UTransformGizmo::VerifyPartIdentifier(uint32 InPartIdentifier) const
 
 void UTransformGizmo::SetModeLastHitPart(const EGizmoTransformMode InMode, const ETransformGizmoPartIdentifier InIdentifier)
 {
-	if (InMode < EGizmoTransformMode::None || InMode > EGizmoTransformMode::Max)
+	if (InMode >= EGizmoTransformMode::None && InMode < EGizmoTransformMode::Max)
 	{
-		return;
+		LastHitPartPerMode[static_cast<uint8>(InMode)] = InIdentifier;
 	}
-	LastHitPartPerMode[static_cast<uint8>(InMode)] = InIdentifier;
 }
 
 ETransformGizmoPartIdentifier UTransformGizmo::GetCurrentModeLastHitPart() const
