New TRS Gizmo
    - support for current FEditorModeTools::WidgetScale
    - reset default indirect target when re-entering the same transform mode
    - added FEditorTransformGizmoDataBinder to easily interface between transform gizmos and the mode manager

TODO
    - FEditorModeTools::WidgetScale (currently acting as a uniform matrix scale as designed in 14299282) should be replaced by gizmo scale, line with and handle size settings once the design done.

#rnx

diff --git a/Engine/Source/Editor/Experimental/EditorInteractiveToolsFramework/Private/EditorGizmos/EditorTransformGizmoBuilder.cpp b/Engine/Source/Editor/Experimental/EditorInteractiveToolsFramework/Private/EditorGizmos/EditorTransformGizmoBuilder.cpp
--- a/Engine/Source/Editor/Experimental/EditorInteractiveToolsFramework/Private/EditorGizmos/EditorTransformGizmoBuilder.cpp
+++ b/Engine/Source/Editor/Experimental/EditorInteractiveToolsFramework/Private/EditorGizmos/EditorTransformGizmoBuilder.cpp
@@ -6,7 +6,6 @@
 #include "BaseGizmos/GizmoElementHitTargets.h"
 #include "BaseGizmos/GizmoViewContext.h"
 #include "ContextObjectStore.h"
-#include "EditorModeManager.h"
 #include "EdModeInteractiveToolsContext.h"
 #include "EditorGizmos/EditorTransformGizmo.h"
 #include "EditorGizmos/EditorTransformGizmoSource.h"
diff --git a/Engine/Source/Editor/Experimental/EditorInteractiveToolsFramework/Private/EditorGizmos/EditorTransformGizmoDataBinder.cpp b/Engine/Source/Editor/Experimental/EditorInteractiveToolsFramework/Private/EditorGizmos/EditorTransformGizmoDataBinder.cpp
new file mode 100644
--- /dev/null
+++ b/Engine/Source/Editor/Experimental/EditorInteractiveToolsFramework/Private/EditorGizmos/EditorTransformGizmoDataBinder.cpp
@@ -0,0 +1,92 @@
+﻿// Copyright Epic Games, Inc. All Rights Reserved.
+
+
+#include "EditorGizmos/EditorTransformGizmoDataBinder.h"
+
+#include "EditorModeManager.h"
+#include "EditorGizmos/EditorTransformGizmoUtil.h"
+#include "EditorGizmos/TransformGizmo.h"
+
+FEditorTransformGizmoDataBinder::~FEditorTransformGizmoDataBinder()
+{
+	const TSet<TWeakObjectPtr<UTransformGizmo>> Gizmos(BoundGizmos);
+	for (const TWeakObjectPtr<UTransformGizmo> Gizmo : Gizmos)
+	{
+		if (Gizmo.IsValid())
+		{
+			UnbindFromGizmo(Gizmo.Get(), Gizmo->ActiveTarget);
+			Gizmo->OnSetActiveTarget.RemoveAll(this);
+			Gizmo->OnAboutToClearActiveTarget.RemoveAll(this);
+		}
+	}
+	BoundGizmos.Empty();
+
+	if (WeakContext.IsValid())
+	{
+		WeakContext->OnGizmoCreatedDelegate().RemoveAll(this);
+	}
+	WeakContext.Reset();
+}
+
+void FEditorTransformGizmoDataBinder::BindToGizmoContextObject(UEditorTransformGizmoContextObject* InContextObject)
+{
+	if (!ensure(IsValid(InContextObject)))
+	{
+		return;
+	}
+
+	InContextObject->OnGizmoCreatedDelegate().AddSP(this, &FEditorTransformGizmoDataBinder::BindToUninitializedGizmo);
+	
+	WeakContext = InContextObject;
+}
+
+void FEditorTransformGizmoDataBinder::BindToUninitializedGizmo(UTransformGizmo* InGizmo)
+{
+	FEditorModeTools* ModeTools = WeakContext.IsValid() ? WeakContext->GetModeTools() : nullptr;
+	if (ensure(ModeTools))
+	{
+		if (!ModeTools->OnWidgetModeChanged().IsBoundToObject(InGizmo))
+		{
+			ModeTools->OnWidgetModeChanged().AddUObject(InGizmo, &UTransformGizmo::HandleWidgetModeChanged);
+		}
+	}
+
+	const bool bSetBound = InGizmo->OnSetActiveTarget.IsBoundToObject(this);
+	const bool bClearBound = InGizmo->OnAboutToClearActiveTarget.IsBoundToObject(this);
+	if ( ensure(!bSetBound && !bClearBound))
+	{
+		InGizmo->OnSetActiveTarget.AddSP(this, &FEditorTransformGizmoDataBinder::BindToInitializedGizmo);
+		InGizmo->OnAboutToClearActiveTarget.AddSP(this, &FEditorTransformGizmoDataBinder::UnbindFromGizmo);
+	}
+}
+
+void FEditorTransformGizmoDataBinder::BindToInitializedGizmo(UTransformGizmo* InGizmo, UTransformProxy* InProxy)
+{
+	const bool bHasTarget = InGizmo && InGizmo->ActiveTarget && InGizmo->ActiveTarget == InProxy; 
+	if (ensure(!bHasTarget && !BoundGizmos.Contains(InGizmo)))
+	{
+		return;
+	}
+
+	// bind InProxy OnBeginTransformEdit/OnTransformChanged/OnEndTransformEdit etc here if needed
+
+	BoundGizmos.Add(InGizmo);
+}
+
+void FEditorTransformGizmoDataBinder::UnbindFromGizmo(UTransformGizmo* InGizmo, UTransformProxy* InProxy)
+{
+	if (!InGizmo)
+	{
+		return;
+	}
+
+	FEditorModeTools* ModeTools = WeakContext.IsValid() ? WeakContext->GetModeTools() : nullptr;
+	if (ensure(ModeTools))
+	{
+		ModeTools->OnWidgetModeChanged().RemoveAll(InGizmo);
+	}
+
+	// unbind InProxy OnBeginTransformEdit/OnTransformChanged/OnEndTransformEdit etc here if needed
+
+	BoundGizmos.Remove(InGizmo);
+}
diff --git a/Engine/Source/Editor/Experimental/EditorInteractiveToolsFramework/Private/EditorGizmos/EditorTransformGizmoUtil.cpp b/Engine/Source/Editor/Experimental/EditorInteractiveToolsFramework/Private/EditorGizmos/EditorTransformGizmoUtil.cpp
--- a/Engine/Source/Editor/Experimental/EditorInteractiveToolsFramework/Private/EditorGizmos/EditorTransformGizmoUtil.cpp
+++ b/Engine/Source/Editor/Experimental/EditorInteractiveToolsFramework/Private/EditorGizmos/EditorTransformGizmoUtil.cpp
@@ -10,6 +10,7 @@
 #include "EdModeInteractiveToolsContext.h"
 #include "EditorGizmos/TransformGizmo.h"
 #include "EditorGizmos/EditorTransformGizmoBuilder.h"
+#include "EditorGizmos/EditorTransformGizmoDataBinder.h"
 #include "EditorViewportClient.h"
 
 #include UE_INLINE_GENERATED_CPP_BY_NAME(EditorTransformGizmoUtil)
@@ -72,12 +73,7 @@ UTransformGizmo* UE::EditorTransformGizmoUtil::CreateTransformGizmo(
 		return nullptr;
 	}
 
-	UInteractiveGizmo* NewGizmo = GizmoManager->CreateGizmo(
-		UEditorInteractiveGizmoManager::TransformBuilderIdentifier(), InInstanceIdentifier, InOwner);
-
-	ensure(NewGizmo);
-	
-	return Cast<UTransformGizmo>(NewGizmo);
+	return Found->CreateTransformGizmo(GizmoManager, InInstanceIdentifier, InOwner);
 }
 
 UTransformGizmo* UE::EditorTransformGizmoUtil::GetDefaultTransformGizmo(UInteractiveToolManager* InToolManager)
@@ -166,11 +162,19 @@ void UEditorTransformGizmoContextObject::Initialize(FEditorModeTools* InModeTool
 	{
 		ModeTools = InModeTools;
 		InitializeCVarBinding();
+
+		DataBinder = MakeShared<FEditorTransformGizmoDataBinder>();
+		DataBinder->BindToGizmoContextObject(this);
 	}
 }
 
 void UEditorTransformGizmoContextObject::Shutdown()
 {
+	if (ensure(DataBinder))
+	{
+		DataBinder.Reset();
+	}
+	
 	if (ensure(ModeTools))
 	{
 		UpdateGizmo({});
@@ -181,6 +185,35 @@ void UEditorTransformGizmoContextObject::Shutdown()
 	}
 }
 
+UTransformGizmo* UEditorTransformGizmoContextObject::CreateTransformGizmo(
+	UEditorInteractiveGizmoManager* InGizmoManager, const FString& InInstanceIdentifier, void* InOwner) const
+{
+	if (!InGizmoManager)
+	{
+		return nullptr;
+	}
+	
+	UTransformGizmo* NewGizmo = Cast<UTransformGizmo>( InGizmoManager->CreateGizmo(
+		UEditorInteractiveGizmoManager::TransformBuilderIdentifier(), InInstanceIdentifier, InOwner) );
+
+	if (ensure(NewGizmo))
+	{
+		OnGizmoCreated.Broadcast(NewGizmo);
+	}
+	
+	return NewGizmo;
+}
+
+FEditorModeTools* UEditorTransformGizmoContextObject::GetModeTools() const
+{
+	return ModeTools;
+}
+
+UEditorTransformGizmoContextObject::FOnGizmoCreated& UEditorTransformGizmoContextObject::OnGizmoCreatedDelegate()
+{
+	return OnGizmoCreated;
+}
+
 void UEditorTransformGizmoContextObject::UpdateGizmo(const TArray<FEditorViewportClient*>& InViewportClients) const
 {
 	if (!ModeTools)
@@ -320,8 +353,3 @@ void UEditorTransformGizmoContextObject::RemoveViewportsBinding()
 		ViewportClientsChangedHandle.Reset();
 	}
 }
-
-const FEditorModeTools* UEditorTransformGizmoContextObject::GetModeTools() const
-{
-	return ModeTools;
-}
\ No newline at end of file
diff --git a/Engine/Source/Editor/Experimental/EditorInteractiveToolsFramework/Private/EditorGizmos/EditorTransformProxy.cpp b/Engine/Source/Editor/Experimental/EditorInteractiveToolsFramework/Private/EditorGizmos/EditorTransformProxy.cpp
--- a/Engine/Source/Editor/Experimental/EditorInteractiveToolsFramework/Private/EditorGizmos/EditorTransformProxy.cpp
+++ b/Engine/Source/Editor/Experimental/EditorInteractiveToolsFramework/Private/EditorGizmos/EditorTransformProxy.cpp
@@ -19,7 +19,8 @@ FTransform UEditorTransformProxy::GetTransform() const
 	{
 		const FVector Location = ViewportClient->GetWidgetLocation();
 		const FQuat Rotation(ViewportClient->GetWidgetCoordSystem());
-		return FTransform(Rotation, Location, FVector::OneVector);
+		const FVector Scale(WeakContext->GetModeTools()->GetWidgetScale());
+		return FTransform(Rotation, Location, Scale);
 	}
 	
 	return FTransform::Identity;
diff --git a/Engine/Source/Editor/Experimental/EditorInteractiveToolsFramework/Private/EditorGizmos/TransformGizmo.cpp b/Engine/Source/Editor/Experimental/EditorInteractiveToolsFramework/Private/EditorGizmos/TransformGizmo.cpp
--- a/Engine/Source/Editor/Experimental/EditorInteractiveToolsFramework/Private/EditorGizmos/TransformGizmo.cpp
+++ b/Engine/Source/Editor/Experimental/EditorInteractiveToolsFramework/Private/EditorGizmos/TransformGizmo.cpp
@@ -52,7 +52,6 @@ void UTransformGizmo::Setup()
 	SetModeLastHitPart(EGizmoTransformMode::Translate, ETransformGizmoPartIdentifier::TranslateScreenSpace);
 	SetModeLastHitPart(EGizmoTransformMode::Rotate, ETransformGizmoPartIdentifier::RotateArcball);
 	SetModeLastHitPart(EGizmoTransformMode::Scale, ETransformGizmoPartIdentifier::ScaleUniform);
-	SetModeLastHitPart(EGizmoTransformMode::Max, ETransformGizmoPartIdentifier::Default);
 }
 
 void UTransformGizmo::SetupBehaviors()
@@ -163,6 +162,8 @@ void UTransformGizmo::SetupMaterials()
 void UTransformGizmo::Shutdown()
 {
 	ClearActiveTarget();
+	OnSetActiveTarget.Clear();
+	OnAboutToClearActiveTarget.Clear();
 }
 
 FTransform UTransformGizmo::GetGizmoTransform() const
@@ -747,8 +748,9 @@ void UTransformGizmo::SetActiveTarget(UTransformProxy* Target, IToolContextTrans
 			LOCTEXT("UTransformGizmoTransaction", "Transform"), TransactionProvider, this);	
 	}
 
-
 	CameraAxisSource = NewObject<UGizmoConstantFrameAxisSource>(this);
+
+	OnSetActiveTarget.Broadcast(this, ActiveTarget);
 }
 
 // @todo: This should either be named to "SetScale" or removed, since it can be done with ReinitializeGizmoTransform
@@ -771,6 +773,46 @@ void UTransformGizmo::SetCustomizationFunction(const TFunction<const FGizmoCusto
 	CustomizationFunction = InFunction;
 }
 
+void UTransformGizmo::HandleWidgetModeChanged(UE::Widget::EWidgetMode InWidgetMode)
+{
+	auto GetTransformMode = [InWidgetMode]()
+	{
+		switch (InWidgetMode)
+		{
+		case UE::Widget::EWidgetMode::WM_Translate: return EGizmoTransformMode::Translate;
+		case UE::Widget::EWidgetMode::WM_Rotate: return EGizmoTransformMode::Rotate;
+		case UE::Widget::EWidgetMode::WM_Scale: return EGizmoTransformMode::Scale;
+		default: return EGizmoTransformMode::None;
+		}
+		return EGizmoTransformMode::None;
+	};
+	const EGizmoTransformMode NewMode = GetTransformMode();
+
+	if (CurrentMode != EGizmoTransformMode::None && NewMode == CurrentMode)
+	{
+		auto GetModeDefaultHitPart = [NewMode]()
+		{
+			switch (NewMode)
+			{
+			case EGizmoTransformMode::Translate: return ETransformGizmoPartIdentifier::TranslateScreenSpace;
+			case EGizmoTransformMode::Rotate: return ETransformGizmoPartIdentifier::RotateArcball;
+			case EGizmoTransformMode::Scale: return ETransformGizmoPartIdentifier::ScaleUniform;
+			default: return ETransformGizmoPartIdentifier::Default;
+			}
+			return ETransformGizmoPartIdentifier::Default;
+		};
+
+		const ETransformGizmoPartIdentifier DefaultHitPart = GetModeDefaultHitPart();
+		if (DefaultHitPart != GetCurrentModeLastHitPart())
+		{
+			// reset indirect manipulation to default
+			UpdateInteractingState(false, GetCurrentModeLastHitPart(), true);
+			SetModeLastHitPart(CurrentMode, DefaultHitPart);
+			UpdateInteractingState(true, DefaultHitPart, true);
+		}
+	}
+}
+
 UGizmoElementArrow* UTransformGizmo::MakeTranslateAxis(ETransformGizmoPartIdentifier InPartId, const FVector& InAxisDir, const FVector& InSideDir, UMaterialInterface* InMaterial)
 {
 	const float SizeCoeff = GetSizeCoefficient();
@@ -949,6 +991,8 @@ void UTransformGizmo::ClearActiveTarget()
 
 	if (ActiveTarget)
 	{
+		OnAboutToClearActiveTarget.Broadcast(this, ActiveTarget);
+		
 		ActiveTarget->OnBeginTransformEdit.RemoveAll(this);
 		ActiveTarget->OnEndTransformEdit.RemoveAll(this);
 		ActiveTarget = nullptr;
@@ -1912,7 +1956,8 @@ void UTransformGizmo::OnClickReleaseArcBallRotate(const FInputDeviceRay& Release
 float UTransformGizmo::GetArcBallWorldRadius() const
 {
 	const float PixelToWorldScale = GizmoRenderingUtil::CalculateLocalPixelToWorldScale(GizmoViewContext, CurrentTransform.GetLocation());
-	return RotateArcballSphereRadius * GetSizeCoefficient() * PixelToWorldScale;
+	const float GizmoScale = TransformGizmoSource ? TransformGizmoSource->GetGizmoScale() : 1.0f;
+	return RotateArcballSphereRadius * GetSizeCoefficient() * PixelToWorldScale * GizmoScale;
 }
 
 float UTransformGizmo::GetSizeCoefficient() const
diff --git a/Engine/Source/Editor/Experimental/EditorInteractiveToolsFramework/Public/EditorGizmos/EditorTransformGizmoDataBinder.h b/Engine/Source/Editor/Experimental/EditorInteractiveToolsFramework/Public/EditorGizmos/EditorTransformGizmoDataBinder.h
new file mode 100644
--- /dev/null
+++ b/Engine/Source/Editor/Experimental/EditorInteractiveToolsFramework/Public/EditorGizmos/EditorTransformGizmoDataBinder.h
@@ -0,0 +1,49 @@
+﻿// Copyright Epic Games, Inc. All Rights Reserved.
+
+#pragma once
+
+#include "Templates/SharedPointer.h"
+#include "UObject/WeakObjectPtr.h"
+#include "UObject/WeakObjectPtrTemplates.h"
+#include "Containers/Set.h"
+
+class UEditorTransformGizmoContextObject;
+class UTransformGizmo;
+class UTransformProxy;
+
+/**
+ * FEditorTransformGizmoDataBinder is a helper class for binding a UTransformGizmo to a FEditorModeTools
+ */
+
+class EDITORINTERACTIVETOOLSFRAMEWORK_API FEditorTransformGizmoDataBinder : public TSharedFromThis<FEditorTransformGizmoDataBinder>
+{
+public:
+
+	virtual ~FEditorTransformGizmoDataBinder();
+	
+	/**
+	 * Makes it so that the gizmo binder attaches to any gizmos created by the context object
+	 * in the future. The binding is automatically removed if FEditorTransformGizmoDataBinder is
+	 * destroyed.
+	 */
+	void BindToGizmoContextObject(UEditorTransformGizmoContextObject* InContextObject);
+
+	/** Used for binding to context object, called from the OnGizmoCreated delegate. */
+	void BindToUninitializedGizmo(UTransformGizmo* Gizmo);
+	
+	/** Binds to a specific gizmo for tracking. Requires ActiveTarget to be set. */
+	void BindToInitializedGizmo(UTransformGizmo* InGizmo, UTransformProxy* InProxy);
+	
+	/** Unbinds from a given gizmo. Done automatically for gizmos when their ActiveTarget is cleared. */
+	void UnbindFromGizmo(UTransformGizmo* InGizmo, UTransformProxy* InProxy);
+
+private:
+
+	/** List of the gizmos currently bound to this so that we ensure that they are actually unbound on destruction.
+	 * Note that this should not happen and bound gizmos should have called their Shutdown function before this being destroyed.
+	 */
+	TSet<TWeakObjectPtr<UTransformGizmo>> BoundGizmos;
+
+	/** Weak ptr to the context so we can interface with the current mode manager */
+	TWeakObjectPtr<UEditorTransformGizmoContextObject> WeakContext;
+};
\ No newline at end of file
diff --git a/Engine/Source/Editor/Experimental/EditorInteractiveToolsFramework/Public/EditorGizmos/EditorTransformGizmoUtil.h b/Engine/Source/Editor/Experimental/EditorInteractiveToolsFramework/Public/EditorGizmos/EditorTransformGizmoUtil.h
--- a/Engine/Source/Editor/Experimental/EditorInteractiveToolsFramework/Public/EditorGizmos/EditorTransformGizmoUtil.h
+++ b/Engine/Source/Editor/Experimental/EditorInteractiveToolsFramework/Public/EditorGizmos/EditorTransformGizmoUtil.h
@@ -12,6 +12,10 @@ class UInteractiveToolManager;
 class UTransformGizmo;
 class FEditorModeTools;
 class FEditorViewportClient;
+class UEditorInteractiveGizmoManager;
+class UEditorTransformGizmoContextObject;
+class UTransformProxy;
+class FEditorTransformGizmoDataBinder;
 
 namespace UE
 {
@@ -78,7 +82,14 @@ public:
 	EDITORINTERACTIVETOOLSFRAMEWORK_API void Initialize(FEditorModeTools* InModeTools);
 	EDITORINTERACTIVETOOLSFRAMEWORK_API void Shutdown();
 
-	EDITORINTERACTIVETOOLSFRAMEWORK_API const FEditorModeTools* GetModeTools() const;
+	EDITORINTERACTIVETOOLSFRAMEWORK_API UTransformGizmo* CreateTransformGizmo(
+		UEditorInteractiveGizmoManager* InGizmoManager, const FString& InInstanceIdentifier, void* InOwner) const;
+	
+	EDITORINTERACTIVETOOLSFRAMEWORK_API FEditorModeTools* GetModeTools() const;
+
+	DECLARE_MULTICAST_DELEGATE_OneParam(FOnGizmoCreated, UTransformGizmo*);
+	FOnGizmoCreated& OnGizmoCreatedDelegate();
+
 	
 private:
 
@@ -90,8 +101,12 @@ private:
 	EDITORINTERACTIVETOOLSFRAMEWORK_API void RemoveCVarBinding();
 	EDITORINTERACTIVETOOLSFRAMEWORK_API void RemoveViewportsBinding();
 
+	FOnGizmoCreated OnGizmoCreated;
+	
 	FEditorModeTools* ModeTools = nullptr;
 	FDelegateHandle ViewportClientsChangedHandle;
 	FDelegateHandle UseLegacyChangedHandled;
 	bool bGizmosRegistered = false;
+
+	TSharedPtr<FEditorTransformGizmoDataBinder> DataBinder;
 };
diff --git a/Engine/Source/Editor/Experimental/EditorInteractiveToolsFramework/Public/EditorGizmos/TransformGizmo.h b/Engine/Source/Editor/Experimental/EditorInteractiveToolsFramework/Public/EditorGizmos/TransformGizmo.h
--- a/Engine/Source/Editor/Experimental/EditorInteractiveToolsFramework/Public/EditorGizmos/TransformGizmo.h
+++ b/Engine/Source/Editor/Experimental/EditorInteractiveToolsFramework/Public/EditorGizmos/TransformGizmo.h
@@ -207,6 +207,11 @@ public:
 	 * Set customization function for this Gizmo
 	 */
 	void SetCustomizationFunction(const TFunction<const FGizmoCustomization()>& InFunction);
+
+	/**
+	 * Handle widget mode changed.
+	 */
+	void HandleWidgetModeChanged(UE::Widget::EWidgetMode InWidgetMode);
 	
 public:
 
@@ -268,6 +273,14 @@ public:
 	UPROPERTY()
 	bool bSnapToWorldRotGrid = false;
 
+	/** Broadcast at the end of a SetActiveTarget call. */
+	DECLARE_MULTICAST_DELEGATE_TwoParams(FOnSetActiveTarget, UTransformGizmo*, UTransformProxy*);
+	FOnSetActiveTarget OnSetActiveTarget;
+	
+	/** Broadcast at the beginning of a ClearActiveTarget call, when the ActiveTarget is not yet disconnected. */
+	DECLARE_MULTICAST_DELEGATE_TwoParams(FOnClearActiveTarget, UTransformGizmo*, UTransformProxy*);
+	FOnClearActiveTarget OnAboutToClearActiveTarget;
+
 protected:
 
 	//
