InteractiveToolsFramework: Fix mouse state not being correctly updated in the editor for buttons other than the actively pressed one.

diff --git a/Engine/Source/Editor/Experimental/EditorInteractiveToolsFramework/Private/EdModeInteractiveToolsContext.cpp b/Engine/Source/Editor/Experimental/EditorInteractiveToolsFramework/Private/EdModeInteractiveToolsContext.cpp
--- a/Engine/Source/Editor/Experimental/EditorInteractiveToolsFramework/Private/EdModeInteractiveToolsContext.cpp
+++ b/Engine/Source/Editor/Experimental/EditorInteractiveToolsFramework/Private/EdModeInteractiveToolsContext.cpp
@@ -890,6 +890,31 @@ void UModeManagerInteractiveToolsContext::DeactivateAllActiveTools(EToolShutdown
 	Super::DeactivateAllActiveTools(ShutdownType);
 }
 
+void UModeManagerInteractiveToolsContext::UpdateStateWithoutRoutingInputKey(FEditorViewportClient* ViewportClient, FViewport* Viewport, FKey Key, EInputEvent Event)
+{
+	// Currently, the only internal state we keep is the state of various mouse keys being down. Note
+	// that we don't want to save bPressed or bReleased for them as those are one-time events issued
+	// from InputKey, and shouldn't show up on mouse moves or on other keys being pressed/released.
+	if ((Event == IE_Pressed || Event == IE_Released)
+		&& Key.IsMouseButton())
+	{
+		if (Key.IsMouseButton())
+		{
+			if (Key == EKeys::LeftMouseButton)
+			{
+				CurrentMouseState.Mouse.Left.bDown = (Event == IE_Pressed);
+			}
+			else if (Key == EKeys::MiddleMouseButton)
+			{
+				CurrentMouseState.Mouse.Middle.bDown = (Event == IE_Pressed);
+			}
+			else if (Key == EKeys::RightMouseButton)
+			{
+				CurrentMouseState.Mouse.Right.bDown = (Event == IE_Pressed);
+			}
+		}
+	}
+}
 
 bool UModeManagerInteractiveToolsContext::InputKey(FEditorViewportClient* ViewportClient, FViewport* Viewport, FKey Key, EInputEvent Event)
 {
@@ -901,6 +926,9 @@ bool UModeManagerInteractiveToolsContext::InputKey(FEditorViewportClient* Viewpo
 	else if (Event == IE_DoubleClick) { UE_LOG(LogTemp, Warning, TEXT("DOUBLECLICK EVENT")); }
 #endif
 
+	// Update the current state, then route result
+	UpdateStateWithoutRoutingInputKey(ViewportClient, Viewport, Key, Event);
+
 	if (Event == IE_Pressed || Event == IE_Released)
 	{
 		if (Key.IsMouseButton())
@@ -932,22 +960,16 @@ bool UModeManagerInteractiveToolsContext::InputKey(FEditorViewportClient* Viewpo
 				{
 					InputState.Mouse.Left.SetStates(
 						(Event == IE_Pressed), (Event == IE_Pressed), (Event == IE_Released));
-					CurrentMouseState.Mouse.Left.bDown = (Event == IE_Pressed);
-					CurrentMouseState.Mouse.Left.bReleased = (Event == IE_Released);
 				}
 				else if (bIsMiddleMouse)
 				{
 					InputState.Mouse.Middle.SetStates(
 						(Event == IE_Pressed), (Event == IE_Pressed), (Event == IE_Released));
-					CurrentMouseState.Mouse.Middle.bDown = (Event == IE_Pressed);
-					CurrentMouseState.Mouse.Middle.bReleased = (Event == IE_Released);
 				}
 				else
 				{
 					InputState.Mouse.Right.SetStates(
 						(Event == IE_Pressed), (Event == IE_Pressed), (Event == IE_Released));
-					CurrentMouseState.Mouse.Right.bDown = (Event == IE_Pressed);
-					CurrentMouseState.Mouse.Right.bReleased = (Event == IE_Released);
 				}
 				if (InputRouter->PostInputEvent(InputState))
 				{
diff --git a/Engine/Source/Editor/Experimental/EditorInteractiveToolsFramework/Public/EdModeInteractiveToolsContext.h b/Engine/Source/Editor/Experimental/EditorInteractiveToolsFramework/Public/EdModeInteractiveToolsContext.h
--- a/Engine/Source/Editor/Experimental/EditorInteractiveToolsFramework/Public/EdModeInteractiveToolsContext.h
+++ b/Engine/Source/Editor/Experimental/EditorInteractiveToolsFramework/Public/EdModeInteractiveToolsContext.h
@@ -256,6 +256,12 @@ public:
 	//
 public:
 	bool InputKey(FEditorViewportClient* ViewportClient, FViewport* Viewport, FKey Key, EInputEvent Event);
+	/**
+	 * This updates internal state like InputKey, but doesn't route the results to the input router. 
+	 * Use this if the input is captured by some higher system, to avoid this class from having an
+	 * incorrect view of e.g. the mouse state because it did not receive a mouse release event.
+	 */
+	void UpdateStateWithoutRoutingInputKey(FEditorViewportClient* ViewportClient, FViewport* Viewport, FKey Key, EInputEvent Event);
 
 	bool MouseEnter(FEditorViewportClient* ViewportClient, FViewport* Viewport, int32 x, int32 y);
 	bool MouseLeave(FEditorViewportClient* ViewportClient, FViewport* Viewport);
diff --git a/Engine/Source/Editor/UnrealEd/Private/EditorModeManager.cpp b/Engine/Source/Editor/UnrealEd/Private/EditorModeManager.cpp
--- a/Engine/Source/Editor/UnrealEd/Private/EditorModeManager.cpp
+++ b/Engine/Source/Editor/UnrealEd/Private/EditorModeManager.cpp
@@ -1309,7 +1309,18 @@ bool FEditorModeTools::ProcessCapturedMouseMoves( FEditorViewportClient* InViewp
 /** Notifies all active modes of keyboard input via a viewport client */
 bool FEditorModeTools::InputKey(FEditorViewportClient* InViewportClient, FViewport* Viewport, FKey Key, EInputEvent Event, bool bRouteToToolsContext)
 {
-	const bool bWasHandledByToolsContext = bRouteToToolsContext && InteractiveToolsContext->InputKey(InViewportClient, Viewport, Key, Event);
+	bool bWasHandledByToolsContext = false;
+	if (bRouteToToolsContext)
+	{
+		bWasHandledByToolsContext = InteractiveToolsContext->InputKey(InViewportClient, Viewport, Key, Event);
+	}
+	else
+	{
+		// If we're not routing to the tools context, we still need to let it look at the event so that it can update
+		// its internal memory of which mouse keys are down, to pass the correct mouse state later.
+		InteractiveToolsContext->UpdateStateWithoutRoutingInputKey(InViewportClient, Viewport, Key, Event);
+	}
+
 	if (bWasHandledByToolsContext && !bIsTracking && GetInteractiveToolsContext()->InputRouter->HasActiveMouseCapture())
 	{
 		StartTracking(InViewportClient, Viewport);
