New TRS Gizmo:
    - fixed axis manipulation for both direct & indirect manipulation
    - indirect manipulation using CTRL+LMB/RMB is now handled by the gizmo
    - added UMultiButtonClickDragBehavior to support click & drag with several buttons in one single behavior

diff --git a/Engine/Source/Editor/Experimental/EditorInteractiveToolsFramework/Private/Behaviors/MultiButtonClickDragBehavior.cpp b/Engine/Source/Editor/Experimental/EditorInteractiveToolsFramework/Private/Behaviors/MultiButtonClickDragBehavior.cpp
new file mode 100644
--- /dev/null
+++ b/Engine/Source/Editor/Experimental/EditorInteractiveToolsFramework/Private/Behaviors/MultiButtonClickDragBehavior.cpp
@@ -0,0 +1,168 @@
+// Copyright Epic Games, Inc. All Rights Reserved.
+
+#include "Behaviors/MultiButtonClickDragBehavior.h"
+
+#include UE_INLINE_GENERATED_CPP_BY_NAME(MultiButtonClickDragBehavior)
+
+void UMultiButtonClickDragBehavior::Initialize()
+{
+	this->Initialize(this);
+}
+
+FInputCaptureRequest UMultiButtonClickDragBehavior::WantsCapture(const FInputDeviceState& InInput)
+{
+	bInClickDrag = false;	// should never be true here, but weird things can happen w/ focus
+
+	if (IsAnyButtonDown(InInput) && (ModifierCheckFunc == nullptr || ModifierCheckFunc(InInput)) )
+	{
+		const FInputRayHit HitResult = CanBeginClickDragSequence(GetDeviceRay(InInput));
+		if (HitResult.bHit)
+		{
+			return FInputCaptureRequest::Begin(this, EInputCaptureSide::Any, HitResult.HitDepth);
+		}
+	}
+	return FInputCaptureRequest::Ignore();
+}
+
+
+FInputCaptureUpdate UMultiButtonClickDragBehavior::BeginCapture(const FInputDeviceState& InInput, EInputCaptureSide Side)
+{
+	Modifiers.UpdateModifiers(InInput, Target);
+
+	OnStateUpdated(InInput);
+	
+	OnClickPressInternal(InInput, Side);
+	
+	bInClickDrag = true;
+	return FInputCaptureUpdate::Begin(this, EInputCaptureSide::Any);
+}
+
+
+FInputCaptureUpdate UMultiButtonClickDragBehavior::UpdateCapture(const FInputDeviceState& InInput, const FInputCaptureData& Data)
+{
+	if (bUpdateModifiersDuringDrag)
+	{
+		Modifiers.UpdateModifiers(InInput, Target);
+	}
+
+	// We have to check the device before going further because we get passed captures from 
+	// keyboard for modifier key press/releases, and those don't have valid ray data.
+	if ((GetSupportedDevices() & InInput.InputDevice) == EInputDevices::None)
+	{
+		return FInputCaptureUpdate::Continue();
+	}
+
+	if (!IsAnyButtonDown(InInput)) 
+	{
+		OnClickReleaseInternal(InInput, Data);
+		bInClickDrag = false;
+		return FInputCaptureUpdate::End();
+	}
+
+	if (DidAnyButtonChangeState(InInput))
+    {
+		OnStateUpdated(InInput);
+    }
+	
+	OnClickDragInternal(InInput, Data);
+	return FInputCaptureUpdate::Continue();
+}
+
+FInputRayHit UMultiButtonClickDragBehavior::CanBeginClickDragSequence(const FInputDeviceRay& InPressPos)
+{
+	return CanBeginClickDragFunc(InPressPos);
+}
+
+void UMultiButtonClickDragBehavior::OnClickPress(const FInputDeviceRay& InPressPos)
+{
+	OnClickPressFunc(InPressPos);
+}
+
+void UMultiButtonClickDragBehavior::OnClickDrag(const FInputDeviceRay& InDragPos)
+{
+	OnClickDragFunc(InDragPos);
+}
+
+void UMultiButtonClickDragBehavior::OnClickRelease(const FInputDeviceRay& InReleasePos)
+{
+	OnClickReleaseFunc(InReleasePos);
+}
+
+void UMultiButtonClickDragBehavior::OnTerminateDragSequence()
+{
+	OnTerminateFunc();
+}
+
+void UMultiButtonClickDragBehavior::EnableButton(const FKey& InButton)
+{
+	const int32 Index = (InButton == EKeys::LeftMouseButton) ? 0 :
+						(InButton == EKeys::MiddleMouseButton) ? 1 :
+						(InButton == EKeys::RightMouseButton) ? 2 : INDEX_NONE;
+	if (ensure(Index != INDEX_NONE))
+	{
+		CapturedInputs[Index] = 1;
+	}
+}
+
+void UMultiButtonClickDragBehavior::DisableButton(const FKey& InButton)
+{
+	const int32 Index = (InButton == EKeys::LeftMouseButton) ? 0 :
+						(InButton == EKeys::MiddleMouseButton) ? 1 :
+						(InButton == EKeys::RightMouseButton) ? 2 : INDEX_NONE;
+	if (ensure(Index != INDEX_NONE))
+	{
+		CapturedInputs[Index] = 0;
+	}
+}
+
+bool UMultiButtonClickDragBehavior::HandlesLeftMouseButton() const
+{
+	return CapturedInputs[0] == 1;	
+}
+
+bool UMultiButtonClickDragBehavior::HandlesMiddleMouseButton() const
+{
+	return CapturedInputs[1] == 1;
+}
+
+bool UMultiButtonClickDragBehavior::HandlesRightMouseButton() const
+{
+	return CapturedInputs[2] == 1;
+}
+
+bool UMultiButtonClickDragBehavior::IsAnyButtonDown(const FInputDeviceState& InInput)
+{
+	if (InInput.IsFromDevice(EInputDevices::Mouse)) 
+	{
+		ActiveDevice = EInputDevices::Mouse;
+
+		const FDeviceButtonState& Left = InInput.Mouse.Left;
+		const FDeviceButtonState& Middle = InInput.Mouse.Middle;
+		const FDeviceButtonState& Right = InInput.Mouse.Right;
+		
+		const bool bLeftDown = HandlesLeftMouseButton() ? Left.bDown : false;
+		const bool bMiddleDown = HandlesMiddleMouseButton() ? Middle.bDown : false;
+		const bool bRightDown = HandlesRightMouseButton() ? Right.bDown : false;
+		
+		return bLeftDown || bMiddleDown || bRightDown;
+	}
+
+	return false;
+}
+
+bool UMultiButtonClickDragBehavior::DidAnyButtonChangeState(const FInputDeviceState& InInput) const
+{
+	if (InInput.IsFromDevice(EInputDevices::Mouse))
+	{
+		const FDeviceButtonState& Left = InInput.Mouse.Left;
+		const FDeviceButtonState& Middle = InInput.Mouse.Middle;
+		const FDeviceButtonState& Right = InInput.Mouse.Right;
+
+		const bool bLeftChanged = HandlesLeftMouseButton() ? Left.bPressed || Left.bReleased : false;
+		const bool bMiddleChanged = HandlesMiddleMouseButton() ? Middle.bPressed || Middle.bReleased : false;
+		const bool bRightChanged = HandlesRightMouseButton() ? Right.bPressed || Right.bReleased : false;
+
+		return bLeftChanged || bMiddleChanged || bRightChanged;
+	}
+	return false;
+}
diff --git a/Engine/Source/Editor/Experimental/EditorInteractiveToolsFramework/Private/EditorGizmos/TransformGizmo.cpp b/Engine/Source/Editor/Experimental/EditorInteractiveToolsFramework/Private/EditorGizmos/TransformGizmo.cpp
--- a/Engine/Source/Editor/Experimental/EditorInteractiveToolsFramework/Private/EditorGizmos/TransformGizmo.cpp
+++ b/Engine/Source/Editor/Experimental/EditorInteractiveToolsFramework/Private/EditorGizmos/TransformGizmo.cpp
@@ -13,6 +13,7 @@
 #include "Engine/World.h"
 #include "Materials/Material.h"
 #include "UnrealEngine.h"
+#include "Behaviors/MultiButtonClickDragBehavior.h"
 #include "Intersection/IntersectionUtil.h"
 
 #define LOCTEXT_NAMESPACE "UTransformGizmo"
@@ -34,6 +35,7 @@ void UTransformGizmo::Setup()
 	UInteractiveGizmo::Setup();
 
 	SetupBehaviors();
+	SetupIndirectBehaviors();
 	SetupMaterials();
 	SetupOnClickFunctions();
 
@@ -67,22 +69,30 @@ void UTransformGizmo::SetupBehaviors()
 	MouseBehavior->Initialize(this);
 	MouseBehavior->SetDefaultPriority(FInputCapturePriority(FInputCapturePriority::DEFAULT_GIZMO_PRIORITY));
 	AddInputBehavior(MouseBehavior);
+}
 
+void UTransformGizmo::SetupIndirectBehaviors()
+{
 	// Add middle mouse input behavior for indirect manipulation
 	ULocalClickDragInputBehavior* MiddleClickDragBehavior = NewObject<ULocalClickDragInputBehavior>();
 	MiddleClickDragBehavior->Initialize();
 	MiddleClickDragBehavior->SetUseMiddleMouseButton();
-	MiddleClickDragBehavior->CanBeginClickDragFunc = [this](const FInputDeviceRay& InPressPos)
+	MiddleClickDragBehavior->CanBeginClickDragFunc = [](const FInputDeviceRay&)
 	{
-		bIndirectManipulation = true;
-		return CanBeginClickDragSequence(InPressPos);
+		return FInputRayHit(TNumericLimits<double>::Max());
 	};
 	MiddleClickDragBehavior->OnClickPressFunc = [this](const FInputDeviceRay& InPressPos)
 	{
+		bIndirectManipulation = true;
+		if (LastHitPart == ETransformGizmoPartIdentifier::Default)
+		{
+			LastHitPart = GetCurrentModeLastHitPart();
+		}
 		return OnClickPress(InPressPos);
 	};
 	MiddleClickDragBehavior->OnClickDragFunc = [this](const FInputDeviceRay& InDragPos)
 	{
+		bIndirectManipulation = true;
 		return OnClickDrag(InDragPos);
 	};
 	MiddleClickDragBehavior->OnClickReleaseFunc = [this](const FInputDeviceRay& InReleasePos)
@@ -92,11 +102,116 @@ void UTransformGizmo::SetupBehaviors()
 	};
 	MiddleClickDragBehavior->OnTerminateFunc = [this]()
 	{
+		bIndirectManipulation = false;
 		return OnTerminateDragSequence();
 	};
 	// add this to enable indirect manipulation using CTRL + MMB
 	// MiddleClickDragBehavior->ModifierCheckFunc = FInputDeviceState::IsCtrlKeyDown;
 	AddInputBehavior(MiddleClickDragBehavior);
+
+	
+	// Add left/right mouse input behavior for indirect manipulation
+	UMultiButtonClickDragBehavior* LeftRightClickDragBehavior = NewObject<UMultiButtonClickDragBehavior>();
+	LeftRightClickDragBehavior->Initialize();
+	LeftRightClickDragBehavior->EnableButton(EKeys::LeftMouseButton);
+	LeftRightClickDragBehavior->EnableButton(EKeys::RightMouseButton);
+	LeftRightClickDragBehavior->ModifierCheckFunc = FInputDeviceState::IsCtrlKeyDown;
+	LeftRightClickDragBehavior->CanBeginClickDragFunc = [](const FInputDeviceRay&)
+	{
+		return FInputRayHit(TNumericLimits<double>::Max());
+	};
+	LeftRightClickDragBehavior->OnClickPressFunc = [this](const FInputDeviceRay& InPressPos)
+	{
+		bIndirectManipulation = true;
+		if (LastHitPart == ETransformGizmoPartIdentifier::Default)
+		{
+			LastHitPart = GetCurrentModeLastHitPart();
+		}
+		return OnClickPress(InPressPos);
+	};
+	LeftRightClickDragBehavior->OnClickDragFunc = [this](const FInputDeviceRay& InDragPos)
+	{
+		bIndirectManipulation = true;
+		return OnClickDrag(InDragPos);
+	};
+	LeftRightClickDragBehavior->OnClickReleaseFunc = [this](const FInputDeviceRay& InReleasePos)
+	{
+		bIndirectManipulation = false;
+		return OnClickRelease(InReleasePos);
+	};
+	LeftRightClickDragBehavior->OnTerminateFunc = [this]()
+	{
+		bIndirectManipulation = false;
+		return OnTerminateDragSequence();
+	};
+	
+	auto GetAxis = [](const FInputDeviceState& InInput)
+	{
+		const bool bAddX = InInput.Mouse.Left.bDown;
+		const bool bAddY = InInput.Mouse.Right.bDown;
+		return bAddX && bAddY ? EAxis::Z : bAddX ? EAxis::X : bAddY ? EAxis::Y : EAxis::None;
+	};
+
+	auto GetHitPart = [this](const EAxis::Type InAxis)
+	{
+		static constexpr ETransformGizmoPartIdentifier TranslateIds [4] = {
+			ETransformGizmoPartIdentifier::Default,
+			ETransformGizmoPartIdentifier::TranslateXAxis,
+			ETransformGizmoPartIdentifier::TranslateYAxis,
+			ETransformGizmoPartIdentifier::TranslateZAxis};
+			
+		static constexpr ETransformGizmoPartIdentifier RotateIds [4] = {
+			ETransformGizmoPartIdentifier::Default,
+			ETransformGizmoPartIdentifier::RotateXAxis,
+			ETransformGizmoPartIdentifier::RotateYAxis,
+			ETransformGizmoPartIdentifier::RotateZAxis};
+
+		static constexpr ETransformGizmoPartIdentifier ScaleIds [4] = {
+			ETransformGizmoPartIdentifier::Default,
+			ETransformGizmoPartIdentifier::ScaleXAxis,
+			ETransformGizmoPartIdentifier::ScaleYAxis,
+			ETransformGizmoPartIdentifier::ScaleZAxis};
+
+		switch (CurrentMode)
+		{
+		case EGizmoTransformMode::Translate: return TranslateIds[InAxis];
+		case EGizmoTransformMode::Rotate: return RotateIds[InAxis];
+		case EGizmoTransformMode::Scale: return ScaleIds[InAxis];
+		default: break;
+		}
+		return ETransformGizmoPartIdentifier::Default;
+	};
+	
+	LeftRightClickDragBehavior->OnStateUpdated = [this, GetAxis, GetHitPart, LeftRightClickDragBehavior](const FInputDeviceState& Input)
+	{
+		// disable indirect if the current axis is none 
+		const EAxis::Type Axis = GetAxis(Input);
+		if (Axis == EAxis::None)
+		{
+			bIndirectManipulation = false;
+			return;
+		}
+
+		bIndirectManipulation = true;
+		
+		const ETransformGizmoPartIdentifier HitPart = GetHitPart(Axis);
+		if (HitPart != GetCurrentModeLastHitPart())
+		{
+			// update interaction state
+			UpdateInteractingState(false, GetCurrentModeLastHitPart(), true);
+			SetModeLastHitPart(CurrentMode, HitPart);
+			UpdateInteractingState(true, HitPart, true);
+
+			// reinitialize OnClickPress data
+			LastHitPart = HitPart;
+			const uint8 HitPartIndex = static_cast<uint8>(LastHitPart);
+			if (OnClickPressFunctions.IsValidIndex(HitPartIndex) && OnClickPressFunctions[HitPartIndex])
+			{
+				OnClickPressFunctions[HitPartIndex](this, LeftRightClickDragBehavior->GetDeviceRay(Input));
+			}
+		}
+	};
+	AddInputBehavior(LeftRightClickDragBehavior);
 }
 
 void UTransformGizmo::SetupMaterials()
@@ -337,15 +452,6 @@ ETransformGizmoPartIdentifier UTransformGizmo::GetCurrentModeLastHitPart() const
 
 FInputRayHit UTransformGizmo::CanBeginClickDragSequence(const FInputDeviceRay& PressPos)
 {
-	if (bIndirectManipulation)
-	{
-		if (LastHitPart == ETransformGizmoPartIdentifier::Default)
-		{
-			LastHitPart = GetCurrentModeLastHitPart();
-		}
-		return FInputRayHit(0.0);
-	}
-	
 	FInputRayHit RayHit;
 
 	if (HitTarget)
@@ -1405,19 +1511,94 @@ void UTransformGizmo::OnClickPressTranslateZAxis(const FInputDeviceRay& PressPos
 
 void UTransformGizmo::OnClickPressAxis(const FInputDeviceRay& PressPos)
 {
+	InteractionPlanarOrigin = CurrentTransform.GetLocation();
 	InteractionAxisStartParam = GetNearestRayParamToInteractionAxis(PressPos);
 	InteractionAxisCurrParam = InteractionAxisStartParam;
-	bInInteraction = true;
 
+	// indirect manipulation uses a 2D approach instead as there's no guaranty to intersect a plane 
+	if (bIndirectManipulation)
+	{
+		InteractionScreenCurrPos = PressPos.ScreenPosition;
+		StartRotation = CurrentRotation = ActiveTarget->GetTransform().GetRotation();
+		bInInteraction = true;
+		SetModeLastHitPart(EGizmoTransformMode::Translate, LastHitPart);
+		return;
+	}
+
+	// compute plane and axis to mute
+	const FVector XAxis = GetWorldAxis(FVector::XAxisVector);
+	const FVector YAxis = GetWorldAxis(FVector::YAxisVector);
+	const FVector ZAxis = GetWorldAxis(FVector::ZAxisVector);
+
+	const FVector ViewDirection = GizmoViewContext->GetViewDirection();
+	const double XDot = FMath::Abs(FVector::DotProduct(ViewDirection, XAxis));
+	const double YDot = FMath::Abs(FVector::DotProduct(ViewDirection, YAxis));
+	const double ZDot = FMath::Abs(FVector::DotProduct(ViewDirection, ZAxis));
+
+	if (FVector::DotProduct(InteractionAxisDirection, XAxis) > 0.1)
+	{
+		InteractionPlanarNormal = (YDot > ZDot) ? YAxis : ZAxis;
+		NormalToRemove = (YDot > ZDot) ? ZAxis : YAxis;
+	}
+	else if (FVector::DotProduct(InteractionAxisDirection, YAxis) > 0.1)
+	{
+		InteractionPlanarNormal = (XDot > ZDot) ? XAxis : ZAxis;
+		NormalToRemove = (XDot > ZDot) ? ZAxis : XAxis;
+	}
+	else
+	{
+		InteractionPlanarNormal = (XDot > YDot) ? XAxis : YAxis;
+		NormalToRemove = (XDot > YDot) ? YAxis : XAxis;
+	}
+	
+	float HitDepth;
+	if (GetRayParamIntersectionWithInteractionPlane(PressPos, HitDepth))
+	{
+		InteractionPlanarStartPoint = PressPos.WorldRay.Origin + PressPos.WorldRay.Direction * HitDepth;
+		InteractionPlanarCurrPoint = InteractionPlanarStartPoint;
+	}
+
+	bInInteraction = true;
 	SetModeLastHitPart(EGizmoTransformMode::Translate, LastHitPart);
 }
 
 void UTransformGizmo::OnClickDragTranslateAxis(const FInputDeviceRay& DragPos)
 {
-	float AxisNearestParam = GetNearestRayParamToInteractionAxis(DragPos);
-	FVector Delta = ComputeAxisTranslateDelta(InteractionAxisCurrParam, AxisNearestParam);
-	ApplyTranslateDelta(Delta);
-	InteractionAxisCurrParam = AxisNearestParam;
+	// indirect manipulation uses a 2D projection approach instead of plane intersection
+	if (bIndirectManipulation)
+	{
+		const FVector2D DragDir = DragPos.ScreenPosition - InteractionScreenCurrPos;
+
+		const FVector2D XAxisDir = GetScreenProjectedAxis(GizmoViewContext, GetWorldAxis(FVector::XAxisVector));
+		const FVector2D YAxisDir = GetScreenProjectedAxis(GizmoViewContext, GetWorldAxis(FVector::YAxisVector));
+		const FVector2D ZAxisDir = GetScreenProjectedAxis(GizmoViewContext, GetWorldAxis(FVector::ZAxisVector));
+		
+		FVector Delta((InteractionAxisList == EAxisList::X) ? FVector2D::DotProduct(XAxisDir, DragDir) : 0.0,
+					  (InteractionAxisList == EAxisList::Y) ? FVector2D::DotProduct(YAxisDir, DragDir) : 0.0,
+					  (InteractionAxisList == EAxisList::Z) ? FVector2D::DotProduct(ZAxisDir, DragDir) : 0.0);
+		Delta = CurrentRotation * Delta;
+		
+		ApplyTranslateDelta(Delta);
+
+		InteractionScreenCurrPos = DragPos.ScreenPosition;
+		return;
+	}
+
+	float HitDepth;
+	if (GetRayParamIntersectionWithInteractionPlane(DragPos, HitDepth))
+	{
+		FVector HitPoint = DragPos.WorldRay.Origin + DragPos.WorldRay.Direction * HitDepth;
+
+		const FVector DeltaToStart = HitPoint-InteractionPlanarStartPoint;
+		const FVector AxisToRemove = NormalToRemove * FVector::DotProduct(DeltaToStart, NormalToRemove); 
+		
+		HitPoint -= AxisToRemove;
+	
+		const FVector Delta = ComputePlanarTranslateDelta(InteractionPlanarCurrPoint, HitPoint);
+
+		ApplyTranslateDelta(Delta);
+		InteractionPlanarCurrPoint = HitPoint;
+	}
 }
 
 void UTransformGizmo::OnClickReleaseTranslateAxis(const FInputDeviceRay& InReleasePos)
@@ -1473,8 +1654,8 @@ void UTransformGizmo::OnClickDragTranslatePlanar(const FInputDeviceRay& DragPos)
 	float HitDepth;
 	if (GetRayParamIntersectionWithInteractionPlane(DragPos, HitDepth))
 	{
-		FVector HitPoint = DragPos.WorldRay.Origin + DragPos.WorldRay.Direction * HitDepth;
-		FVector Delta = ComputePlanarTranslateDelta(InteractionPlanarCurrPoint, HitPoint);
+		const FVector HitPoint = DragPos.WorldRay.Origin + DragPos.WorldRay.Direction * HitDepth;
+		const FVector Delta = ComputePlanarTranslateDelta(InteractionPlanarCurrPoint, HitPoint);
 		ApplyTranslateDelta(Delta);
 		InteractionPlanarCurrPoint = HitPoint;
 	}
diff --git a/Engine/Source/Editor/Experimental/EditorInteractiveToolsFramework/Public/Behaviors/MultiButtonClickDragBehavior.h b/Engine/Source/Editor/Experimental/EditorInteractiveToolsFramework/Public/Behaviors/MultiButtonClickDragBehavior.h
new file mode 100644
--- /dev/null
+++ b/Engine/Source/Editor/Experimental/EditorInteractiveToolsFramework/Public/Behaviors/MultiButtonClickDragBehavior.h
@@ -0,0 +1,85 @@
+// Copyright Epic Games, Inc. All Rights Reserved.
+
+#pragma once
+
+#include "BaseBehaviors/ClickDragBehavior.h"
+
+#include "MultiButtonClickDragBehavior.generated.h"
+
+class UObject;
+
+/**
+ * MultiButtonClickDragBehavior is an implementation of UClickDragInputBehavior that can manage several
+ * mouse buttons at once and also implements IClickDragBehaviorTarget directly, via a set of local lambda functions.
+ * To use/customize this class the client replaces the lambda functions with their own and enables/disables the
+ * buttons to capture. 
+ */
+
+UCLASS(MinimalAPI)
+class UMultiButtonClickDragBehavior : public UClickDragInputBehavior, public IClickDragBehaviorTarget
+{
+	GENERATED_BODY()
+protected:
+	using UClickDragInputBehavior::Initialize;
+	
+public:
+
+	// UClickDragInputBehavior  implementation
+	void Initialize();
+
+	virtual FInputRayHit CanBeginClickDragSequence(const FInputDeviceRay& InPressPos) override;
+	virtual void OnClickPress(const FInputDeviceRay& InPressPos) override;
+	virtual void OnClickDrag(const FInputDeviceRay& InDragPos) override;
+	virtual void OnClickRelease(const FInputDeviceRay& InReleasePos) override;
+	virtual void OnTerminateDragSequence() override;
+	
+	// UInputBehavior implementation
+
+	virtual FInputCaptureRequest WantsCapture(const FInputDeviceState& Input) override;
+	virtual FInputCaptureUpdate BeginCapture(const FInputDeviceState& Input, EInputCaptureSide Side) override;
+	virtual FInputCaptureUpdate UpdateCapture(const FInputDeviceState& Input, const FInputCaptureData& Data) override;
+
+	/** lambda implementation of CanBeginClickDragSequence */
+	TUniqueFunction<FInputRayHit(const FInputDeviceRay& PressPos)> CanBeginClickDragFunc = [](const FInputDeviceRay&) { return FInputRayHit(); };
+
+	/** lambda implementation of OnClickPress */
+	TUniqueFunction<void(const FInputDeviceRay& PressPos)> OnClickPressFunc = [](const FInputDeviceRay&) {};
+
+	/** lambda implementation of OnClickDrag */
+	TUniqueFunction<void(const FInputDeviceRay& PressPos)> OnClickDragFunc = [](const FInputDeviceRay&) {};
+
+	/** lambda implementation of OnClickRelease */
+	TUniqueFunction<void(const FInputDeviceRay& ReleasePos)> OnClickReleaseFunc = [](const FInputDeviceRay&) {};
+
+	/** lambda implementation of OnTerminateDragSequence */
+	TUniqueFunction<void()> OnTerminateFunc = []() {};
+	
+	/** lambda implementation of OnStateUpdated */
+	TUniqueFunction<void(const FInputDeviceState& Input)> OnStateUpdated = [](const FInputDeviceState&) {};
+
+	/** enable InButton so that in can be captured */
+	void EnableButton(const FKey& InButton);
+	
+	/** disable InButton so that it's not captured */
+	void DisableButton(const FKey& InButton);
+
+protected:
+
+	/** @return true if any captured button is down. */
+	bool IsAnyButtonDown(const FInputDeviceState& InInput);
+
+	/** @return true if any captured button changed its bPressed/bReleased state. */
+	bool DidAnyButtonChangeState(const FInputDeviceState& InInput) const;
+	
+	/** @return true if left mouse button is handled. */
+	bool HandlesLeftMouseButton() const;
+	
+	/** @return true if middle mouse button is handled. */
+	bool HandlesMiddleMouseButton() const;
+	
+	/** @return true if right mouse button is handled. */
+	bool HandlesRightMouseButton() const;
+
+	/** mouse buttons that need to be captured. Left=0, Middle=1, Right=2*/
+	uint8 CapturedInputs[3] = {0 ,0, 0};
+};
diff --git a/Engine/Source/Editor/Experimental/EditorInteractiveToolsFramework/Public/EditorGizmos/TransformGizmo.h b/Engine/Source/Editor/Experimental/EditorInteractiveToolsFramework/Public/EditorGizmos/TransformGizmo.h
--- a/Engine/Source/Editor/Experimental/EditorInteractiveToolsFramework/Public/EditorGizmos/TransformGizmo.h
+++ b/Engine/Source/Editor/Experimental/EditorInteractiveToolsFramework/Public/EditorGizmos/TransformGizmo.h
@@ -5,7 +5,6 @@
 #include "BaseBehaviors/BehaviorTargetInterfaces.h"
 #include "BaseGizmos/GizmoElementHitTargets.h"
 #include "BaseGizmos/GizmoViewContext.h"
-#include "BaseGizmos/StateTargets.h"
 #include "BaseGizmos/TransformProxy.h"
 #include "Containers/EnumAsByte.h"
 #include "CoreMinimal.h"
@@ -13,7 +12,6 @@
 #include "InputState.h"
 #include "InteractiveGizmo.h"
 #include "InteractiveToolChange.h"
-#include "InteractiveToolObjects.h"
 #include "Materials/MaterialInstanceDynamic.h"
 #include "Math/Axis.h"
 #include "Math/Color.h"
@@ -392,6 +390,9 @@ protected:
 	/** Setup behaviors */
 	virtual void SetupBehaviors();
 
+	/** Setup indirect behaviors */
+	void SetupIndirectBehaviors();
+
 	/** Setup materials */
 	virtual void SetupMaterials();
 
@@ -702,12 +703,6 @@ protected:
 	/** Array of function pointers, indexed by gizmo part id, to handle click release behavior */
 	TArray<TFunction<void(UTransformGizmo* TransformGizmo, const FInputDeviceRay& ReleasePos)> > OnClickReleaseFunctions;
 
-	/** Array of function pointers, indexed by gizmo part id, to handle update hovering state */
-	TArray<TFunction<void(UTransformGizmo* TransformGizmo, bool bInHover, uint32 InHitParts)> > OnUpdateHoverFunctions;
-
-	/** Array of function pointers, indexed by gizmo part id, to handle update interacting state */
-	TArray<TFunction<void(UTransformGizmo* TransformGizmo, bool bInInteracting, uint32 InHitPart)> > OnUpdateInteractingFunctions;
-
 	/** Customization function (to override default material or increment gizmo size for example) */
 	TFunction<const FGizmoCustomization()> CustomizationFunction;
 
@@ -770,6 +765,10 @@ protected:
 	UPROPERTY()
 	FVector InteractionPlanarNormal;
 
+	/** Active normal to remove from axis translation (only valid between state target BeginModify/EndModify) */
+	UPROPERTY()
+	FVector NormalToRemove;
+
 	/** Active world space axis X used for planar (only valid between state target BeginModify/EndModify) */
 	UPROPERTY()
 	FVector InteractionPlanarAxisX;
