New TRS GIzmo: fixed indirect translation screen axis computation

#rnx

diff --git a/Engine/Source/Editor/Experimental/EditorInteractiveToolsFramework/Private/EditorGizmos/TransformGizmo.cpp b/Engine/Source/Editor/Experimental/EditorInteractiveToolsFramework/Private/EditorGizmos/TransformGizmo.cpp
--- a/Engine/Source/Editor/Experimental/EditorInteractiveToolsFramework/Private/EditorGizmos/TransformGizmo.cpp
+++ b/Engine/Source/Editor/Experimental/EditorInteractiveToolsFramework/Private/EditorGizmos/TransformGizmo.cpp
@@ -1569,9 +1569,9 @@ void UTransformGizmo::OnClickDragTranslateAxis(const FInputDeviceRay& DragPos)
 	{
 		const FVector2D DragDir = DragPos.ScreenPosition - InteractionScreenCurrPos;
 
-		const FVector2D XAxisDir = GetScreenProjectedAxis(GizmoViewContext, GetWorldAxis(FVector::XAxisVector));
-		const FVector2D YAxisDir = GetScreenProjectedAxis(GizmoViewContext, GetWorldAxis(FVector::YAxisVector));
-		const FVector2D ZAxisDir = GetScreenProjectedAxis(GizmoViewContext, GetWorldAxis(FVector::ZAxisVector));
+		const FVector2D XAxisDir = GetScreenProjectedAxis(GizmoViewContext, FVector::XAxisVector, CurrentTransform);
+		const FVector2D YAxisDir = GetScreenProjectedAxis(GizmoViewContext, FVector::YAxisVector, CurrentTransform);
+		const FVector2D ZAxisDir = GetScreenProjectedAxis(GizmoViewContext, FVector::ZAxisVector, CurrentTransform);
 		
 		FVector Delta((InteractionAxisList == EAxisList::X) ? FVector2D::DotProduct(XAxisDir, DragDir) : 0.0,
 					  (InteractionAxisList == EAxisList::Y) ? FVector2D::DotProduct(YAxisDir, DragDir) : 0.0,
