New TRS Gizmo: do not handle indirect manipulation if the gizmo is hidden

#rnx

diff --git a/Engine/Source/Editor/Experimental/EditorInteractiveToolsFramework/Private/EditorGizmos/TransformGizmo.cpp b/Engine/Source/Editor/Experimental/EditorInteractiveToolsFramework/Private/EditorGizmos/TransformGizmo.cpp
--- a/Engine/Source/Editor/Experimental/EditorInteractiveToolsFramework/Private/EditorGizmos/TransformGizmo.cpp
+++ b/Engine/Source/Editor/Experimental/EditorInteractiveToolsFramework/Private/EditorGizmos/TransformGizmo.cpp
@@ -77,9 +77,11 @@ void UTransformGizmo::SetupIndirectBehaviors()
 	ULocalClickDragInputBehavior* MiddleClickDragBehavior = NewObject<ULocalClickDragInputBehavior>();
 	MiddleClickDragBehavior->Initialize();
 	MiddleClickDragBehavior->SetUseMiddleMouseButton();
-	MiddleClickDragBehavior->CanBeginClickDragFunc = [](const FInputDeviceRay&)
+	MiddleClickDragBehavior->CanBeginClickDragFunc = [this](const FInputDeviceRay&)
 	{
-		return FInputRayHit(TNumericLimits<double>::Max());
+		static const FInputRayHit InvalidRayHit;
+		static const FInputRayHit ValidRayHit(TNumericLimits<double>::Max());
+		return bVisible ? ValidRayHit : InvalidRayHit;
 	};
 	MiddleClickDragBehavior->OnClickPressFunc = [this](const FInputDeviceRay& InPressPos)
 	{
@@ -116,9 +118,11 @@ void UTransformGizmo::SetupIndirectBehaviors()
 	LeftRightClickDragBehavior->EnableButton(EKeys::LeftMouseButton);
 	LeftRightClickDragBehavior->EnableButton(EKeys::RightMouseButton);
 	LeftRightClickDragBehavior->ModifierCheckFunc = FInputDeviceState::IsCtrlKeyDown;
-	LeftRightClickDragBehavior->CanBeginClickDragFunc = [](const FInputDeviceRay&)
+	LeftRightClickDragBehavior->CanBeginClickDragFunc = [this](const FInputDeviceRay&)
 	{
-		return FInputRayHit(TNumericLimits<double>::Max());
+		static const FInputRayHit InvalidRayHit;
+		static const FInputRayHit ValidRayHit(TNumericLimits<double>::Max());
+		return bVisible ? ValidRayHit : InvalidRayHit;
 	};
 	LeftRightClickDragBehavior->OnClickPressFunc = [this](const FInputDeviceRay& InPressPos)
 	{
@@ -454,7 +458,7 @@ FInputRayHit UTransformGizmo::CanBeginClickDragSequence(const FInputDeviceRay& P
 {
 	FInputRayHit RayHit;
 
-	if (HitTarget)
+	if (bVisible && HitTarget)
 	{
 		RayHit = HitTarget->IsHit(PressPos);
 		ETransformGizmoPartIdentifier HitPart;
@@ -698,6 +702,7 @@ void UTransformGizmo::EnableRotate(EAxisList::Type InAxisListToDraw)
 		if (RotateOuterCircleElement == nullptr)
 		{
 			RotateOuterCircleElement = MakeRotateCircleHandle(ETransformGizmoPartIdentifier::Default, RotateOuterCircleRadius, RotateOuterCircleColor, false);
+			RotateOuterCircleElement->SetHittableState(false);
 			GizmoElementRoot->Add(RotateOuterCircleElement);
 		}
 
