New TRS Gizmo: hide gizmo if FEditorViewportClient::bShowWidget is false

#rnx

diff --git a/Engine/Source/Editor/Experimental/EditorInteractiveToolsFramework/Private/EditorGizmos/EditorTransformGizmoSource.cpp b/Engine/Source/Editor/Experimental/EditorInteractiveToolsFramework/Private/EditorGizmos/EditorTransformGizmoSource.cpp
--- a/Engine/Source/Editor/Experimental/EditorInteractiveToolsFramework/Private/EditorGizmos/EditorTransformGizmoSource.cpp
+++ b/Engine/Source/Editor/Experimental/EditorInteractiveToolsFramework/Private/EditorGizmos/EditorTransformGizmoSource.cpp
@@ -70,6 +70,11 @@ bool UEditorTransformGizmoSource::GetVisible() const
 {
 	if (const FEditorViewportClient* ViewportClient = GetViewportClient()) 
 	{
+		if (!ViewportClient->GetShowWidget())
+		{
+			return false;
+		}
+		
 		const FEditorModeTools& ModeTools = GetModeTools();
 		if (ModeTools.GetShowWidget() && ModeTools.UsesTransformWidget())
 		{
diff --git a/Engine/Source/Editor/UnrealEd/Public/EditorViewportClient.h b/Engine/Source/Editor/UnrealEd/Public/EditorViewportClient.h
--- a/Engine/Source/Editor/UnrealEd/Public/EditorViewportClient.h
+++ b/Engine/Source/Editor/UnrealEd/Public/EditorViewportClient.h
@@ -1251,6 +1251,7 @@ public:
 
 	/** Show or hide the widget. */
 	UNREALED_API void ShowWidget(const bool bShow);
+	UNREALED_API bool GetShowWidget() const { return bShowWidget; }
 
 	/**
 	 * Returns whether or not the flight camera is active
