New TRS Gizmo: CTRL + MMB does indirect manipulation on the Y axiz (CTRL+LMB does Z instead of CTRL+RMB+LMB)

#rnx

diff --git a/Engine/Source/Editor/Experimental/EditorInteractiveToolsFramework/Private/EditorGizmos/TransformGizmo.cpp b/Engine/Source/Editor/Experimental/EditorInteractiveToolsFramework/Private/EditorGizmos/TransformGizmo.cpp
--- a/Engine/Source/Editor/Experimental/EditorInteractiveToolsFramework/Private/EditorGizmos/TransformGizmo.cpp
+++ b/Engine/Source/Editor/Experimental/EditorInteractiveToolsFramework/Private/EditorGizmos/TransformGizmo.cpp
@@ -107,8 +107,14 @@ void UTransformGizmo::SetupIndirectBehaviors()
 		bIndirectManipulation = false;
 		return OnTerminateDragSequence();
 	};
-	// add this to enable indirect manipulation using CTRL + MMB
-	// MiddleClickDragBehavior->ModifierCheckFunc = FInputDeviceState::IsCtrlKeyDown;
+	if (bCtrlMiddleDoesY)
+	{
+		// disable ctrl + mmb for that behavior
+		MiddleClickDragBehavior->ModifierCheckFunc = [](const FInputDeviceState& InputState)
+		{
+			return !FInputDeviceState::IsCtrlKeyDown(InputState);
+		};
+	}
 	AddInputBehavior(MiddleClickDragBehavior);
 
 	
@@ -116,6 +122,10 @@ void UTransformGizmo::SetupIndirectBehaviors()
 	UMultiButtonClickDragBehavior* LeftRightClickDragBehavior = NewObject<UMultiButtonClickDragBehavior>();
 	LeftRightClickDragBehavior->Initialize();
 	LeftRightClickDragBehavior->EnableButton(EKeys::LeftMouseButton);
+	if (bCtrlMiddleDoesY)
+	{
+		LeftRightClickDragBehavior->EnableButton(EKeys::MiddleMouseButton);
+	}
 	LeftRightClickDragBehavior->EnableButton(EKeys::RightMouseButton);
 	LeftRightClickDragBehavior->ModifierCheckFunc = FInputDeviceState::IsCtrlKeyDown;
 	LeftRightClickDragBehavior->CanBeginClickDragFunc = [this](const FInputDeviceRay&)
@@ -149,9 +159,15 @@ void UTransformGizmo::SetupIndirectBehaviors()
 		return OnTerminateDragSequence();
 	};
 	
-	auto GetAxis = [](const FInputDeviceState& InInput)
+	auto GetAxis = [this](const FInputDeviceState& InInput)
 	{
 		const bool bAddX = InInput.Mouse.Left.bDown;
+		if (bCtrlMiddleDoesY)
+		{
+			const bool bAddY = InInput.Mouse.Middle.bDown;
+			const bool bAddZ = InInput.Mouse.Right.bDown;
+			return bAddX ? EAxis::X : bAddY ? EAxis::Y : bAddZ ? EAxis::Z : EAxis::None;
+		}
 		const bool bAddY = InInput.Mouse.Right.bDown;
 		return bAddX && bAddY ? EAxis::Z : bAddX ? EAxis::X : bAddY ? EAxis::Y : EAxis::None;
 	};
@@ -355,7 +371,7 @@ void UTransformGizmo::OnEndHover()
 
 FInputRayHit UTransformGizmo::UpdateHoveredPart(const FInputDeviceRay& PressPos)
 {
-	if (!HitTarget)
+	if (!HitTarget || !bVisible)
 	{
 		return FInputRayHit();
 	}
@@ -701,7 +717,7 @@ void UTransformGizmo::EnableRotate(EAxisList::Type InAxisListToDraw)
 
 		if (RotateOuterCircleElement == nullptr)
 		{
-			RotateOuterCircleElement = MakeRotateCircleHandle(ETransformGizmoPartIdentifier::Default, RotateOuterCircleRadius, RotateOuterCircleColor, false);
+			RotateOuterCircleElement = MakeRotateCircleHandle(ETransformGizmoPartIdentifier::RotateAll, RotateOuterCircleRadius, RotateOuterCircleColor, false);
 			RotateOuterCircleElement->SetHittableState(false);
 			GizmoElementRoot->Add(RotateOuterCircleElement);
 		}
@@ -1023,7 +1039,7 @@ UGizmoElementRectangle* UTransformGizmo::MakeTranslateScreenSpaceHandle()
 	RectangleElement->SetViewAlignAxis(FVector::UpVector);
 	RectangleElement->SetViewAlignNormal(-FVector::ForwardVector);
 	RectangleElement->SetMaterial(TransparentVertexColorMaterial);
-	RectangleElement->SetLineColor(ScreenSpaceColor);
+	RectangleElement->SetVertexColor(WhiteColor, true);
 	RectangleElement->SetHitMesh(true);
 	RectangleElement->SetDrawMesh(false);
 	RectangleElement->SetDrawLine(true);
diff --git a/Engine/Source/Editor/Experimental/EditorInteractiveToolsFramework/Public/EditorGizmos/TransformGizmo.h b/Engine/Source/Editor/Experimental/EditorInteractiveToolsFramework/Public/EditorGizmos/TransformGizmo.h
--- a/Engine/Source/Editor/Experimental/EditorInteractiveToolsFramework/Public/EditorGizmos/TransformGizmo.h
+++ b/Engine/Source/Editor/Experimental/EditorInteractiveToolsFramework/Public/EditorGizmos/TransformGizmo.h
@@ -834,4 +834,8 @@ protected:
 
 	/** Pending drag function to be called if bDeferDrag is true */
 	TFunction<void()> PendingDragFunction;
+
+	/** Use Ctrl + MMB to do indirect manipulation on the Y axis */
+	UPROPERTY()
+	bool bCtrlMiddleDoesY = true;
 };
