New TRS Gizmo: partial view dependent rendering improvements
    - use the ViewAlignAxialMaxCosAngleTol property instead of the default one to enable per-element customization
    - reduced the max partial angle of rotation handles so that they don't depth fight with the arcball handle on edges

#rnx

diff --git a/Engine/Source/Editor/Experimental/EditorInteractiveToolsFramework/Private/EditorGizmos/TransformGizmo.cpp b/Engine/Source/Editor/Experimental/EditorInteractiveToolsFramework/Private/EditorGizmos/TransformGizmo.cpp
--- a/Engine/Source/Editor/Experimental/EditorInteractiveToolsFramework/Private/EditorGizmos/TransformGizmo.cpp
+++ b/Engine/Source/Editor/Experimental/EditorInteractiveToolsFramework/Private/EditorGizmos/TransformGizmo.cpp
@@ -348,7 +348,7 @@ void UTransformGizmo::OnBeginHover(const FInputDeviceRay& DevicePos)
 
 bool UTransformGizmo::OnUpdateHover(const FInputDeviceRay& DevicePos)
 {
-	FInputRayHit RayHit = UpdateHoveredPart(DevicePos);
+	const FInputRayHit RayHit = UpdateHoveredPart(DevicePos);
 	return RayHit.bHit;
 }
 
@@ -715,13 +715,6 @@ void UTransformGizmo::EnableRotate(EAxisList::Type InAxisListToDraw)
 			GizmoElementRoot->Add(RotateScreenSpaceElement);
 		}
 
-		if (RotateOuterCircleElement == nullptr)
-		{
-			RotateOuterCircleElement = MakeRotateCircleHandle(ETransformGizmoPartIdentifier::Default, RotateOuterCircleRadius, RotateOuterCircleColor, false);
-			RotateOuterCircleElement->SetHittableState(false);
-			GizmoElementRoot->Add(RotateOuterCircleElement);
-		}
-
 		if (RotateArcballElement == nullptr)
 		{
 			RotateArcballElement = MakeRotateCircleHandle(ETransformGizmoPartIdentifier::RotateArcball, RotateArcballSphereRadius, RotateArcballCircleColor, true);
@@ -1068,6 +1061,8 @@ UGizmoElementTorus* UTransformGizmo::MakeRotateAxis(ETransformGizmoPartIdentifie
 	RotateAxisElement->SetPartialEndAngle(UE_PI);
 	RotateAxisElement->SetViewDependentAxis(TorusNormal);
 	RotateAxisElement->SetViewAlignType(EGizmoElementViewAlignType::Axial);
+	RotateAxisElement->SetViewAlignAxialAngleTol(UE_DOUBLE_SMALL_NUMBER);
+	RotateAxisElement->SetPartialViewDependentMaxCosTol(1.0-UE_DOUBLE_SMALL_NUMBER);	
 	RotateAxisElement->SetViewAlignAxis(TorusNormal);
 	RotateAxisElement->SetViewAlignNormal(TorusAxis1);
 	RotateAxisElement->SetMaterial(InMaterial);
diff --git a/Engine/Source/Runtime/InteractiveToolsFramework/Private/BaseGizmos/GizmoElementBase.cpp b/Engine/Source/Runtime/InteractiveToolsFramework/Private/BaseGizmos/GizmoElementBase.cpp
--- a/Engine/Source/Runtime/InteractiveToolsFramework/Private/BaseGizmos/GizmoElementBase.cpp
+++ b/Engine/Source/Runtime/InteractiveToolsFramework/Private/BaseGizmos/GizmoElementBase.cpp
@@ -139,7 +139,7 @@ bool UGizmoElementBase::GetViewAlignRot(const FVector& InViewLocation, const FVe
 	else if (ViewAlignType == EGizmoElementViewAlignType::Axial)
 	{
 		// if Axis and Dir are almost coincident, do not adjust the rotation
-		if ((FMath::Abs(FVector::DotProduct(ViewAlignAxis, -LocalViewDir))) >= DefaultViewAlignMaxCosAngleTol)
+		if ((FMath::Abs(FVector::DotProduct(ViewAlignAxis, -LocalViewDir))) >= ViewAlignAxialMaxCosAngleTol)
 		{
 			return false;
 		}
diff --git a/Engine/Source/Runtime/InteractiveToolsFramework/Private/BaseGizmos/GizmoElementCircleBase.cpp b/Engine/Source/Runtime/InteractiveToolsFramework/Private/BaseGizmos/GizmoElementCircleBase.cpp
--- a/Engine/Source/Runtime/InteractiveToolsFramework/Private/BaseGizmos/GizmoElementCircleBase.cpp
+++ b/Engine/Source/Runtime/InteractiveToolsFramework/Private/BaseGizmos/GizmoElementCircleBase.cpp
@@ -23,21 +23,29 @@ bool UGizmoElementCircleBase::IsPartial(const UGizmoViewContext* View, const FVe
 }
 
 bool UGizmoElementCircleBase::IsPartial(const FVector& InWorldCenter, const FVector& InWorldNormal,
-	const FVector& InViewLocation, const FVector& InViewDirection, bool bIsPerspectiveProjection)
+	const FVector& InViewLocation, const FVector& InViewDirection, const bool bIsPerspectiveProjection)
 {
-	FVector ViewToCircleBaseCenter;
-	if (bIsPerspectiveProjection)
+	switch (PartialType)
 	{
-		ViewToCircleBaseCenter = (InWorldCenter - InViewLocation).GetSafeNormal();
-	}
-	else
-	{
-		ViewToCircleBaseCenter = InViewDirection;
+		case EGizmoElementPartialType::None:
+			return false;
+		case EGizmoElementPartialType::Partial:
+			return true;
+		case EGizmoElementPartialType::PartialViewDependent:
+		{
+			FVector ViewToCircleBaseCenter(InViewDirection);
+			if (bIsPerspectiveProjection)
+			{
+				ViewToCircleBaseCenter = (InWorldCenter - InViewLocation).GetSafeNormal();
+			}
+			const double DotP = FMath::Abs(FVector::DotProduct(InWorldNormal, ViewToCircleBaseCenter));
+			return (DotP <= PartialViewDependentMaxCosTol);
+		}
+		default:
+			break;
 	}
 
-	double DotP = FMath::Abs(FVector::DotProduct(InWorldNormal, ViewToCircleBaseCenter));
-	return (PartialType == EGizmoElementPartialType::Partial ||
-		(PartialType == EGizmoElementPartialType::PartialViewDependent && DotP <= PartialViewDependentMaxCosTol));
+	return false;
 }
 
 void UGizmoElementCircleBase::SetCenter(const FVector& InCenter)
diff --git a/Engine/Source/Runtime/InteractiveToolsFramework/Private/BaseGizmos/GizmoElementTorus.cpp b/Engine/Source/Runtime/InteractiveToolsFramework/Private/BaseGizmos/GizmoElementTorus.cpp
--- a/Engine/Source/Runtime/InteractiveToolsFramework/Private/BaseGizmos/GizmoElementTorus.cpp
+++ b/Engine/Source/Runtime/InteractiveToolsFramework/Private/BaseGizmos/GizmoElementTorus.cpp
@@ -1,7 +1,6 @@
 // Copyright Epic Games, Inc. All Rights Reserved.
  
 #include "BaseGizmos/GizmoElementTorus.h"
-#include "BaseGizmos/GizmoRenderingUtil.h"
 #include "BaseGizmos/GizmoMath.h"
 #include "InputState.h"
 #include "Materials/MaterialInterface.h"
@@ -24,7 +23,7 @@ void UGizmoElementTorus::Render(IToolsContextRenderAPI* RenderAPI, const FRender
 			const FVector Normal = Axis0 ^ Axis1; 
 			const FVector WorldNormal = CurrentRenderState.LocalToWorldTransform.TransformVectorNoScale(Normal);
 
-			bool bPartial = IsPartial(RenderAPI->GetSceneView(), WorldCenter, WorldNormal);
+			const bool bPartial = IsPartial(RenderAPI->GetSceneView(), WorldCenter, WorldNormal);
 			FVector BeginAxis = bPartial ? Axis0.RotateAngleAxis(PartialStartAngle, Normal).GetSafeNormal() : Axis0;
 
 			const float PartialAngle = static_cast<float>(PartialEndAngle - PartialStartAngle);			
diff --git a/Engine/Source/Runtime/InteractiveToolsFramework/Public/BaseGizmos/GizmoElementCircleBase.h b/Engine/Source/Runtime/InteractiveToolsFramework/Public/BaseGizmos/GizmoElementCircleBase.h
--- a/Engine/Source/Runtime/InteractiveToolsFramework/Public/BaseGizmos/GizmoElementCircleBase.h
+++ b/Engine/Source/Runtime/InteractiveToolsFramework/Public/BaseGizmos/GizmoElementCircleBase.h
@@ -111,5 +111,5 @@ protected:
 	// Returns whether element should be partial based on current view
 	INTERACTIVETOOLSFRAMEWORK_API bool IsPartial(const FVector& InWorldCenter, const FVector& InWorldNormal,
 		const FVector& InViewLocation, const FVector& InViewDirection,
-		bool bIsPerspectiveProjection);
+		const bool bIsPerspectiveProjection);
 };
