New TRS Gizmo: toggle between arcball and screen space when re-entering rotate mode.

#rnx

diff --git a/Engine/Source/Editor/Experimental/EditorInteractiveToolsFramework/Private/EditorGizmos/TransformGizmo.cpp b/Engine/Source/Editor/Experimental/EditorInteractiveToolsFramework/Private/EditorGizmos/TransformGizmo.cpp
--- a/Engine/Source/Editor/Experimental/EditorInteractiveToolsFramework/Private/EditorGizmos/TransformGizmo.cpp
+++ b/Engine/Source/Editor/Experimental/EditorInteractiveToolsFramework/Private/EditorGizmos/TransformGizmo.cpp
@@ -910,23 +910,29 @@ void UTransformGizmo::HandleWidgetModeChanged(UE::Widget::EWidgetMode InWidgetMo
 
 	if (CurrentMode != EGizmoTransformMode::None && NewMode == CurrentMode)
 	{
-		auto GetModeDefaultHitPart = [NewMode]()
+		const ETransformGizmoPartIdentifier CurrentModeLastHitPart = GetCurrentModeLastHitPart();
+		auto GetModeDefaultHitPart = [NewMode, CurrentModeLastHitPart]()
 		{
+			const bool bIsRotateArcBall = (CurrentModeLastHitPart == ETransformGizmoPartIdentifier::RotateArcball);
 			switch (NewMode)
 			{
-			case EGizmoTransformMode::Translate: return ETransformGizmoPartIdentifier::TranslateScreenSpace;
-			case EGizmoTransformMode::Rotate: return ETransformGizmoPartIdentifier::RotateArcball;
-			case EGizmoTransformMode::Scale: return ETransformGizmoPartIdentifier::ScaleUniform;
-			default: return ETransformGizmoPartIdentifier::Default;
+			case EGizmoTransformMode::Translate:
+				return ETransformGizmoPartIdentifier::TranslateScreenSpace;
+			case EGizmoTransformMode::Rotate:
+				return bIsRotateArcBall ? ETransformGizmoPartIdentifier::RotateScreenSpace : ETransformGizmoPartIdentifier::RotateArcball;
+			case EGizmoTransformMode::Scale:
+				return ETransformGizmoPartIdentifier::ScaleUniform;
+			default:
+				return ETransformGizmoPartIdentifier::Default;
 			}
 			return ETransformGizmoPartIdentifier::Default;
 		};
 
 		const ETransformGizmoPartIdentifier DefaultHitPart = GetModeDefaultHitPart();
-		if (DefaultHitPart != GetCurrentModeLastHitPart())
+		if (DefaultHitPart != CurrentModeLastHitPart)
 		{
 			// reset indirect manipulation to default
-			UpdateInteractingState(false, GetCurrentModeLastHitPart(), true);
+			UpdateInteractingState(false, CurrentModeLastHitPart, true);
 			SetModeLastHitPart(CurrentMode, DefaultHitPart);
 			UpdateInteractingState(true, DefaultHitPart, true);
 		}
