New TRS Gizmo: new behavior for rotation handles

    - the "pull" direction is now based on the closest projection to the rotation circle
    - normal to that pull direction is removed to avoid poor drag behavior
    - added a few temp cvars to experiment and fine tune the behavior

        - Gizmos.DotThreshold 0.2: Dot threshold for determining whether the rotation plane is perpendicular to the camera view [0.2, 1.0]
        - Gizmos.DebugDraw false: Displays debugging information.
        - Gizmos.ProjectIndirect true: Project to the nearest point of the curve when handling indirect rotation.

diff --git a/Engine/Source/Editor/Experimental/EditorInteractiveToolsFramework/Private/EditorGizmos/TransformGizmo.cpp b/Engine/Source/Editor/Experimental/EditorInteractiveToolsFramework/Private/EditorGizmos/TransformGizmo.cpp
--- a/Engine/Source/Editor/Experimental/EditorInteractiveToolsFramework/Private/EditorGizmos/TransformGizmo.cpp
+++ b/Engine/Source/Editor/Experimental/EditorInteractiveToolsFramework/Private/EditorGizmos/TransformGizmo.cpp
@@ -15,11 +15,41 @@
 #include "UnrealEngine.h"
 #include "Behaviors/MultiButtonClickDragBehavior.h"
 #include "Intersection/IntersectionUtil.h"
+#include "SceneManagement.h"
 
 #define LOCTEXT_NAMESPACE "UTransformGizmo"
 
 DEFINE_LOG_CATEGORY_STATIC(LogTransformGizmo, Log, All);
 
+namespace GizmoLocals
+{
+	
+// NOTE these variables are not intended to remain here indefinitely.
+// Their purpose is to experiment the new behavior of rotation gizmos.
+
+static float DotThreshold = 0.2f;
+static FAutoConsoleVariableRef CVarDotThreshold(
+	TEXT("Gizmos.DotThreshold"),
+	DotThreshold,
+	TEXT("Dot threshold for determining whether the rotation plane is perpendicular to the camera view [0.2, 1.0]")
+);
+
+static bool	bDebugDraw = false;
+static FAutoConsoleVariableRef CVarDebugDraw(
+	TEXT("Gizmos.DebugDraw"),
+	bDebugDraw,
+	TEXT("Displays debugging information.")
+	);
+
+static bool	ProjectIndirect = true;
+static FAutoConsoleVariableRef CVarProjectIndirect(
+	TEXT("Gizmos.ProjectIndirect"),
+	ProjectIndirect,
+	TEXT("Project to the nearest point of the curve when handling indirect rotation.")
+	);
+	
+}
+
 void UTransformGizmo::SetDisallowNegativeScaling(bool bDisallow)
 {
 	bDisallowNegativeScaling = bDisallow;
@@ -334,6 +364,16 @@ void UTransformGizmo::Render(IToolsContextRenderAPI* RenderAPI)
 		UGizmoElementBase::FRenderTraversalState RenderState;
 		RenderState.Initialize(RenderAPI->GetSceneView(), GetGizmoTransform());
 		GizmoElementRoot->Render(RenderAPI, RenderState);
+
+		if (GizmoLocals::bDebugDraw)
+		{
+			if (bDebugRotate)
+			{
+				const float Radius = 2.f * GetWorldRadius(RotateAxisOuterRadius);
+				FPrimitiveDrawInterface* PDI = RenderAPI->GetPrimitiveDrawInterface();
+				PDI->DrawLine(DebugClosest - (DebugDirection * Radius), DebugClosest + (DebugDirection * Radius), FLinearColor::Yellow, SDPG_Foreground);
+			}
+		}
 	}
 }
 
@@ -1866,58 +1906,133 @@ FVector UTransformGizmo::ComputeScaleDelta(const FVector2D& InStartPos, const FV
 	return Scale;
 }
 
-void UTransformGizmo::OnClickPressRotateXAxis(const FInputDeviceRay& PressPos)
+void UTransformGizmo::OnClickPressRotateXAxis(const FInputDeviceRay& InPressPos)
 {
-	InteractionScreenAxisDirection = GetScreenRotateAxisDir(FVector::ZAxisVector, FVector::YAxisVector).GetSafeNormal();
+	InteractionScreenAxisDirection = GetScreenRotateAxisDir(InPressPos);
 	InteractionAxisList = EAxisList::X;
-	InteractionScreenStartPos = InteractionScreenCurrPos = PressPos.ScreenPosition;
+	InteractionScreenStartPos = InteractionScreenCurrPos = InPressPos.ScreenPosition;
 	bInInteraction = true;
 
 	SetModeLastHitPart(EGizmoTransformMode::Rotate, LastHitPart);
+	
+	bDebugRotate = true;
 }
 
-void UTransformGizmo::OnClickPressRotateYAxis(const FInputDeviceRay& PressPos)
+void UTransformGizmo::OnClickPressRotateYAxis(const FInputDeviceRay& InPressPos)
 {
-	InteractionScreenAxisDirection = GetScreenRotateAxisDir(FVector::XAxisVector, FVector::ZAxisVector).GetSafeNormal();
+	InteractionScreenAxisDirection = GetScreenRotateAxisDir(InPressPos);
 	InteractionAxisList = EAxisList::Y;
-	InteractionScreenStartPos = InteractionScreenCurrPos = PressPos.ScreenPosition;
+	InteractionScreenStartPos = InteractionScreenCurrPos = InPressPos.ScreenPosition;
 	bInInteraction = true;
 	
 	SetModeLastHitPart(EGizmoTransformMode::Rotate, LastHitPart);
+	
+	bDebugRotate = true;
 }
 
-void UTransformGizmo::OnClickPressRotateZAxis(const FInputDeviceRay& PressPos)
+void UTransformGizmo::OnClickPressRotateZAxis(const FInputDeviceRay& InPressPos)
 {
-	InteractionScreenAxisDirection = GetScreenRotateAxisDir(FVector::XAxisVector, FVector::YAxisVector).GetSafeNormal();
+	InteractionScreenAxisDirection = GetScreenRotateAxisDir(InPressPos);
 	InteractionAxisList = EAxisList::Z;
-	InteractionScreenStartPos = InteractionScreenCurrPos = PressPos.ScreenPosition;
+	InteractionScreenStartPos = InteractionScreenCurrPos = InPressPos.ScreenPosition;
 	bInInteraction = true;
 
 	SetModeLastHitPart(EGizmoTransformMode::Rotate, LastHitPart);
+
+	bDebugRotate = true;
 }
 
-FVector2D UTransformGizmo::GetScreenRotateAxisDir(const FVector& InAxis0, const FVector& InAxis1)
+FVector2D UTransformGizmo::GetScreenRotateAxisDir(const FInputDeviceRay& InPressPos)
 {
-	check(GizmoViewContext);
-	const FVector DirectionToWidget = CurrentTransform.GetLocation() - GizmoViewContext->ViewLocation;
+	// NOTE that function is not intended to remain here indefinitely, its purpose is to debug closest point computation
+	const auto PrintProjection = [this](const TCHAR* const InMessage)
+	{
+		if (bDebugRotate && GizmoLocals::bDebugDraw)
+		{
+			UE_LOG(LogTransformGizmo, Warning, TEXT("%s"), InMessage);
+		}
+	};
+
+	static const TArray RotateIDs({	ETransformGizmoPartIdentifier::RotateXAxis,
+									ETransformGizmoPartIdentifier::RotateYAxis,
+									ETransformGizmoPartIdentifier::RotateZAxis});
+	const int32 RotateID = RotateIDs.IndexOfByKey(LastHitPart);
+	if (!ensure(RotateID != INDEX_NONE))
+	{
+		return FVector2D::ZeroVector;
+	}
+	
+	const FRay& Ray = InPressPos.WorldRay;
+
+	// store world origin and axis
+	static const TArray RotateAxis({FVector::XAxisVector, FVector::YAxisVector, -FVector::ZAxisVector});
+	const FVector WorldAxis = GetWorldAxis(RotateAxis[RotateID]);
+	const FVector WorldOrigin = CurrentTransform.GetLocation();
+
+	// compute axis / view direction projection: is the rotation plane nearly perpendicular to the view plane?  
+	const double Threshold = FMath::Clamp(static_cast<double>(GizmoLocals::DotThreshold), 0.2, 1.0);
+	const bool bAxisPerpendicularToView = FMath::Abs(FVector::DotProduct(WorldAxis, GizmoViewContext->GetViewDirection())) < Threshold;
+	// compute axis / ray direction projection: is the ray direction parallel to the axis?
+	const bool bRayPerpendicularToAxis = FMath::IsNearlyZero(FVector::DotProduct(WorldAxis, Ray.Direction));
+
+	// compute closest point on the rotate handle
+	const bool bUseRayForIndirect = bIndirectManipulation && !GizmoLocals::ProjectIndirect;
+	const bool bUseRayOrigin = bUseRayForIndirect || bAxisPerpendicularToView || bRayPerpendicularToAxis;
+	// compute the closest point from plane intersection if we can
+	FVector QueryPoint = Ray.Origin;
+	if (!bUseRayOrigin)
+	{
+		const FPlane Plane(WorldOrigin, WorldAxis);
+
+		// if the projection is in front of the camera then use it
+		const double HitDepth = FMath::RayPlaneIntersectionParam(Ray.Origin, Ray.Direction, Plane);
+		if (HitDepth >= 0.0)
+		{
+			PrintProjection(TEXT("front"));
+			QueryPoint = Ray.Origin + Ray.Direction * HitDepth;
+		}
+		else
+		{
+			PrintProjection(TEXT("behind"));
+		}
+	}
+	else
+	{
+		PrintProjection(TEXT("ray origin"));
+	}
+
+	// compute nearest point
+	const float Radius = GetWorldRadius(RotateAxisOuterRadius);
+	FVector ClosestPointOnCircle;
+	GizmoMath::ClosetPointOnCircle(QueryPoint, WorldOrigin, WorldAxis, Radius, ClosestPointOnCircle);
+
+	// compute world directions
+	const FVector ToClosestDirection = (ClosestPointOnCircle - WorldOrigin).GetSafeNormal();
+	const FVector PullDirection = FVector::CrossProduct(ToClosestDirection, WorldAxis);
 
-	const FVector Axis0 = GetWorldAxis(InAxis0);
-	const FVector Axis1 = GetWorldAxis(InAxis1);
+	// compute screen projections
+	const FTransform ToClosest(ClosestPointOnCircle);
+	const FVector2D PullProjection = GetScreenProjectedAxis(GizmoViewContext, PullDirection, ToClosest);
+	const FVector2D AxisProjection = GetScreenProjectedAxis(GizmoViewContext, WorldAxis, ToClosest);
+	const FVector2D ToClosestProjection = GetScreenProjectedAxis(GizmoViewContext, ToClosestDirection, ToClosest);
 
-	// Reverse the axes based on camera view
-	const bool bMirrorAxis0 = (FVector::DotProduct(Axis0, DirectionToWidget) <= 0.0f);
-	const bool bMirrorAxis1 = (FVector::DotProduct(Axis1, DirectionToWidget) <= 0.0f);
-	const float Direction = (bMirrorAxis0 ^ bMirrorAxis1) ? -1.0f : 1.0f;
+	// compute which projection to remove from drag
+	const double DotAxis = FMath::Abs( FVector2D::DotProduct(PullProjection, AxisProjection) );
+	const double DotClosest = FMath::Abs( FVector2D::DotProduct(PullProjection, ToClosestProjection) );
+	NormalProjectionToRemove = (DotAxis < DotClosest) ? AxisProjection : ToClosestProjection;
 
-	const FVector2D Axis0Screen = GetScreenProjectedAxis(GizmoViewContext, bMirrorAxis0 ? Axis0:-Axis0);
-	const FVector2D Axis1Screen = GetScreenProjectedAxis(GizmoViewContext, bMirrorAxis1 ? Axis1:-Axis1);
+	// debug
+	{
+		DebugDirection = PullDirection;
+		DebugClosest = ClosestPointOnCircle;
+	}
 	
-	return ((Axis1Screen - Axis0Screen) * Direction).GetSafeNormal();
+	return PullProjection;
 }
 
 void UTransformGizmo::OnClickDragRotateAxis(const FInputDeviceRay& DragPos)
 {
-	FQuat DeltaRot = ComputeAxisRotateDelta(InteractionScreenCurrPos, DragPos.ScreenPosition);
+	const FQuat DeltaRot = ComputeAxisRotateDelta(InteractionScreenCurrPos, DragPos.ScreenPosition);
 	ApplyRotateDelta(DeltaRot);
 	InteractionScreenCurrPos = DragPos.ScreenPosition;
 }
@@ -1925,19 +2040,15 @@ void UTransformGizmo::OnClickDragRotateAxis(const FInputDeviceRay& DragPos)
 FQuat UTransformGizmo::ComputeAxisRotateDelta(const FVector2D& InStartPos, const FVector2D& InEndPos)
 {
 	FVector2D DragDir = InEndPos - InStartPos;
-	FRotator DeltaRot(0.0, 0.0, 0.0);
-	if (InteractionAxisList == EAxisList::X)
-	{
-		DeltaRot.Roll = FVector2D::DotProduct(InteractionScreenAxisDirection, DragDir);
-	}
-	else if (InteractionAxisList == EAxisList::Y)
-	{
-		DeltaRot.Pitch = FVector2D::DotProduct(InteractionScreenAxisDirection, DragDir);
-	}
-	else
-	{
-		DeltaRot.Yaw = FVector2D::DotProduct(InteractionScreenAxisDirection, DragDir);
-	}
+
+	const FVector2D DragDirToRemove = NormalProjectionToRemove * FVector2D::DotProduct(DragDir, NormalProjectionToRemove);
+	DragDir -= DragDirToRemove;
+	
+	const double Delta = FVector2D::DotProduct(InteractionScreenAxisDirection, DragDir);
+	FRotator DeltaRot(
+		InteractionAxisList == EAxisList::Y ? Delta : 0.0,
+		InteractionAxisList == EAxisList::Z ? Delta : 0.0,
+		InteractionAxisList == EAxisList::X ? Delta : 0.0);
 
 	auto GetCoordinateSystem = [&]()
 	{
@@ -1951,7 +2062,7 @@ FQuat UTransformGizmo::ComputeAxisRotateDelta(const FVector2D& InStartPos, const
 	if (GetCoordinateSystem() == EToolContextCoordinateSystem::Local)
 	{
 		check(ActiveTarget);
-		FMatrix CurrCoordSystem = ActiveTarget->GetTransform().ToMatrixNoScale();
+		const FMatrix CurrCoordSystem = ActiveTarget->GetTransform().ToMatrixNoScale();
 		DeltaRot = (CurrCoordSystem.Inverse() * FRotationMatrix(DeltaRot) * CurrCoordSystem).Rotator();
 	}
 
@@ -1961,6 +2072,7 @@ FQuat UTransformGizmo::ComputeAxisRotateDelta(const FVector2D& InStartPos, const
 void UTransformGizmo::OnClickReleaseRotateAxis(const FInputDeviceRay& InReleasePos)
 {
 	bInInteraction = false;
+	bDebugRotate = false;
 }
 
 void UTransformGizmo::OnClickPressScreenSpaceRotate(const FInputDeviceRay& PressPos)
@@ -2073,7 +2185,7 @@ void UTransformGizmo::OnClickPressArcBallRotate(const FInputDeviceRay& PressPos)
 	
 	const FVector& RayOrigin = PressPos.WorldRay.Origin;
 	const FVector& RayDir = PressPos.WorldRay.Direction;		
-	const double SphereRadius = GetArcBallWorldRadius();
+	const double SphereRadius = GetWorldRadius(RotateArcballSphereRadius);
 
 	StartRotation = CurrentRotation = CurrentTransform.GetRotation();
 	InteractionPlanarOrigin = CurrentTransform.GetLocation();
@@ -2121,7 +2233,7 @@ void UTransformGizmo::OnClickDragArcBallRotate(const FInputDeviceRay& DragPos)
 {
 	const FVector& RayOrigin = DragPos.WorldRay.Origin;
 	const FVector& RayDir = DragPos.WorldRay.Direction;		
-	const float SphereRadius = GetArcBallWorldRadius();
+	const float SphereRadius = GetWorldRadius(RotateArcballSphereRadius);
 
 	// compute projection
 	ArcBallLocals::GetSphereAndHyperbolicProjection(
@@ -2156,11 +2268,11 @@ void UTransformGizmo::OnClickReleaseArcBallRotate(const FInputDeviceRay& Release
 	bInInteraction = false;
 }
 
-float UTransformGizmo::GetArcBallWorldRadius() const
+float UTransformGizmo::GetWorldRadius(const float InRadius) const
 {
 	const float PixelToWorldScale = GizmoRenderingUtil::CalculateLocalPixelToWorldScale(GizmoViewContext, CurrentTransform.GetLocation());
 	const float GizmoScale = TransformGizmoSource ? TransformGizmoSource->GetGizmoScale() : 1.0f;
-	return RotateArcballSphereRadius * GetSizeCoefficient() * PixelToWorldScale * GizmoScale;
+	return InRadius * GetSizeCoefficient() * PixelToWorldScale * GizmoScale;
 }
 
 float UTransformGizmo::GetSizeCoefficient() const
diff --git a/Engine/Source/Editor/Experimental/EditorInteractiveToolsFramework/Public/EditorGizmos/TransformGizmo.h b/Engine/Source/Editor/Experimental/EditorInteractiveToolsFramework/Public/EditorGizmos/TransformGizmo.h
--- a/Engine/Source/Editor/Experimental/EditorInteractiveToolsFramework/Public/EditorGizmos/TransformGizmo.h
+++ b/Engine/Source/Editor/Experimental/EditorInteractiveToolsFramework/Public/EditorGizmos/TransformGizmo.h
@@ -535,22 +535,22 @@ protected:
 	 */
 
 	/** Handle click press for rotate X axis */
-	virtual void OnClickPressRotateXAxis(const FInputDeviceRay& PressPos);
+	virtual void OnClickPressRotateXAxis(const FInputDeviceRay& InPressPos);
 
 	/** Handle click press for rotate Y axis */
-	virtual void OnClickPressRotateYAxis(const FInputDeviceRay& PressPos);
+	virtual void OnClickPressRotateYAxis(const FInputDeviceRay& InPressPos);
 
 	/** Handle click press for rotate Z axis */
-	virtual void OnClickPressRotateZAxis(const FInputDeviceRay& PressPos);
+	virtual void OnClickPressRotateZAxis(const FInputDeviceRay& InPressPos);
 
 	/** Handle click drag for rotate axis */
-	virtual void OnClickDragRotateAxis(const FInputDeviceRay& DragPos);
+	virtual void OnClickDragRotateAxis(const FInputDeviceRay& InDragPos);
 
 	/** Handle click release for rotate axes */
-	virtual void OnClickReleaseRotateAxis(const FInputDeviceRay& ReleasePos);
+	virtual void OnClickReleaseRotateAxis(const FInputDeviceRay& InReleasePos);
 
 	/** Get screen-space axis for rotation drag */
-	FVector2D GetScreenRotateAxisDir(const FVector& InAxis0, const FVector& InAxis1);
+	FVector2D GetScreenRotateAxisDir(const FInputDeviceRay& InPressPos);
 
 	/** Compute rotate delta based on screen-space start/end positions */
 	virtual FQuat ComputeAxisRotateDelta(const FVector2D& InStartPos, const FVector2D& InEndPos);
@@ -585,7 +585,7 @@ protected:
 	virtual void OnClickReleaseArcBallRotate(const FInputDeviceRay& ReleasePos);
 
 	/** Get the arc ball sphere world radius */
-	float GetArcBallWorldRadius() const;
+	float GetWorldRadius(const float InRadius) const;
 
 	/** */
 	float GetSizeCoefficient() const;
@@ -797,6 +797,10 @@ protected:
 	UPROPERTY()
 	FVector2D InteractionScreenAxisDirection;
 
+	/** Active normal projection to remove from drag when rotating */
+	UPROPERTY()
+	FVector2D NormalProjectionToRemove;
+
 	/** Active interaction screen start pos (only valid between state target BeginModify/EndModify) */
 	UPROPERTY()
 	FVector2D InteractionScreenStartPos;
@@ -838,4 +842,10 @@ protected:
 	/** Use Ctrl + MMB to do indirect manipulation on the Y axis */
 	UPROPERTY()
 	bool bCtrlMiddleDoesY = true;
+
+private:
+	/** Debug attributes to display the pull direction */
+	bool bDebugRotate = false;
+	FVector DebugDirection = FVector::ZeroVector;
+	FVector DebugClosest = FVector::ZeroVector;
 };
