New TRS Gizmo: fixed view direction projection

#rnx

diff --git a/Engine/Source/Editor/Experimental/EditorInteractiveToolsFramework/Private/EditorGizmos/TransformGizmo.cpp b/Engine/Source/Editor/Experimental/EditorInteractiveToolsFramework/Private/EditorGizmos/TransformGizmo.cpp
--- a/Engine/Source/Editor/Experimental/EditorInteractiveToolsFramework/Private/EditorGizmos/TransformGizmo.cpp
+++ b/Engine/Source/Editor/Experimental/EditorInteractiveToolsFramework/Private/EditorGizmos/TransformGizmo.cpp
@@ -31,7 +31,11 @@ static float DotThreshold = 0.2f;
 static FAutoConsoleVariableRef CVarDotThreshold(
 	TEXT("Gizmos.DotThreshold"),
 	DotThreshold,
-	TEXT("Dot threshold for determining whether the rotation plane is perpendicular to the camera view [0.2, 1.0]")
+	TEXT("Dot threshold for determining whether the rotation plane is perpendicular to the camera view [0.2, 1.0]"),
+	FConsoleVariableDelegate::CreateLambda([](IConsoleVariable*)
+	{
+		DotThreshold = FMath::Clamp(DotThreshold, 0.2, 1.0);
+	})
 );
 
 static bool	bDebugDraw = false;
@@ -1967,9 +1971,11 @@ FVector2D UTransformGizmo::GetScreenRotateAxisDir(const FInputDeviceRay& InPress
 	const FVector WorldAxis = GetWorldAxis(RotateAxis[RotateID]);
 	const FVector WorldOrigin = CurrentTransform.GetLocation();
 
-	// compute axis / view direction projection: is the rotation plane nearly perpendicular to the view plane?  
-	const double Threshold = FMath::Clamp(static_cast<double>(GizmoLocals::DotThreshold), 0.2, 1.0);
-	const bool bAxisPerpendicularToView = FMath::Abs(FVector::DotProduct(WorldAxis, GizmoViewContext->GetViewDirection())) < Threshold;
+	// compute axis / view direction projection: is the rotation plane nearly perpendicular to the view plane?
+	const FVector ViewDirection = GizmoViewContext->IsPerspectiveProjection() ?
+		(WorldOrigin - GizmoViewContext->ViewLocation).GetSafeNormal() :
+		GizmoViewContext->GetViewDirection();
+	const bool bAxisPerpendicularToView = FMath::Abs(FVector::DotProduct(WorldAxis, ViewDirection)) < GizmoLocals::DotThreshold;
 	// compute axis / ray direction projection: is the ray direction parallel to the axis?
 	const bool bRayPerpendicularToAxis = FMath::IsNearlyZero(FVector::DotProduct(WorldAxis, Ray.Direction));
 
