New TRS Gizmo: enable indirect manipulation with hidden gizmo

diff --git a/Engine/Source/Editor/Experimental/EditorInteractiveToolsFramework/Private/EditorGizmos/EditorTransformGizmoSource.cpp b/Engine/Source/Editor/Experimental/EditorInteractiveToolsFramework/Private/EditorGizmos/EditorTransformGizmoSource.cpp
--- a/Engine/Source/Editor/Experimental/EditorInteractiveToolsFramework/Private/EditorGizmos/EditorTransformGizmoSource.cpp
+++ b/Engine/Source/Editor/Experimental/EditorInteractiveToolsFramework/Private/EditorGizmos/EditorTransformGizmoSource.cpp
@@ -72,7 +72,15 @@ bool UEditorTransformGizmoSource::GetVisible() const
 		{
 			return false;
 		}
-		
+		return CanInteract();
+	}
+	return false;
+}
+
+bool UEditorTransformGizmoSource::CanInteract() const
+{
+	if (const FEditorViewportClient* ViewportClient = GetViewportClient())
+	{
 		const FEditorModeTools& ModeTools = GetModeTools();
 		if (ModeTools.GetShowWidget() && ModeTools.UsesTransformWidget())
 		{
@@ -86,11 +94,9 @@ bool UEditorTransformGizmoSource::GetVisible() const
 					bUseLegacyWidget = UseLegacyWidgetCVar->GetInt() > 0;
 				}
 			}
-
 			return !bUseLegacyWidget;
 		}
 	}
-
 	return false;
 }
 
diff --git a/Engine/Source/Editor/Experimental/EditorInteractiveToolsFramework/Private/EditorGizmos/TransformGizmo.cpp b/Engine/Source/Editor/Experimental/EditorInteractiveToolsFramework/Private/EditorGizmos/TransformGizmo.cpp
--- a/Engine/Source/Editor/Experimental/EditorInteractiveToolsFramework/Private/EditorGizmos/TransformGizmo.cpp
+++ b/Engine/Source/Editor/Experimental/EditorInteractiveToolsFramework/Private/EditorGizmos/TransformGizmo.cpp
@@ -115,7 +115,7 @@ void UTransformGizmo::SetupIndirectBehaviors()
 	{
 		static const FInputRayHit InvalidRayHit;
 		static const FInputRayHit ValidRayHit(TNumericLimits<double>::Max());
-		return bVisible ? ValidRayHit : InvalidRayHit;
+		return CanInteract() ? ValidRayHit : InvalidRayHit;
 	};
 	MiddleClickDragBehavior->OnClickPressFunc = [this](const FInputDeviceRay& InPressPos)
 	{
@@ -166,7 +166,7 @@ void UTransformGizmo::SetupIndirectBehaviors()
 	{
 		static const FInputRayHit InvalidRayHit;
 		static const FInputRayHit ValidRayHit(TNumericLimits<double>::Max());
-		return bVisible ? ValidRayHit : InvalidRayHit;
+		return CanInteract() ? ValidRayHit : InvalidRayHit;
 	};
 	LeftRightClickDragBehavior->OnClickPressFunc = [this](const FInputDeviceRay& InPressPos)
 	{
@@ -361,7 +361,7 @@ FTransform UTransformGizmo::GetGizmoTransform() const
 
 void UTransformGizmo::Render(IToolsContextRenderAPI* RenderAPI)
 {
-	if (bVisible && GizmoElementRoot && RenderAPI)
+	if (IsVisible() && GizmoElementRoot && RenderAPI)
 	{
 		CurrentTransform = ActiveTarget->GetTransform();
 
@@ -415,7 +415,7 @@ void UTransformGizmo::OnEndHover()
 
 FInputRayHit UTransformGizmo::UpdateHoveredPart(const FInputDeviceRay& PressPos)
 {
-	if (!HitTarget || !bVisible)
+	if (!HitTarget || !IsVisible())
 	{
 		return FInputRayHit();
 	}
@@ -517,8 +517,8 @@ ETransformGizmoPartIdentifier UTransformGizmo::GetCurrentModeLastHitPart() const
 FInputRayHit UTransformGizmo::CanBeginClickDragSequence(const FInputDeviceRay& PressPos)
 {
 	FInputRayHit RayHit;
-
-	if (bVisible && HitTarget)
+		
+	if (CanInteract() && HitTarget)
 	{
 		RayHit = HitTarget->IsHit(PressPos);
 		ETransformGizmoPartIdentifier HitPart;
@@ -568,38 +568,21 @@ void UTransformGizmo::UpdateMode()
 		return EAxisList::Type::All;
 	};
 
-	auto GetVisible = [&]()
-	{
-		if (TransformGizmoSource)
-		{
-			return TransformGizmoSource->GetVisible();
-		}
-		return true;
-	};
+	const EGizmoTransformMode NewMode = GetTransformMode();
+	const EAxisList::Type NewAxisToDraw = GetAxisToDraw();
 	
-	if (GetVisible())
+	if (NewMode != CurrentMode)
 	{
-		const EGizmoTransformMode NewMode = GetTransformMode();
-		const EAxisList::Type NewAxisToDraw = GetAxisToDraw();
-		
-		if (NewMode != CurrentMode)
-		{
-			EnableMode(CurrentMode, EAxisList::None);
-			EnableMode(NewMode, NewAxisToDraw);
+		EnableMode(CurrentMode, EAxisList::None);
+		EnableMode(NewMode, NewAxisToDraw);
 
-			CurrentMode = NewMode;
-			CurrentAxisToDraw = NewAxisToDraw;
-		}
-		else if (NewAxisToDraw != CurrentAxisToDraw)
-		{
-			EnableMode(CurrentMode, NewAxisToDraw);
-			CurrentAxisToDraw = NewAxisToDraw;
-		}
+		CurrentMode = NewMode;
+		CurrentAxisToDraw = NewAxisToDraw;
 	}
-	else
+	else if (NewAxisToDraw != CurrentAxisToDraw)
 	{
-		EnableMode(CurrentMode, EAxisList::None);
-		CurrentMode = EGizmoTransformMode::None;
+		EnableMode(CurrentMode, NewAxisToDraw);
+		CurrentAxisToDraw = NewAxisToDraw;
 	}
 }
 
@@ -865,6 +848,25 @@ void UTransformGizmo::UpdateCameraAxisSource()
 	}
 }
 
+bool UTransformGizmo::IsVisible() const
+{
+	if (TransformGizmoSource)
+	{
+		return bVisible && TransformGizmoSource->GetVisible();
+	}	
+	return bVisible;
+}
+
+bool UTransformGizmo::CanInteract() const
+{
+	const bool bValidMode = CurrentMode > EGizmoTransformMode::None && CurrentMode < EGizmoTransformMode::Max;
+	if (TransformGizmoSource)
+	{
+		return bValidMode && TransformGizmoSource->CanInteract(); 
+	}
+	return bValidMode && bVisible;
+}
+
 void UTransformGizmo::Tick(float DeltaTime)
 {
 	if (PendingDragFunction)
@@ -904,7 +906,8 @@ void UTransformGizmo::SetActiveTarget(UTransformProxy* Target, IToolContextTrans
 
 	if (InStateTarget)
 	{
-		StateTarget = Cast<UObject>(InStateTarget);
+		StateTarget.SetInterface(InStateTarget);
+		StateTarget.SetObject(CastChecked<UObject>(InStateTarget));
 	}
 	else
 	{
@@ -974,7 +977,9 @@ void UTransformGizmo::HandleWidgetModeChanged(UE::Widget::EWidgetMode InWidgetMo
 		if (DefaultHitPart != CurrentModeLastHitPart)
 		{
 			// reset indirect manipulation to default
-			UpdateInteractingState(false, CurrentModeLastHitPart, true);
+			ResetInteractingStates(CurrentMode);
+			ResetHoverStates(CurrentMode);
+			
 			SetModeLastHitPart(CurrentMode, DefaultHitPart);
 			UpdateInteractingState(true, DefaultHitPart, true);
 		}
@@ -1347,7 +1352,7 @@ bool UTransformGizmo::GetRayParamIntersectionWithInteractionPlane(const FInputDe
 	return true;
 }
 
-void UTransformGizmo::UpdateHoverState(bool bInHover, ETransformGizmoPartIdentifier InHitPartId)
+void UTransformGizmo::UpdateHoverState(const bool bInHover, const ETransformGizmoPartIdentifier InHitPartId)
 {
 	HitTarget->UpdateHoverState(bInHover, static_cast<uint32>(InHitPartId));
 
@@ -1370,10 +1375,42 @@ void UTransformGizmo::UpdateHoverState(bool bInHover, ETransformGizmoPartIdentif
 		HitTarget->UpdateHoverState(bInHover, static_cast<uint32>(ETransformGizmoPartIdentifier::ScaleXAxis));
 		HitTarget->UpdateHoverState(bInHover, static_cast<uint32>(ETransformGizmoPartIdentifier::ScaleZAxis));
 		break;
+	default:
+		break;
 	}
 }
 
-void UTransformGizmo::UpdateInteractingState(bool bInInteracting, ETransformGizmoPartIdentifier InHitPartId, const bool bIdOnly)
+void UTransformGizmo::ResetHoverStates(const EGizmoTransformMode InMode)
+{
+	ETransformGizmoPartIdentifier IdBegin = ETransformGizmoPartIdentifier::Default;
+	ETransformGizmoPartIdentifier IdEnd = ETransformGizmoPartIdentifier::Max;
+
+	switch (InMode)
+	{
+	case EGizmoTransformMode::Translate:
+		IdBegin = ETransformGizmoPartIdentifier::TranslateAll;
+		IdEnd = ETransformGizmoPartIdentifier::RotateAll;
+		break;
+	case EGizmoTransformMode::Rotate:
+		IdBegin = ETransformGizmoPartIdentifier::RotateAll;
+		IdEnd = ETransformGizmoPartIdentifier::ScaleAll;
+		break;
+	case EGizmoTransformMode::Scale:
+		IdBegin = ETransformGizmoPartIdentifier::ScaleAll;
+		IdEnd = ETransformGizmoPartIdentifier::Max;
+		break;
+	default:
+		break;
+	}
+
+	static constexpr bool bInHover = false;
+	for (uint32 Id = static_cast<uint32>(IdBegin); Id < static_cast<uint32>(IdEnd); ++Id)
+	{
+		UpdateHoverState(bInHover, static_cast<ETransformGizmoPartIdentifier>(Id));
+	}
+}
+
+void UTransformGizmo::UpdateInteractingState(const bool bInInteracting, const ETransformGizmoPartIdentifier InHitPartId, const bool bIdOnly)
 {
 	HitTarget->UpdateInteractingState(bInInteracting, static_cast<uint32>(InHitPartId));
 
@@ -1404,10 +1441,44 @@ void UTransformGizmo::UpdateInteractingState(bool bInInteracting, ETransformGizm
 			HitTarget->UpdateInteractingState(bInInteracting, static_cast<uint32>(ETransformGizmoPartIdentifier::ScaleXAxis));
 			HitTarget->UpdateInteractingState(bInInteracting, static_cast<uint32>(ETransformGizmoPartIdentifier::ScaleZAxis));
 			break;
+		default:
+			break;
 		}
 	}
 }
 
+void UTransformGizmo::ResetInteractingStates(const EGizmoTransformMode InMode)
+{
+	ETransformGizmoPartIdentifier IdBegin = ETransformGizmoPartIdentifier::Default;
+	ETransformGizmoPartIdentifier IdEnd = ETransformGizmoPartIdentifier::Max;
+	bool bIdOnly = true;
+
+	switch (InMode)
+	{
+	case EGizmoTransformMode::Translate:
+		IdBegin = ETransformGizmoPartIdentifier::TranslateAll;
+		IdEnd = ETransformGizmoPartIdentifier::RotateAll;
+		break;
+	case EGizmoTransformMode::Rotate:
+		IdBegin = ETransformGizmoPartIdentifier::RotateAll;
+		IdEnd = ETransformGizmoPartIdentifier::ScaleAll;
+		break;
+	case EGizmoTransformMode::Scale:
+		IdBegin = ETransformGizmoPartIdentifier::ScaleAll;
+		IdEnd = ETransformGizmoPartIdentifier::Max;
+		bIdOnly = false; 
+		break;
+	default:
+		break;
+	}
+
+	static constexpr bool bInInteracting = false;
+	for (uint32 Id = static_cast<uint32>(IdBegin); Id < static_cast<uint32>(IdEnd); ++Id)
+	{
+		UpdateInteractingState(bInInteracting, static_cast<ETransformGizmoPartIdentifier>(Id), bIdOnly);
+	}
+}
+
 void UTransformGizmo::BeginTransformEditSequence()
 {
 	if (ensure(StateTarget))
diff --git a/Engine/Source/Editor/Experimental/EditorInteractiveToolsFramework/Public/EditorGizmos/EditorTransformGizmoSource.h b/Engine/Source/Editor/Experimental/EditorInteractiveToolsFramework/Public/EditorGizmos/EditorTransformGizmoSource.h
--- a/Engine/Source/Editor/Experimental/EditorInteractiveToolsFramework/Public/EditorGizmos/EditorTransformGizmoSource.h
+++ b/Engine/Source/Editor/Experimental/EditorInteractiveToolsFramework/Public/EditorGizmos/EditorTransformGizmoSource.h
@@ -64,6 +64,12 @@ public:
 	 */
 	virtual bool GetVisible() const override;
 
+	/* 
+	 * Returns whether the gizmo can interact.
+	 * Note that this can be true even if the gizmo is hidden to support indirect manipulation in game mode.
+	 */
+	virtual bool CanInteract() const override;
+
 	/**
  	 * Get current scale type
 	 */
diff --git a/Engine/Source/Editor/Experimental/EditorInteractiveToolsFramework/Public/EditorGizmos/TransformGizmo.h b/Engine/Source/Editor/Experimental/EditorInteractiveToolsFramework/Public/EditorGizmos/TransformGizmo.h
--- a/Engine/Source/Editor/Experimental/EditorInteractiveToolsFramework/Public/EditorGizmos/TransformGizmo.h
+++ b/Engine/Source/Editor/Experimental/EditorInteractiveToolsFramework/Public/EditorGizmos/TransformGizmo.h
@@ -456,10 +456,16 @@ protected:
 	virtual bool GetRayParamIntersectionWithInteractionPlane(const FInputDeviceRay& InRay, float& OutHitParam);
 
 	/** Update hover state for given part id */
-	virtual void UpdateHoverState(bool bInHover, ETransformGizmoPartIdentifier InPartId);
+	virtual void UpdateHoverState(const bool bInHover, const ETransformGizmoPartIdentifier InPartId);
+
+	/** Reset all hover states related to the transform mode to false */
+	void ResetHoverStates(const EGizmoTransformMode InMode);
 
 	/** Update interacting state for given part id */
-	virtual void UpdateInteractingState(bool bInInteracting, ETransformGizmoPartIdentifier InPartId, const bool bIdOnly = false);
+	virtual void UpdateInteractingState(const bool bInInteracting, const ETransformGizmoPartIdentifier InPartId, const bool bIdOnly = false);
+
+	/** Reset all interacting states related to the transform mode to false */
+	void ResetInteractingStates(const EGizmoTransformMode InMode);
 
 	/** Called at the start of a sequence of gizmo transform edits */
 	virtual void BeginTransformEditSequence();
@@ -672,6 +678,12 @@ protected:
 	// Verify part identifier is within recognized range of transform gizmo part ids
 	virtual bool VerifyPartIdentifier(uint32 InPartIdentifier) const;
 
+	// Returns whether the gizmo is visible in the viewport.
+	bool IsVisible() const;
+
+	// Returns whether the gizmo can interact. Note that this can be true even if the gizmo is not visible to support indirect manipulation.
+	bool CanInteract() const;
+	
 protected:
 
 	/** Materials and colors to be used when drawing the items for each axis */
diff --git a/Engine/Source/Editor/Experimental/EditorInteractiveToolsFramework/Public/EditorGizmos/TransformGizmoInterfaces.h b/Engine/Source/Editor/Experimental/EditorInteractiveToolsFramework/Public/EditorGizmos/TransformGizmoInterfaces.h
--- a/Engine/Source/Editor/Experimental/EditorInteractiveToolsFramework/Public/EditorGizmos/TransformGizmoInterfaces.h
+++ b/Engine/Source/Editor/Experimental/EditorInteractiveToolsFramework/Public/EditorGizmos/TransformGizmoInterfaces.h
@@ -67,6 +67,11 @@ public:
 	 */
 	virtual bool GetVisible() const PURE_VIRTUAL(ITransformGizmoSource::GetVisible, return false;);
 
+	/* 
+	 * Returns whether the gizmo can interact.
+	 */
+	virtual bool CanInteract() const PURE_VIRTUAL(ITransformGizmoSource::CanInteract, return false;);
+
 	/*
 	 * Get current scale type
 	 */
