New TRS Gizmo: rotate modes

two rotation modes available via Gizmos.RotateMode:
    - 0: Pull, to rotate the object by pulling along the tangent to the circle (default)
    - 1: Arc, to rotate the object by rotating around the gizmo center (like the screen space rotation does)

diff --git a/Engine/Source/Editor/Experimental/EditorInteractiveToolsFramework/Private/EditorGizmos/TransformGizmo.cpp b/Engine/Source/Editor/Experimental/EditorInteractiveToolsFramework/Private/EditorGizmos/TransformGizmo.cpp
--- a/Engine/Source/Editor/Experimental/EditorInteractiveToolsFramework/Private/EditorGizmos/TransformGizmo.cpp
+++ b/Engine/Source/Editor/Experimental/EditorInteractiveToolsFramework/Private/EditorGizmos/TransformGizmo.cpp
@@ -51,7 +51,17 @@ static FAutoConsoleVariableRef CVarProjectIndirect(
 	ProjectIndirect,
 	TEXT("Project to the nearest point of the curve when handling indirect rotation.")
 	);
-	
+
+static int32 cvRotateMode = 0;
+static FAutoConsoleVariableRef CVarRotateMode(
+	TEXT("Gizmos.RotateMode"),
+	cvRotateMode,
+	TEXT("Rotation mode to choose between Pull = 0 and Arc = 1."),
+	FConsoleVariableDelegate::CreateLambda([](IConsoleVariable*)
+	{
+		cvRotateMode = FMath::Clamp(cvRotateMode, 0, 1);
+	})
+	);
 }
 
 void UTransformGizmo::SetDisallowNegativeScaling(bool bDisallow)
@@ -1979,40 +1989,99 @@ FVector UTransformGizmo::ComputeScaleDelta(const FVector2D& InStartPos, const FV
 	return Scale;
 }
 
-void UTransformGizmo::OnClickPressRotateXAxis(const FInputDeviceRay& InPressPos)
+bool UTransformGizmo::OnClickPressRotateArc(
+	const FInputDeviceRay& InPressPos, const FVector& InPlaneNormal,
+	const FVector& InPlaneAxis1, const FVector& InPlaneAxis2)
 {
-	InteractionScreenAxisDirection = GetScreenRotateAxisDir(InPressPos);
-	InteractionAxisList = EAxisList::X;
-	InteractionScreenStartPos = InteractionScreenCurrPos = InPressPos.ScreenPosition;
-	bInInteraction = true;
+	const FRay& Ray = InPressPos.WorldRay;
 
-	SetModeLastHitPart(EGizmoTransformMode::Rotate, LastHitPart);
+	// compute axis / view direction projection: is the rotation plane nearly perpendicular to the view plane?
+	const FVector WorldOrigin = CurrentTransform.GetLocation();
+	const FVector ViewDirection = GizmoViewContext->IsPerspectiveProjection() ?
+		(WorldOrigin - GizmoViewContext->ViewLocation).GetSafeNormal() :
+		GizmoViewContext->GetViewDirection();
+	const bool bAxisPerpendicularToView = FMath::Abs(FVector::DotProduct(InPlaneNormal, ViewDirection)) < GizmoLocals::DotThreshold;
+	const bool bRayPerpendicularToAxis = FMath::IsNearlyZero(FVector::DotProduct(InPlaneNormal, Ray.Direction));
+
+	// can we project
+	const bool bCanProject = bIndirectManipulation || bAxisPerpendicularToView || bRayPerpendicularToAxis;
+	if (!bCanProject)
+	{
+		InteractionPlanarOrigin = CurrentTransform.GetLocation();
+		InteractionPlanarNormal = InPlaneNormal;
+		InteractionPlanarAxisX = InPlaneAxis1;
+		InteractionPlanarAxisY = InPlaneAxis2;
 	
-	bDebugRotate = true;
+		float HitDepth;
+		if (GetRayParamIntersectionWithInteractionPlane(InPressPos, HitDepth))
+		{
+			const FVector HitPoint = InPressPos.WorldRay.Origin + InPressPos.WorldRay.Direction * HitDepth;
+			InteractionStartAngle = GizmoMath::ComputeAngleInPlane(HitPoint,
+				InteractionPlanarOrigin, InteractionPlanarNormal, InteractionPlanarAxisX, InteractionPlanarAxisY);
+			InteractionCurrAngle = InteractionStartAngle;
+			bInInteraction = true;
+			return true;
+		}
+	}
+	return false;
 }
 
-void UTransformGizmo::OnClickPressRotateYAxis(const FInputDeviceRay& InPressPos)
+void UTransformGizmo::OnClickPressRotateAxis(const FInputDeviceRay& InPressPos)
 {
+	static const TArray RotateIDs({	ETransformGizmoPartIdentifier::RotateXAxis,
+									ETransformGizmoPartIdentifier::RotateYAxis,
+									ETransformGizmoPartIdentifier::RotateZAxis});
+
+	const int32 RotateID = RotateIDs.IndexOfByKey(LastHitPart);
+	if (!ensure(RotateID != INDEX_NONE))
+	{
+		bInInteraction = false;
+		return;
+	}
+
+	static const TArray AxisList({EAxisList::X, EAxisList::Y, EAxisList::Z});
+
+	bDebugRotate = true;
+	
+	// initialize pull data
 	InteractionScreenAxisDirection = GetScreenRotateAxisDir(InPressPos);
-	InteractionAxisList = EAxisList::Y;
+	InteractionAxisList = AxisList[RotateID];
 	InteractionScreenStartPos = InteractionScreenCurrPos = InPressPos.ScreenPosition;
+	RotateMode = EAxisRotateMode::Pull;
+
+	// initialize arc/mixed data
+	const bool bRotatePull = bIndirectManipulation || GizmoLocals::cvRotateMode == 0;
+	if (!bRotatePull)
+	{
+		static const TArray RotateAxis({FVector::XAxisVector, FVector::YAxisVector, FVector::ZAxisVector});
+		const FVector WorldPlaneNormal = GetWorldAxis(RotateAxis[RotateID]);
+		const FVector WorldPlaneAxis1 = GetWorldAxis(RotateAxis[(RotateID+1) % 3]);
+		const FVector WorldPlaneAxis2 = GetWorldAxis(RotateAxis[(RotateID+2) % 3]);
+		
+		const bool bCanRotateArc = OnClickPressRotateArc(InPressPos, WorldPlaneNormal, WorldPlaneAxis1, WorldPlaneAxis2);
+		if (bCanRotateArc)
+		{
+			RotateMode = EAxisRotateMode::Arc;
+		}
+	}
+
 	bInInteraction = true;
-	
 	SetModeLastHitPart(EGizmoTransformMode::Rotate, LastHitPart);
-	
-	bDebugRotate = true;
 }
 
-void UTransformGizmo::OnClickPressRotateZAxis(const FInputDeviceRay& InPressPos)
+void UTransformGizmo::OnClickPressRotateXAxis(const FInputDeviceRay& InPressPos)
 {
-	InteractionScreenAxisDirection = GetScreenRotateAxisDir(InPressPos);
-	InteractionAxisList = EAxisList::Z;
-	InteractionScreenStartPos = InteractionScreenCurrPos = InPressPos.ScreenPosition;
-	bInInteraction = true;
+	OnClickPressRotateAxis(InPressPos);
+}
 
-	SetModeLastHitPart(EGizmoTransformMode::Rotate, LastHitPart);
+void UTransformGizmo::OnClickPressRotateYAxis(const FInputDeviceRay& InPressPos)
+{
+	OnClickPressRotateAxis(InPressPos);
+}
 
-	bDebugRotate = true;
+void UTransformGizmo::OnClickPressRotateZAxis(const FInputDeviceRay& InPressPos)
+{
+	OnClickPressRotateAxis(InPressPos);
 }
 
 FVector2D UTransformGizmo::GetScreenRotateAxisDir(const FInputDeviceRay& InPressPos)
@@ -2107,9 +2176,35 @@ FVector2D UTransformGizmo::GetScreenRotateAxisDir(const FInputDeviceRay& InPress
 
 void UTransformGizmo::OnClickDragRotateAxis(const FInputDeviceRay& DragPos)
 {
-	const FQuat DeltaRot = ComputeAxisRotateDelta(InteractionScreenCurrPos, DragPos.ScreenPosition);
-	ApplyRotateDelta(DeltaRot);
-	InteractionScreenCurrPos = DragPos.ScreenPosition;
+	switch (RotateMode)
+	{
+	case EAxisRotateMode::Pull:
+	{
+		const FQuat DeltaRot = ComputeAxisRotateDelta(InteractionScreenCurrPos, DragPos.ScreenPosition);
+		ApplyRotateDelta(DeltaRot);
+		InteractionScreenCurrPos = DragPos.ScreenPosition;
+		break;
+	}
+	case EAxisRotateMode::Arc:
+	{
+		float HitDepth;
+		if (GetRayParamIntersectionWithInteractionPlane(DragPos, HitDepth))
+		{
+			const FVector HitPoint = DragPos.WorldRay.Origin + DragPos.WorldRay.Direction * HitDepth;
+			const float HitAngle = GizmoMath::ComputeAngleInPlane(HitPoint,
+				InteractionPlanarOrigin, InteractionPlanarNormal, InteractionPlanarAxisX, InteractionPlanarAxisY);
+
+			const FQuat Delta = ComputeAngularRotateDelta(InteractionCurrAngle, HitAngle);
+	
+			ApplyRotateDelta(Delta);
+			InteractionCurrAngle = HitAngle;
+		}
+		break;
+	}
+	default:
+		ensure(false);
+		break;
+	}
 }
 
 FQuat UTransformGizmo::ComputeAxisRotateDelta(const FVector2D& InStartPos, const FVector2D& InEndPos)
@@ -2118,8 +2213,8 @@ FQuat UTransformGizmo::ComputeAxisRotateDelta(const FVector2D& InStartPos, const
 
 	const FVector2D DragDirToRemove = NormalProjectionToRemove * FVector2D::DotProduct(DragDir, NormalProjectionToRemove);
 	DragDir -= DragDirToRemove;
-	
-	const double Delta = FVector2D::DotProduct(InteractionScreenAxisDirection, DragDir);
+
+	const double Delta = FVector2D::DotProduct(InteractionScreenAxisDirection, DragDir) * 0.25;
 	FRotator DeltaRot(
 		InteractionAxisList == EAxisList::Y ? Delta : 0.0,
 		InteractionAxisList == EAxisList::Z ? Delta : 0.0,
diff --git a/Engine/Source/Editor/Experimental/EditorInteractiveToolsFramework/Public/EditorGizmos/TransformGizmo.h b/Engine/Source/Editor/Experimental/EditorInteractiveToolsFramework/Public/EditorGizmos/TransformGizmo.h
--- a/Engine/Source/Editor/Experimental/EditorInteractiveToolsFramework/Public/EditorGizmos/TransformGizmo.h
+++ b/Engine/Source/Editor/Experimental/EditorInteractiveToolsFramework/Public/EditorGizmos/TransformGizmo.h
@@ -93,6 +93,16 @@ enum class ETransformGizmoPartIdentifier
 	Max
 };
 
+namespace EAxisRotateMode
+{
+	enum Type
+	{
+		Pull,
+		Arc
+	};
+}
+
+
 /**
  * UTransformGizmo provides standard Transformation Gizmo interactions,
  * applied to a UTransformProxy target object. By default the Gizmo will be
@@ -549,6 +559,9 @@ protected:
 	/** Handle click press for rotate Z axis */
 	virtual void OnClickPressRotateZAxis(const FInputDeviceRay& InPressPos);
 
+	/** Handle click press for any rotate axis */
+	void OnClickPressRotateAxis(const FInputDeviceRay& InPressPos);
+
 	/** Handle click drag for rotate axis */
 	virtual void OnClickDragRotateAxis(const FInputDeviceRay& InDragPos);
 
@@ -561,6 +574,10 @@ protected:
 	/** Compute rotate delta based on screen-space start/end positions */
 	virtual FQuat ComputeAxisRotateDelta(const FVector2D& InStartPos, const FVector2D& InEndPos);
 
+	/** Prepares data for arc rotation. This will return false if this is not possible (the rotate handle is perpendicular to the view) */
+	bool OnClickPressRotateArc( const FInputDeviceRay& InPressPos,
+		const FVector& InPlaneNormal, const FVector& InPlaneAxis1, const FVector& InPlaneAxis2);
+
 	/**
 	 * Screen-space rotate interaction methods
 	 */
@@ -595,7 +612,7 @@ protected:
 
 	/** */
 	float GetSizeCoefficient() const;
-
+	
 	/**
 	* Scale click-drag handling methods
 	*/
@@ -855,6 +872,8 @@ protected:
 	UPROPERTY()
 	bool bCtrlMiddleDoesY = true;
 
+	TEnumAsByte<EAxisRotateMode::Type> RotateMode = EAxisRotateMode::Pull;
+	
 private:
 	/** Debug attributes to display the pull direction */
 	bool bDebugRotate = false;
