New TRS Gizmo: change the default editor mode for the asset editors

    The new TRS gizmo system uses FEditorModeTools::GetShowWidget() to update the visible state of the gizmo.
    However, for asset editors that don't override their default mode (as is the case with BP), the result is obtained from the FLevelEditorSelectModeWidgetHelper::ShouldDrawWidget function, which gets its data from the GCurrentLevelEditingViewportClient...
    The FWidget rendering works mainly because it bypasses the result of FEditorModeTools::GetShowWidget() and relies on the value of FEditorViewportClient::GetWidgetMode() instead. VPCs like FSCSEditorViewportClient, for example, will change the widget mode dynamically according to the selection.
    To change this behavior for new TRS gizmo without changing the current behavior, we swap the default mode with AssetEdModes::EM_AssetDefault if (and only if) FBuiltinEditorModes::EM_Default is one of the default modes and we're not dealing with the level editor's mode manager.

+ remove the direct uses of "Gizmos.UseLegacyWidget" in several places to ease the migration to user settings

diff --git a/Engine/Plugins/Animation/ControlRig/Source/ControlRigEditor/Private/EditMode/ControlRigEditMode.cpp b/Engine/Plugins/Animation/ControlRig/Source/ControlRigEditor/Private/EditMode/ControlRigEditMode.cpp
--- a/Engine/Plugins/Animation/ControlRig/Source/ControlRigEditor/Private/EditMode/ControlRigEditMode.cpp
+++ b/Engine/Plugins/Animation/ControlRig/Source/ControlRigEditor/Private/EditMode/ControlRigEditMode.cpp
@@ -71,6 +71,7 @@
 #include "DragTool_BoxSelect.h"
 #include "DragTool_FrustumSelect.h"
 #include "AnimationEditorViewportClient.h"
+#include "EditorInteractiveGizmoManager.h"
 
 #include UE_INLINE_GENERATED_CPP_BY_NAME(ControlRigEditMode)
 
@@ -1978,8 +1979,7 @@ bool FControlRigEditMode::InputDelta(FEditorViewportClient* InViewportClient, FV
 	//button down if left and ctrl and right is down, needed for indirect posting
 
 	// enable MMB with the new TRS gizmos
-	static IConsoleVariable* const UseLegacyWidgetCVar = IConsoleManager::Get().FindConsoleVariable(TEXT("Gizmos.UseLegacyWidget"));
-	const bool bEnableMMB = UseLegacyWidgetCVar ? !UseLegacyWidgetCVar->GetBool() : false;
+	const bool bEnableMMB = UEditorInteractiveGizmoManager::UsesNewTRSGizmos();
 	
 	const bool bMouseButtonDown =
 		InViewport->KeyState(EKeys::LeftMouseButton) ||
diff --git a/Engine/Source/Editor/Experimental/EditorInteractiveToolsFramework/Private/EditorGizmos/EditorTransformGizmoSource.cpp b/Engine/Source/Editor/Experimental/EditorInteractiveToolsFramework/Private/EditorGizmos/EditorTransformGizmoSource.cpp
--- a/Engine/Source/Editor/Experimental/EditorInteractiveToolsFramework/Private/EditorGizmos/EditorTransformGizmoSource.cpp
+++ b/Engine/Source/Editor/Experimental/EditorInteractiveToolsFramework/Private/EditorGizmos/EditorTransformGizmoSource.cpp
@@ -3,6 +3,7 @@
 #include "EditorGizmos/EditorTransformGizmoSource.h"
 
 #include "Editor.h"
+#include "EditorInteractiveGizmoManager.h"
 #include "EditorModeManager.h"
 #include "EditorViewportClient.h"
 #include "EditorGizmos/EditorTransformGizmoUtil.h"
@@ -88,11 +89,7 @@ bool UEditorTransformGizmoSource::CanInteract() const
 			bool bUseLegacyWidget = (WidgetMode == UE::Widget::WM_TranslateRotateZ || WidgetMode == UE::Widget::WM_2D);
 			if (!bUseLegacyWidget)
 			{
-				static IConsoleVariable* const UseLegacyWidgetCVar = IConsoleManager::Get().FindConsoleVariable(TEXT("Gizmos.UseLegacyWidget"));
-				if (ensure(UseLegacyWidgetCVar))
-				{
-					bUseLegacyWidget = UseLegacyWidgetCVar->GetInt() > 0;
-				}
+				bUseLegacyWidget = !UEditorInteractiveGizmoManager::UsesNewTRSGizmos();
 			}
 			return !bUseLegacyWidget;
 		}
diff --git a/Engine/Source/Editor/Experimental/EditorInteractiveToolsFramework/Private/EditorGizmos/EditorTransformGizmoUtil.cpp b/Engine/Source/Editor/Experimental/EditorInteractiveToolsFramework/Private/EditorGizmos/EditorTransformGizmoUtil.cpp
--- a/Engine/Source/Editor/Experimental/EditorInteractiveToolsFramework/Private/EditorGizmos/EditorTransformGizmoUtil.cpp
+++ b/Engine/Source/Editor/Experimental/EditorInteractiveToolsFramework/Private/EditorGizmos/EditorTransformGizmoUtil.cpp
@@ -12,6 +12,8 @@
 #include "EditorGizmos/EditorTransformGizmoBuilder.h"
 #include "EditorGizmos/EditorTransformGizmoDataBinder.h"
 #include "EditorViewportClient.h"
+#include "EditorModes.h"
+#include "Tools/DefaultEdMode.h"
 
 #include UE_INLINE_GENERATED_CPP_BY_NAME(EditorTransformGizmoUtil)
 
@@ -253,6 +255,9 @@ void UEditorTransformGizmoContextObject::UpdateGizmo(const TArray<FEditorViewpor
 	
 	const UModeManagerInteractiveToolsContext* ToolsContext = ModeTools->GetInteractiveToolsContext();
 	const TObjectPtr<UInteractiveToolManager> ToolManager = ToolsContext->ToolManager;
+
+	// swap default mode
+	const bool Swapped = SwapDefaultMode(FBuiltinEditorModes::EM_Default, FAssetEdModes::EM_AssetDefault);
 	
 	if (UE::EditorTransformGizmoUtil::FindDefaultTransformGizmo(ToolManager))
 	{
@@ -262,7 +267,7 @@ void UEditorTransformGizmoContextObject::UpdateGizmo(const TArray<FEditorViewpor
 	// create a new one
 	if (UTransformGizmo* TransformGizmo = UE::EditorTransformGizmoUtil::GetDefaultTransformGizmo(ToolManager))
 	{
-		TransformGizmo->SetVisibility(false);
+		TransformGizmo->SetVisibility(Swapped ? UEditorInteractiveGizmoManager::UsesNewTRSGizmos() : false);
 	}
 }
 
@@ -274,11 +279,12 @@ void UEditorTransformGizmoContextObject::InitializeCVarBinding()
 		auto OnGizmoVariableChanged = [this](IConsoleVariable* InLegacyWidgetCVar)
 		{
 			const bool bUseLegacyGizmo = InLegacyWidgetCVar ? UseLegacyWidgetCVar->GetInt() > 0 : false;
-
 			if (bUseLegacyGizmo)
 			{
-				// remove viewports' binding + gizmos as they are useless
+				// swap back default mode
+				(void)SwapDefaultMode(FAssetEdModes::EM_AssetDefault, FBuiltinEditorModes::EM_Default);
 				
+				// remove viewports' binding + gizmos as they are useless
 				const UModeManagerInteractiveToolsContext* ToolsContext = ModeTools->GetInteractiveToolsContext();
 				const TObjectPtr<UInteractiveToolManager> ToolManager = ToolsContext->ToolManager;
 				UE::EditorTransformGizmoUtil::RemoveDefaultTransformGizmo(ToolManager);
@@ -353,3 +359,31 @@ void UEditorTransformGizmoContextObject::RemoveViewportsBinding()
 		ViewportClientsChangedHandle.Reset();
 	}
 }
+
+bool UEditorTransformGizmoContextObject::SwapDefaultMode(const FEditorModeID InCurrentDefaultMode, const FEditorModeID InNewDefaultMode) const
+{
+	// swap only if we're not dealing with the level editor + InCurrentDefaultMode is one of the default modes
+	if (!ModeTools->IsDefaultMode(InCurrentDefaultMode))
+	{
+		return false;
+	}
+	
+	const bool bIsLevelMode = ModeTools == &GLevelEditorModeTools();
+	if (bIsLevelMode)
+	{
+		return false;
+	}
+	
+	// replace default mode
+	ModeTools->RemoveDefaultMode(InCurrentDefaultMode);
+	ModeTools->AddDefaultMode(InNewDefaultMode);
+	
+	if (ModeTools->IsModeActive(InCurrentDefaultMode))
+	{
+		// activate it if the previous one was 
+		ModeTools->DeactivateMode(InCurrentDefaultMode);
+		ModeTools->ActivateMode(InNewDefaultMode);
+	}
+
+	return true;
+}
diff --git a/Engine/Source/Editor/Experimental/EditorInteractiveToolsFramework/Private/EditorInteractiveGizmoManager.cpp b/Engine/Source/Editor/Experimental/EditorInteractiveToolsFramework/Private/EditorInteractiveGizmoManager.cpp
--- a/Engine/Source/Editor/Experimental/EditorInteractiveToolsFramework/Private/EditorInteractiveGizmoManager.cpp
+++ b/Engine/Source/Editor/Experimental/EditorInteractiveToolsFramework/Private/EditorInteractiveGizmoManager.cpp
@@ -27,20 +27,33 @@ class FCanvas;
 
 #define LOCTEXT_NAMESPACE "UEditorInteractiveGizmoManager"
 
-static TAutoConsoleVariable<int32> CVarUseLegacyWidget(
+namespace GizmoManagerLocals
+{
+	static int32 UseLegacyWidget = 1;
+}
+
+static FAutoConsoleVariableRef CVarUseLegacyWidget(
 	TEXT("Gizmos.UseLegacyWidget"),
-	1,
+	GizmoManagerLocals::UseLegacyWidget,
 	TEXT("Specify whether to use selection-based gizmos or legacy widget\n")
 	TEXT("0 = enable UE5 transform and other selection-based gizmos.\n")
 	TEXT("1 = enable legacy UE4 transform widget."),
+	FConsoleVariableDelegate::CreateLambda([](IConsoleVariable*)
+	{
+		GizmoManagerLocals::UseLegacyWidget = FMath::Clamp(GizmoManagerLocals::UseLegacyWidget, 0, 1);
+	}),
 	ECVF_RenderThreadSafe);
 
+bool UEditorInteractiveGizmoManager::UsesNewTRSGizmos()
+{
+	return (GizmoManagerLocals::UseLegacyWidget == 0);
+}
 
 UEditorInteractiveGizmoManager::UEditorInteractiveGizmoManager() :
 	UInteractiveGizmoManager()
 {
 	Registry = NewObject<UEditorInteractiveGizmoRegistry>();
-	bShowEditorGizmos = (CVarUseLegacyWidget.GetValueOnGameThread() == 0);
+	bShowEditorGizmos = UsesNewTRSGizmos();
 }
 
 
@@ -175,8 +188,23 @@ bool UEditorInteractiveGizmoManager::GetShowEditorGizmosForView(IToolsContextRen
 
 void UEditorInteractiveGizmoManager::UpdateActiveEditorGizmos()
 {
+	const bool bEnableEditorGizmos = UsesNewTRSGizmos();
+	if (!bEnableEditorGizmos)
+	{
+		if (bShowEditorGizmos)
+		{
+			if (UTransformGizmo* Gizmo = FindDefaultTransformGizmo())
+			{
+				DestroyGizmo(Gizmo);
+			}
+			DestroyAllEditorGizmos();
+		}
+		
+		bShowEditorGizmos = false;
+		return;
+	}
+	
 	const bool bEditorModeToolsSupportsWidgetDrawing = EditorModeManager ? EditorModeManager->GetShowWidget() : true;
-	const bool bEnableEditorGizmos = (CVarUseLegacyWidget.GetValueOnGameThread() == 0);
 	const bool bNewShowEditorGizmos = bEditorModeToolsSupportsWidgetDrawing && bEnableEditorGizmos;
 
 	if (bShowEditorGizmos != bNewShowEditorGizmos)
diff --git a/Engine/Source/Editor/Experimental/EditorInteractiveToolsFramework/Public/EditorGizmos/EditorTransformGizmoUtil.h b/Engine/Source/Editor/Experimental/EditorInteractiveToolsFramework/Public/EditorGizmos/EditorTransformGizmoUtil.h
--- a/Engine/Source/Editor/Experimental/EditorInteractiveToolsFramework/Public/EditorGizmos/EditorTransformGizmoUtil.h
+++ b/Engine/Source/Editor/Experimental/EditorInteractiveToolsFramework/Public/EditorGizmos/EditorTransformGizmoUtil.h
@@ -101,6 +101,9 @@ private:
 	EDITORINTERACTIVETOOLSFRAMEWORK_API void RemoveCVarBinding();
 	EDITORINTERACTIVETOOLSFRAMEWORK_API void RemoveViewportsBinding();
 
+	typedef FName FEditorModeID;
+	bool SwapDefaultMode(const FEditorModeID InCurrentDefaultMode, const FEditorModeID InNewDefaultMode) const;
+
 	FOnGizmoCreated OnGizmoCreated;
 	
 	FEditorModeTools* ModeTools = nullptr;
diff --git a/Engine/Source/Editor/Experimental/EditorInteractiveToolsFramework/Public/EditorInteractiveGizmoManager.h b/Engine/Source/Editor/Experimental/EditorInteractiveToolsFramework/Public/EditorInteractiveGizmoManager.h
--- a/Engine/Source/Editor/Experimental/EditorInteractiveToolsFramework/Public/EditorInteractiveGizmoManager.h
+++ b/Engine/Source/Editor/Experimental/EditorInteractiveToolsFramework/Public/EditorInteractiveGizmoManager.h
@@ -129,6 +129,11 @@ public:
 	/** instance/builder identifiers for transform gizmo */
 	static const FString& TransformInstanceIdentifier();
 	static const FString& TransformBuilderIdentifier();
+
+	/**
+	 * Returns true if the new TRS gizmos are used.
+	 */
+	static bool UsesNewTRSGizmos();
 	
 protected:
 
diff --git a/Engine/Source/Editor/UnrealEd/Private/EditorModeManager.cpp b/Engine/Source/Editor/UnrealEd/Private/EditorModeManager.cpp
--- a/Engine/Source/Editor/UnrealEd/Private/EditorModeManager.cpp
+++ b/Engine/Source/Editor/UnrealEd/Private/EditorModeManager.cpp
@@ -36,6 +36,7 @@
 #include "EngineUtils.h"
 #include "Tools/AssetEditorContextObject.h"
 #include "ContextObjectStore.h"
+#include "EditorInteractiveGizmoManager.h"
 #include "UObject/GCObjectScopeGuard.h"
 #include "Settings/LevelEditorViewportSettings.h"
 #include "Subsystems/EditorElementSubsystem.h"
@@ -133,8 +134,7 @@ void FEditorModeTools::LoadConfig(void)
 	if (static_cast<ECoordSystem>(CoordSystemAsInt) == COORD_Parent)
 	{
 		// parent mode is only supported with new trs gizmos for now
-		static IConsoleVariable* UseLegacyWidgetCVar = IConsoleManager::Get().FindConsoleVariable(TEXT("Gizmos.UseLegacyWidget"));
-		const bool bUseNewGizmo = UseLegacyWidgetCVar ? UseLegacyWidgetCVar->GetInt() < 1 : false;
+		const bool bUseNewGizmo = UEditorInteractiveGizmoManager::UsesNewTRSGizmos();
 		if (!bUseNewGizmo)
 		{
 			CoordSystemAsInt = static_cast<int32>(COORD_Local);
diff --git a/Engine/Source/Editor/UnrealEd/Private/LevelEditorViewport.cpp b/Engine/Source/Editor/UnrealEd/Private/LevelEditorViewport.cpp
--- a/Engine/Source/Editor/UnrealEd/Private/LevelEditorViewport.cpp
+++ b/Engine/Source/Editor/UnrealEd/Private/LevelEditorViewport.cpp
@@ -89,6 +89,7 @@
 #include "IHeadMountedDisplay.h"
 #include "IXRTrackingSystem.h"
 #include "ActorGroupingUtils.h"
+#include "EditorInteractiveGizmoManager.h"
 #include "EditorWorldExtension.h"
 #include "VREditorMode.h"
 #include "EditorWorldExtension.h"
@@ -3278,8 +3279,7 @@ bool FLevelEditorViewportClient::InputWidgetDelta(FViewport* InViewport, EAxisLi
 
 				// We do not want actors updated if we are holding down the middle mouse button.
 				// enable MMB for New TRS Gizmos
-				static IConsoleVariable* const UseLegacyWidgetCVar = IConsoleManager::Get().FindConsoleVariable(TEXT("Gizmos.UseLegacyWidget"));
-				const bool bEnableMMB = UseLegacyWidgetCVar ? !UseLegacyWidgetCVar->GetBool() && !bDraggingByHandle : false;
+				const bool bEnableMMB = UEditorInteractiveGizmoManager::UsesNewTRSGizmos() ? !bDraggingByHandle : false;
 				
 				if(!MiddleMouseButtonDown || bEnableMMB)
 				{
diff --git a/Engine/Source/Editor/UnrealEd/Private/SEditorViewport.cpp b/Engine/Source/Editor/UnrealEd/Private/SEditorViewport.cpp
--- a/Engine/Source/Editor/UnrealEd/Private/SEditorViewport.cpp
+++ b/Engine/Source/Editor/UnrealEd/Private/SEditorViewport.cpp
@@ -1,6 +1,8 @@
 // Copyright Epic Games, Inc. All Rights Reserved.
 
 #include "SEditorViewport.h"
+
+#include "EditorInteractiveGizmoManager.h"
 #include "Misc/Paths.h"
 #include "Framework/Commands/UICommandList.h"
 #include "Misc/App.h"
@@ -765,10 +767,8 @@ void SEditorViewport::OnCycleCoordinateSystem()
 
 	++CoordSystemAsInt;
 
-	static IConsoleVariable* UseLegacyWidgetCVar = IConsoleManager::Get().FindConsoleVariable(TEXT("Gizmos.UseLegacyWidget"));
-	const bool bUseNewGizmos = UseLegacyWidgetCVar ? UseLegacyWidgetCVar->GetInt() < 1 : false;
 	// parent mode is only supported with new trs gizmos for now
-	const int CoordMax = bUseNewGizmos ? COORD_Max : COORD_Parent;
+	const int CoordMax = UEditorInteractiveGizmoManager::UsesNewTRSGizmos() ? COORD_Max : COORD_Parent;
 	if( CoordSystemAsInt >= CoordMax )
 	{
 		CoordSystemAsInt = COORD_World;
diff --git a/Engine/Source/Editor/UnrealEd/Private/STransformViewportToolbar.cpp b/Engine/Source/Editor/UnrealEd/Private/STransformViewportToolbar.cpp
--- a/Engine/Source/Editor/UnrealEd/Private/STransformViewportToolbar.cpp
+++ b/Engine/Source/Editor/UnrealEd/Private/STransformViewportToolbar.cpp
@@ -2,6 +2,8 @@
 
 
 #include "STransformViewportToolbar.h"
+
+#include "EditorInteractiveGizmoManager.h"
 #include "EngineDefines.h"
 #include "Modules/ModuleManager.h"
 #include "Widgets/SBoxPanel.h"
@@ -211,9 +213,7 @@ TSharedRef< SWidget > STransformViewportToolBar::MakeTransformToolBar( const TSh
 
 		TAttribute<FText> CoordSystemToolTip = TAttribute<FText>::CreateLambda([]
 		{
-			static IConsoleVariable* UseLegacyWidgetCVar = IConsoleManager::Get().FindConsoleVariable(TEXT("Gizmos.UseLegacyWidget"));
-			const bool bUseNewGizmo = UseLegacyWidgetCVar ? UseLegacyWidgetCVar->GetInt() < 1 : false;
-			if (bUseNewGizmo)
+			if (UEditorInteractiveGizmoManager::UsesNewTRSGizmos())
 			{
 				return LOCTEXT(	"CycleTransformGizmoCoordSystemWithParent_ToolTip",
 								"Cycles the transform gizmo coordinate systems between world, local and parent space");
diff --git a/Engine/Source/Editor/UnrealEd/Private/Tools/DefaultEdMode.cpp b/Engine/Source/Editor/UnrealEd/Private/Tools/DefaultEdMode.cpp
--- a/Engine/Source/Editor/UnrealEd/Private/Tools/DefaultEdMode.cpp
+++ b/Engine/Source/Editor/UnrealEd/Private/Tools/DefaultEdMode.cpp
@@ -44,3 +44,31 @@ TSharedRef<FLegacyEdModeWidgetHelper> UEdModeDefault::CreateWidgetHelper()
 {
 	return MakeShared<FLevelEditorSelectModeWidgetHelper>();
 }
+
+namespace FAssetEdModes
+{
+	const FEditorModeID EM_AssetDefault(TEXT("EM_AssetDefault"));
+}
+
+class FAssetEdModeWidgetHelper : public FLegacyEdModeWidgetHelper
+{
+public:
+	virtual bool ShouldDrawWidget() const override
+	{
+		return true;
+	}
+};
+
+UAssetEdModeDefault::UAssetEdModeDefault()
+{
+	Info = FEditorModeInfo(
+		FAssetEdModes::EM_AssetDefault,
+		NSLOCTEXT("AssetDefaultMode", "DisplayName", "AssetSelection"),
+		FSlateIcon("EditorStyle", "LevelEditor.SelectMode", "LevelEditor.SelectMode.Small"),
+		true, 0);
+}
+
+TSharedRef<FLegacyEdModeWidgetHelper> UAssetEdModeDefault::CreateWidgetHelper()
+{
+	return MakeShared<FAssetEdModeWidgetHelper>();
+}
\ No newline at end of file
diff --git a/Engine/Source/Editor/UnrealEd/Private/UnrealWidgetRender.cpp b/Engine/Source/Editor/UnrealEd/Private/UnrealWidgetRender.cpp
--- a/Engine/Source/Editor/UnrealEd/Private/UnrealWidgetRender.cpp
+++ b/Engine/Source/Editor/UnrealEd/Private/UnrealWidgetRender.cpp
@@ -13,6 +13,7 @@
 #include "SnappingUtils.h"
 #include "DynamicMeshBuilder.h"
 #include "CanvasTypes.h"
+#include "EditorInteractiveGizmoManager.h"
 
 /*
  *  Simple struct used to create and group data related to the current window's / viewport's space,
@@ -187,11 +188,7 @@ void FWidget::Render( const FSceneView* View,FPrimitiveDrawInterface* PDI, FEdit
 	bool bUseLegacyWidget = (WidgetMode == UE::Widget::EWidgetMode::WM_TranslateRotateZ || WidgetMode == UE::Widget::EWidgetMode::WM_2D);
 	if (!bUseLegacyWidget)
 	{
-		static IConsoleVariable* const UseLegacyWidgetCVar = IConsoleManager::Get().FindConsoleVariable(TEXT("Gizmos.UseLegacyWidget"));
-		if (ensure(UseLegacyWidgetCVar))
-		{
-			bUseLegacyWidget = UseLegacyWidgetCVar->GetInt() > 0;
-		}
+		bUseLegacyWidget = !UEditorInteractiveGizmoManager::UsesNewTRSGizmos();
 	}
 
 	bool bDrawWidget;
diff --git a/Engine/Source/Editor/UnrealEd/Public/Tools/DefaultEdMode.h b/Engine/Source/Editor/UnrealEd/Public/Tools/DefaultEdMode.h
--- a/Engine/Source/Editor/UnrealEd/Public/Tools/DefaultEdMode.h
+++ b/Engine/Source/Editor/UnrealEd/Public/Tools/DefaultEdMode.h
@@ -20,3 +20,20 @@ public:
 protected:
 	UNREALED_API virtual TSharedRef<FLegacyEdModeWidgetHelper> CreateWidgetHelper() override;
 };
+
+namespace FAssetEdModes
+{
+	UNREALED_API extern const FEditorModeID EM_AssetDefault;
+}
+
+UCLASS(Transient, MinimalAPI)
+class UAssetEdModeDefault : public UEdModeDefault
+{
+	GENERATED_BODY()
+
+public:
+	UNREALED_API UAssetEdModeDefault();
+
+protected:
+	UNREALED_API virtual TSharedRef<FLegacyEdModeWidgetHelper> CreateWidgetHelper() override;
+};
\ No newline at end of file
