New TRS Gizmos: expose some settings in the editor preferences
    - need to enable the experimental Gizmo Framework plugin to access those settings
    - removed deprecated NewGizmoEnable module

diff --git a/Engine/Plugins/Experimental/Animation/SkeletalMeshModelingTools/Source/SkeletalMeshModelingTools/Private/SkeletalMeshGizmoUtils.cpp b/Engine/Plugins/Experimental/Animation/SkeletalMeshModelingTools/Source/SkeletalMeshModelingTools/Private/SkeletalMeshGizmoUtils.cpp
--- a/Engine/Plugins/Experimental/Animation/SkeletalMeshModelingTools/Source/SkeletalMeshModelingTools/Private/SkeletalMeshGizmoUtils.cpp
+++ b/Engine/Plugins/Experimental/Animation/SkeletalMeshModelingTools/Source/SkeletalMeshModelingTools/Private/SkeletalMeshGizmoUtils.cpp
@@ -70,12 +70,15 @@ UTransformGizmo* UE::SkeletalMeshGizmoUtils::CreateTransformGizmo(UInteractiveTo
 	}
 
 	static const FString InstanceIdentifier;
-	UInteractiveGizmo* NewGizmo = GizmoManager->CreateGizmo(
-		USkeletalMeshGizmoContextObject::TransformBuilderIdentifier(), InstanceIdentifier, Owner);
+	UTransformGizmo* NewGizmo = Cast<UTransformGizmo>( GizmoManager->CreateGizmo(
+		USkeletalMeshGizmoContextObject::TransformBuilderIdentifier(), InstanceIdentifier, Owner) );
 
-	ensure(NewGizmo);
+	if (ensure(NewGizmo))
+	{
+		GizmoManager->OnGizmosParametersChangedDelegate().AddUObject(NewGizmo, &UTransformGizmo::OnParametersChanged);
+	}
 	
-	return Cast<UTransformGizmo>(NewGizmo);
+	return NewGizmo;
 }
 
 const FString& USkeletalMeshGizmoContextObject::TransformBuilderIdentifier()
diff --git a/Engine/Plugins/Experimental/GizmoFramework/GizmoFramework.uplugin b/Engine/Plugins/Experimental/GizmoFramework/GizmoFramework.uplugin
--- a/Engine/Plugins/Experimental/GizmoFramework/GizmoFramework.uplugin
+++ b/Engine/Plugins/Experimental/GizmoFramework/GizmoFramework.uplugin
@@ -3,7 +3,7 @@
 	"Version": 1,
 	"VersionName": "1.0",
 	"FriendlyName": "New TRS Gizmos",
-	"Description": "Toggle use of InteractiveToolsFramework based TRS Gizmos",
+	"Description": "Toggle use of the new TRS Gizmos",
 	"Category": "Experimental",
 	"CreatedBy": "Epic Games, Inc.",
 	"CreatedByURL": "https://epicgames.com",
@@ -16,15 +16,9 @@
 	"Installed": false,
 	"Modules": [
 		{
-			"Name": "NewGizmoEnable",
+			"Name": "GizmoSettings",
 			"Type": "Editor",
 			"LoadingPhase": "Default"
 		}
-	],
-	"Plugins": [
-		{
-			"Name": "GizmoEdMode",
-			"Enabled": true
-		}
 	]
 }
\ No newline at end of file
diff --git a/Engine/Plugins/Experimental/GizmoFramework/Source/GizmoSettings/GizmoSettings.Build.cs b/Engine/Plugins/Experimental/GizmoFramework/Source/GizmoSettings/GizmoSettings.Build.cs
new file mode 100644
--- /dev/null
+++ b/Engine/Plugins/Experimental/GizmoFramework/Source/GizmoSettings/GizmoSettings.Build.cs
@@ -0,0 +1,21 @@
+// Copyright Epic Games, Inc. All Rights Reserved.
+
+namespace UnrealBuildTool.Rules
+{
+	public class GizmoSettings : ModuleRules
+	{
+		public GizmoSettings(ReadOnlyTargetRules Target) : base(Target)
+		{
+			PrivateDependencyModuleNames.AddRange(
+				new string[]
+				{
+					"Core",
+					"CoreUObject",
+					"DeveloperSettings",
+					"Engine",
+					"EditorInteractiveToolsFramework"
+				}
+			);
+		}
+	}
+}
diff --git a/Engine/Plugins/Experimental/GizmoFramework/Source/GizmoSettings/Private/GizmoSettings.cpp b/Engine/Plugins/Experimental/GizmoFramework/Source/GizmoSettings/Private/GizmoSettings.cpp
new file mode 100644
--- /dev/null
+++ b/Engine/Plugins/Experimental/GizmoFramework/Source/GizmoSettings/Private/GizmoSettings.cpp
@@ -0,0 +1,15 @@
+// Copyright Epic Games, Inc. All Rights Reserved.
+
+#include "GizmoSettings.h"
+
+UGizmoSettings::UGizmoSettings(const FObjectInitializer& ObjectInitlaizer)
+	: UDeveloperSettings(ObjectInitlaizer)
+{}
+
+UGizmoSettings::FOnUpdateSettings UGizmoSettings::OnSettingsChange;
+
+void UGizmoSettings::PostEditChangeProperty(FPropertyChangedEvent& PropertyChangedEvent)
+{
+	Super::PostEditChangeProperty(PropertyChangedEvent);
+	OnSettingsChange.Broadcast(this);
+}
diff --git a/Engine/Plugins/Experimental/GizmoFramework/Source/GizmoSettings/Private/GizmoSettingsModule.cpp b/Engine/Plugins/Experimental/GizmoFramework/Source/GizmoSettings/Private/GizmoSettingsModule.cpp
new file mode 100644
--- /dev/null
+++ b/Engine/Plugins/Experimental/GizmoFramework/Source/GizmoSettings/Private/GizmoSettingsModule.cpp
@@ -0,0 +1,30 @@
+// Copyright Epic Games, Inc. All Rights Reserved.
+
+#include "GizmoSettingsModule.h"
+#include "Modules/ModuleManager.h"
+
+#include "GizmoSettings.h"
+#include "EditorInteractiveGizmoManager.h"
+
+void FGizmoSettingsModule::StartupModule()
+{
+	UGizmoSettings::OnSettingsChange.AddRaw(this, &FGizmoSettingsModule::UpdateGizmos);
+	UpdateGizmos(GetDefault<UGizmoSettings>());
+}
+
+void FGizmoSettingsModule::ShutdownModule()
+{
+	UGizmoSettings::OnSettingsChange.RemoveAll(this);
+}
+
+void FGizmoSettingsModule::UpdateGizmos(const UGizmoSettings* InSettings)
+{
+	if (InSettings)
+	{
+		UEditorInteractiveGizmoManager::SetUsesNewTRSGizmos(InSettings->bEnableNewGizmos);
+		UEditorInteractiveGizmoManager::SetGizmosParameters(InSettings->GizmoParameters);
+	}
+}
+
+
+IMPLEMENT_MODULE(FGizmoSettingsModule, GizmoSettings)
diff --git a/Engine/Plugins/Experimental/GizmoFramework/Source/GizmoSettings/Public/GizmoSettings.h b/Engine/Plugins/Experimental/GizmoFramework/Source/GizmoSettings/Public/GizmoSettings.h
new file mode 100644
--- /dev/null
+++ b/Engine/Plugins/Experimental/GizmoFramework/Source/GizmoSettings/Public/GizmoSettings.h
@@ -0,0 +1,30 @@
+// Copyright Epic Games, Inc. All Rights Reserved.
+
+#pragma once
+
+#include "CoreMinimal.h"
+#include "Engine/DeveloperSettings.h"
+#include "EditorGizmos/TransformGizmo.h"
+
+#include "GizmoSettings.generated.h"
+
+UCLASS(config = EditorPerProjectUserSettings, defaultconfig, meta = (DisplayName = "New TRS Gizmo"))
+class GIZMOSETTINGS_API UGizmoSettings : public UDeveloperSettings
+{
+	GENERATED_UCLASS_BODY()
+
+public:
+
+	/** Enable/Disable New TRS Gizmos across the editor. */
+	UPROPERTY(config, EditAnywhere, Category = "New TRS Gizmo")
+	bool bEnableNewGizmos = false;
+
+	/** Change the current gizmos parameters. */
+	UPROPERTY(config, EditAnywhere, Category = "New TRS Gizmo")
+	FGizmosParameters GizmoParameters;
+	
+	DECLARE_MULTICAST_DELEGATE_OneParam(FOnUpdateSettings,  const UGizmoSettings*);
+	static FOnUpdateSettings OnSettingsChange;
+
+	virtual void PostEditChangeProperty(FPropertyChangedEvent& PropertyChangedEvent) override;
+};
\ No newline at end of file
diff --git a/Engine/Plugins/Experimental/GizmoFramework/Source/GizmoSettings/Public/GizmoSettingsModule.h b/Engine/Plugins/Experimental/GizmoFramework/Source/GizmoSettings/Public/GizmoSettingsModule.h
new file mode 100644
--- /dev/null
+++ b/Engine/Plugins/Experimental/GizmoFramework/Source/GizmoSettings/Public/GizmoSettingsModule.h
@@ -0,0 +1,16 @@
+// Copyright Epic Games, Inc. All Rights Reserved.
+
+#pragma once
+
+#include "Modules/ModuleInterface.h"
+
+class UGizmoSettings;
+
+class FGizmoSettingsModule : public IModuleInterface
+{
+	/** IModuleInterface implementation */
+	virtual void StartupModule() override;
+	virtual void ShutdownModule() override;
+
+	void UpdateGizmos(const UGizmoSettings* InSettings);
+};
diff --git a/Engine/Plugins/Experimental/GizmoFramework/Source/NewGizmoEnable/NewGizmoEnable.build.cs b/Engine/Plugins/Experimental/GizmoFramework/Source/NewGizmoEnable/NewGizmoEnable.build.cs
deleted file mode 100644
diff --git a/Engine/Plugins/Experimental/GizmoFramework/Source/NewGizmoEnable/Private/NewGizmoEnable.cpp b/Engine/Plugins/Experimental/GizmoFramework/Source/NewGizmoEnable/Private/NewGizmoEnable.cpp
deleted file mode 100644
diff --git a/Engine/Plugins/Experimental/GizmoFramework/Source/NewGizmoEnable/Private/NewGizmoEnableCommands.cpp b/Engine/Plugins/Experimental/GizmoFramework/Source/NewGizmoEnable/Private/NewGizmoEnableCommands.cpp
deleted file mode 100644
diff --git a/Engine/Plugins/Experimental/GizmoFramework/Source/NewGizmoEnable/Public/NewGizmoEnable.h b/Engine/Plugins/Experimental/GizmoFramework/Source/NewGizmoEnable/Public/NewGizmoEnable.h
deleted file mode 100644
diff --git a/Engine/Plugins/Experimental/GizmoFramework/Source/NewGizmoEnable/Public/NewGizmoEnableCommands.h b/Engine/Plugins/Experimental/GizmoFramework/Source/NewGizmoEnable/Public/NewGizmoEnableCommands.h
deleted file mode 100644
diff --git a/Engine/Source/Editor/Experimental/EditorInteractiveToolsFramework/Private/EditorGizmos/EditorTransformGizmoUtil.cpp b/Engine/Source/Editor/Experimental/EditorInteractiveToolsFramework/Private/EditorGizmos/EditorTransformGizmoUtil.cpp
--- a/Engine/Source/Editor/Experimental/EditorInteractiveToolsFramework/Private/EditorGizmos/EditorTransformGizmoUtil.cpp
+++ b/Engine/Source/Editor/Experimental/EditorInteractiveToolsFramework/Private/EditorGizmos/EditorTransformGizmoUtil.cpp
@@ -163,7 +163,7 @@ void UEditorTransformGizmoContextObject::Initialize(FEditorModeTools* InModeTool
 	if (ensure(InModeTools && !ModeTools))
 	{
 		ModeTools = InModeTools;
-		InitializeCVarBinding();
+		InitializeGizmoManagerBinding();
 
 		DataBinder = MakeShared<FEditorTransformGizmoDataBinder>();
 		DataBinder->BindToGizmoContextObject(this);
@@ -181,7 +181,7 @@ void UEditorTransformGizmoContextObject::Shutdown()
 	{
 		UpdateGizmo({});
 		RemoveViewportsBinding();
-		RemoveCVarBinding();
+		RemoveGizmoManagerBinding();
 		
 		ModeTools = nullptr;
 	}
@@ -201,6 +201,7 @@ UTransformGizmo* UEditorTransformGizmoContextObject::CreateTransformGizmo(
 	if (ensure(NewGizmo))
 	{
 		OnGizmoCreated.Broadcast(NewGizmo);
+		InGizmoManager->OnGizmosParametersChangedDelegate().AddUObject(NewGizmo, &UTransformGizmo::OnParametersChanged);
 	}
 	
 	return NewGizmo;
@@ -271,41 +272,36 @@ void UEditorTransformGizmoContextObject::UpdateGizmo(const TArray<FEditorViewpor
 	}
 }
 
-void UEditorTransformGizmoContextObject::InitializeCVarBinding()
+void UEditorTransformGizmoContextObject::InitializeGizmoManagerBinding()
 {
-	static IConsoleVariable* UseLegacyWidgetCVar = IConsoleManager::Get().FindConsoleVariable(TEXT("Gizmos.UseLegacyWidget"));
-	if (ensure(UseLegacyWidgetCVar))
+	auto OnGizmoVariableChanged = [this](const bool bUseNewTRSGizmos)
 	{
-		auto OnGizmoVariableChanged = [this](IConsoleVariable* InLegacyWidgetCVar)
+		if (!bUseNewTRSGizmos)
 		{
-			const bool bUseLegacyGizmo = InLegacyWidgetCVar ? UseLegacyWidgetCVar->GetInt() > 0 : false;
-			if (bUseLegacyGizmo)
-			{
-				// swap back default mode
-				(void)SwapDefaultMode(FAssetEdModes::EM_AssetDefault, FBuiltinEditorModes::EM_Default);
-				
-				// remove viewports' binding + gizmos as they are useless
-				const UModeManagerInteractiveToolsContext* ToolsContext = ModeTools->GetInteractiveToolsContext();
-				const TObjectPtr<UInteractiveToolManager> ToolManager = ToolsContext->ToolManager;
-				UE::EditorTransformGizmoUtil::RemoveDefaultTransformGizmo(ToolManager);
-				
-				RemoveViewportsBinding();
-				
-				return;
-			}
+			// swap back default mode
+			(void)SwapDefaultMode(FAssetEdModes::EM_AssetDefault, FBuiltinEditorModes::EM_Default);
+			
+			// remove viewports' binding + gizmos as they are useless
+			const UModeManagerInteractiveToolsContext* ToolsContext = ModeTools->GetInteractiveToolsContext();
+			const TObjectPtr<UInteractiveToolManager> ToolManager = ToolsContext->ToolManager;
+			UE::EditorTransformGizmoUtil::RemoveDefaultTransformGizmo(ToolManager);
+			
+			RemoveViewportsBinding();
+			
+			return;
+		}
 
-			InitializeViewportsBinding();
-		};
+		InitializeViewportsBinding();
+	};
 
-		// bind cvar change delegate
-		if (!UseLegacyChangedHandled.IsValid())
-		{
-			UseLegacyChangedHandled = UseLegacyWidgetCVar->OnChangedDelegate().AddLambda(OnGizmoVariableChanged);
-		}
-		
-		// initialize default viewport list change binding
-		OnGizmoVariableChanged(UseLegacyWidgetCVar);
+	// bind UseNewGizmosChangedHandled if needed
+	if (!UseNewGizmosChangedHandled.IsValid())
+	{
+		UEditorInteractiveGizmoManager::OnUsesNewTRSGizmosChangedDelegate().AddLambda(OnGizmoVariableChanged);
 	}
+	
+	// initialize default viewport list change binding
+	OnGizmoVariableChanged(UEditorInteractiveGizmoManager::UsesNewTRSGizmos());
 }
 
 void UEditorTransformGizmoContextObject::InitializeViewportsBinding()
@@ -324,7 +320,7 @@ void UEditorTransformGizmoContextObject::InitializeViewportsBinding()
 			UpdateGizmo(ViewportClients);
 		};
 
-		// bind if needed
+		// bind ViewportClientsChangedHandle if needed
         if (!ViewportClientsChangedHandle.IsValid())
         {
         	ViewportClientsChangedHandle = GEditor->OnViewportClientListChanged().AddLambda(OnViewportClientsChanged);
@@ -335,16 +331,12 @@ void UEditorTransformGizmoContextObject::InitializeViewportsBinding()
 	}
 }
 
-void UEditorTransformGizmoContextObject::RemoveCVarBinding()
+void UEditorTransformGizmoContextObject::RemoveGizmoManagerBinding()
 {
-	static IConsoleVariable* UseLegacyWidgetCVar = IConsoleManager::Get().FindConsoleVariable(TEXT("Gizmos.UseLegacyWidget"));
-	if (ensure(UseLegacyWidgetCVar))
+	if (UseNewGizmosChangedHandled.IsValid())
 	{
-		if (UseLegacyChangedHandled.IsValid())
-		{
-			UseLegacyWidgetCVar->OnChangedDelegate().Remove(UseLegacyChangedHandled);
-			UseLegacyChangedHandled.Reset();
-		}
+		UEditorInteractiveGizmoManager::OnUsesNewTRSGizmosChangedDelegate().Remove(UseNewGizmosChangedHandled);
+		UseNewGizmosChangedHandled.Reset();
 	}
 }
 
diff --git a/Engine/Source/Editor/Experimental/EditorInteractiveToolsFramework/Private/EditorGizmos/TransformGizmo.cpp b/Engine/Source/Editor/Experimental/EditorInteractiveToolsFramework/Private/EditorGizmos/TransformGizmo.cpp
--- a/Engine/Source/Editor/Experimental/EditorInteractiveToolsFramework/Private/EditorGizmos/TransformGizmo.cpp
+++ b/Engine/Source/Editor/Experimental/EditorInteractiveToolsFramework/Private/EditorGizmos/TransformGizmo.cpp
@@ -51,17 +51,6 @@ static FAutoConsoleVariableRef CVarProjectIndirect(
 	ProjectIndirect,
 	TEXT("Project to the nearest point of the curve when handling indirect rotation.")
 	);
-
-static int32 cvRotateMode = 0;
-static FAutoConsoleVariableRef CVarRotateMode(
-	TEXT("Gizmos.RotateMode"),
-	cvRotateMode,
-	TEXT("Rotation mode to choose between Pull = 0 and Arc = 1."),
-	FConsoleVariableDelegate::CreateLambda([](IConsoleVariable*)
-	{
-		cvRotateMode = FMath::Clamp(cvRotateMode, 0, 1);
-	})
-	);
 }
 
 void UTransformGizmo::SetDisallowNegativeScaling(bool bDisallow)
@@ -109,7 +98,7 @@ void UTransformGizmo::SetupBehaviors()
 	AddInputBehavior(HoverBehavior);
 
 	// Add default mouse input behavior
-	MouseBehavior = NewObject<UClickDragInputBehavior>();
+	UClickDragInputBehavior* MouseBehavior = NewObject<UClickDragInputBehavior>();
 	MouseBehavior->Initialize(this);
 	MouseBehavior->SetDefaultPriority(FInputCapturePriority(FInputCapturePriority::DEFAULT_GIZMO_PRIORITY));
 	AddInputBehavior(MouseBehavior);
@@ -151,34 +140,30 @@ void UTransformGizmo::SetupIndirectBehaviors()
 		bIndirectManipulation = false;
 		return OnTerminateDragSequence();
 	};
-	if (bCtrlMiddleDoesY)
+	// disable ctrl + mmb for that behavior?
+	MiddleClickDragBehavior->ModifierCheckFunc = [this](const FInputDeviceState& InputState)
 	{
-		// disable ctrl + mmb for that behavior
-		MiddleClickDragBehavior->ModifierCheckFunc = [](const FInputDeviceState& InputState)
-		{
-			return !FInputDeviceState::IsCtrlKeyDown(InputState);
-		};
-	}
+		return !bCtrlMiddleDoesY || !FInputDeviceState::IsCtrlKeyDown(InputState);
+	};
 	AddInputBehavior(MiddleClickDragBehavior);
 
-	
 	// Add left/right mouse input behavior for indirect manipulation
-	UMultiButtonClickDragBehavior* LeftRightClickDragBehavior = NewObject<UMultiButtonClickDragBehavior>();
-	LeftRightClickDragBehavior->Initialize();
-	LeftRightClickDragBehavior->EnableButton(EKeys::LeftMouseButton);
+	MultiIndirectClickDragBehavior = NewObject<UMultiButtonClickDragBehavior>();
+	MultiIndirectClickDragBehavior->Initialize();
+	MultiIndirectClickDragBehavior->EnableButton(EKeys::LeftMouseButton);
 	if (bCtrlMiddleDoesY)
 	{
-		LeftRightClickDragBehavior->EnableButton(EKeys::MiddleMouseButton);
+		MultiIndirectClickDragBehavior->EnableButton(EKeys::MiddleMouseButton);
 	}
-	LeftRightClickDragBehavior->EnableButton(EKeys::RightMouseButton);
-	LeftRightClickDragBehavior->ModifierCheckFunc = FInputDeviceState::IsCtrlKeyDown;
-	LeftRightClickDragBehavior->CanBeginClickDragFunc = [this](const FInputDeviceRay&)
+	MultiIndirectClickDragBehavior->EnableButton(EKeys::RightMouseButton);
+	MultiIndirectClickDragBehavior->ModifierCheckFunc = FInputDeviceState::IsCtrlKeyDown;
+	MultiIndirectClickDragBehavior->CanBeginClickDragFunc = [this](const FInputDeviceRay&)
 	{
 		static const FInputRayHit InvalidRayHit;
 		static const FInputRayHit ValidRayHit(TNumericLimits<double>::Max());
 		return CanInteract() ? ValidRayHit : InvalidRayHit;
 	};
-	LeftRightClickDragBehavior->OnClickPressFunc = [this](const FInputDeviceRay& InPressPos)
+	MultiIndirectClickDragBehavior->OnClickPressFunc = [this](const FInputDeviceRay& InPressPos)
 	{
 		bIndirectManipulation = true;
 		if (LastHitPart == ETransformGizmoPartIdentifier::Default)
@@ -187,17 +172,17 @@ void UTransformGizmo::SetupIndirectBehaviors()
 		}
 		return OnClickPress(InPressPos);
 	};
-	LeftRightClickDragBehavior->OnClickDragFunc = [this](const FInputDeviceRay& InDragPos)
+	MultiIndirectClickDragBehavior->OnClickDragFunc = [this](const FInputDeviceRay& InDragPos)
 	{
 		bIndirectManipulation = true;
 		return OnClickDrag(InDragPos);
 	};
-	LeftRightClickDragBehavior->OnClickReleaseFunc = [this](const FInputDeviceRay& InReleasePos)
+	MultiIndirectClickDragBehavior->OnClickReleaseFunc = [this](const FInputDeviceRay& InReleasePos)
 	{
 		bIndirectManipulation = false;
 		return OnClickRelease(InReleasePos);
 	};
-	LeftRightClickDragBehavior->OnTerminateFunc = [this]()
+	MultiIndirectClickDragBehavior->OnTerminateFunc = [this]()
 	{
 		bIndirectManipulation = false;
 		return OnTerminateDragSequence();
@@ -246,7 +231,7 @@ void UTransformGizmo::SetupIndirectBehaviors()
 		return ETransformGizmoPartIdentifier::Default;
 	};
 	
-	LeftRightClickDragBehavior->OnStateUpdated = [this, GetAxis, GetHitPart, LeftRightClickDragBehavior](const FInputDeviceState& Input)
+	MultiIndirectClickDragBehavior->OnStateUpdated = [this, GetAxis, GetHitPart](const FInputDeviceState& Input)
 	{
 		// disable indirect if the current axis is none 
 		const EAxis::Type Axis = GetAxis(Input);
@@ -271,11 +256,11 @@ void UTransformGizmo::SetupIndirectBehaviors()
 			const uint8 HitPartIndex = static_cast<uint8>(LastHitPart);
 			if (OnClickPressFunctions.IsValidIndex(HitPartIndex) && OnClickPressFunctions[HitPartIndex])
 			{
-				OnClickPressFunctions[HitPartIndex](this, LeftRightClickDragBehavior->GetDeviceRay(Input));
+				OnClickPressFunctions[HitPartIndex](this, MultiIndirectClickDragBehavior->GetDeviceRay(Input));
 			}
 		}
 	};
-	AddInputBehavior(LeftRightClickDragBehavior);
+	AddInputBehavior(MultiIndirectClickDragBehavior);
 }
 
 void UTransformGizmo::SetupMaterials()
@@ -996,6 +981,29 @@ void UTransformGizmo::HandleWidgetModeChanged(UE::Widget::EWidgetMode InWidgetMo
 	}
 }
 
+void UTransformGizmo::OnParametersChanged(const FGizmosParameters& InParameters)
+{
+	if (InParameters.bCtrlMiddleDoesY != bCtrlMiddleDoesY)
+	{
+		bCtrlMiddleDoesY = InParameters.bCtrlMiddleDoesY;
+
+		// update CTRL + LMB/MMB/RMB indirect behavior
+		if (MultiIndirectClickDragBehavior)
+		{
+			if (bCtrlMiddleDoesY)
+			{
+				MultiIndirectClickDragBehavior->EnableButton(EKeys::MiddleMouseButton);
+			}
+			else
+			{
+				MultiIndirectClickDragBehavior->DisableButton(EKeys::MiddleMouseButton);
+			}
+		}
+	}
+	
+	DefaultRotateMode = InParameters.RotateMode;
+}
+
 UGizmoElementArrow* UTransformGizmo::MakeTranslateAxis(ETransformGizmoPartIdentifier InPartId, const FVector& InAxisDir, const FVector& InSideDir, UMaterialInterface* InMaterial)
 {
 	const float SizeCoeff = GetSizeCoefficient();
@@ -2050,7 +2058,7 @@ void UTransformGizmo::OnClickPressRotateAxis(const FInputDeviceRay& InPressPos)
 	RotateMode = EAxisRotateMode::Pull;
 
 	// initialize arc/mixed data
-	const bool bRotatePull = bIndirectManipulation || GizmoLocals::cvRotateMode == 0;
+	const bool bRotatePull = bIndirectManipulation || DefaultRotateMode == EAxisRotateMode::Pull;
 	if (!bRotatePull)
 	{
 		static const TArray RotateAxis({FVector::XAxisVector, FVector::YAxisVector, FVector::ZAxisVector});
diff --git a/Engine/Source/Editor/Experimental/EditorInteractiveToolsFramework/Private/EditorInteractiveGizmoManager.cpp b/Engine/Source/Editor/Experimental/EditorInteractiveToolsFramework/Private/EditorInteractiveGizmoManager.cpp
--- a/Engine/Source/Editor/Experimental/EditorInteractiveToolsFramework/Private/EditorInteractiveGizmoManager.cpp
+++ b/Engine/Source/Editor/Experimental/EditorInteractiveToolsFramework/Private/EditorInteractiveGizmoManager.cpp
@@ -9,7 +9,6 @@
 #include "EditorInteractiveGizmoSelectionBuilder.h"
 #include "EditorInteractiveGizmoSubsystem.h"
 #include "EditorModeManager.h"
-#include "Elements/Framework/TypedElementSelectionSet.h"
 #include "HAL/IConsoleManager.h"
 #include "HAL/Platform.h"
 #include "HAL/PlatformCrt.h"
@@ -29,24 +28,38 @@ class FCanvas;
 
 namespace GizmoManagerLocals
 {
-	static int32 UseLegacyWidget = 1;
+	static bool UseNewTRSWidget = false;
+	static UEditorInteractiveGizmoManager::FOnUsesNewTRSGizmosChanged OnUsesNewTRSGizmosChanged;
+	static UEditorInteractiveGizmoManager::FOnGizmosParametersChanged OnGizmosParametersChanged;
 }
 
-static FAutoConsoleVariableRef CVarUseLegacyWidget(
-	TEXT("Gizmos.UseLegacyWidget"),
-	GizmoManagerLocals::UseLegacyWidget,
-	TEXT("Specify whether to use selection-based gizmos or legacy widget\n")
-	TEXT("0 = enable UE5 transform and other selection-based gizmos.\n")
-	TEXT("1 = enable legacy UE4 transform widget."),
-	FConsoleVariableDelegate::CreateLambda([](IConsoleVariable*)
+bool UEditorInteractiveGizmoManager::UsesNewTRSGizmos()
+{
+	return GizmoManagerLocals::UseNewTRSWidget;
+}
+
+void UEditorInteractiveGizmoManager::SetUsesNewTRSGizmos(const bool bUseNewTRSGizmos)
+{
+	if (bUseNewTRSGizmos != GizmoManagerLocals::UseNewTRSWidget)
 	{
-		GizmoManagerLocals::UseLegacyWidget = FMath::Clamp(GizmoManagerLocals::UseLegacyWidget, 0, 1);
-	}),
-	ECVF_RenderThreadSafe);
+		GizmoManagerLocals::UseNewTRSWidget = bUseNewTRSGizmos;
+		GizmoManagerLocals::OnUsesNewTRSGizmosChanged.Broadcast(GizmoManagerLocals::UseNewTRSWidget);
+	}
+}
 
-bool UEditorInteractiveGizmoManager::UsesNewTRSGizmos()
+UEditorInteractiveGizmoManager::FOnUsesNewTRSGizmosChanged& UEditorInteractiveGizmoManager::OnUsesNewTRSGizmosChangedDelegate()
+{
+	return GizmoManagerLocals::OnUsesNewTRSGizmosChanged;
+}
+
+void UEditorInteractiveGizmoManager::SetGizmosParameters(const FGizmosParameters& InParameters)
 {
-	return (GizmoManagerLocals::UseLegacyWidget == 0);
+	GizmoManagerLocals::OnGizmosParametersChanged.Broadcast(InParameters);
+}
+
+UEditorInteractiveGizmoManager::FOnGizmosParametersChanged& UEditorInteractiveGizmoManager::OnGizmosParametersChangedDelegate()
+{
+	return GizmoManagerLocals::OnGizmosParametersChanged;
 }
 
 UEditorInteractiveGizmoManager::UEditorInteractiveGizmoManager() :
@@ -112,6 +125,8 @@ bool UEditorInteractiveGizmoManager::DestroyEditorGizmo(UInteractiveGizmo* Gizmo
 		return false;
 	}
 
+	OnGizmosParametersChangedDelegate().RemoveAll(Gizmo);
+	
 	InputRouter->ForceTerminateSource(Gizmo);
 
 	Gizmo->Shutdown();
@@ -172,6 +187,20 @@ UInteractiveGizmo* UEditorInteractiveGizmoManager::CreateGizmo(const FString& Bu
 	return Super::CreateGizmo(BuilderIdentifier, InstanceIdentifier, Owner);
 }
 
+bool UEditorInteractiveGizmoManager::DestroyGizmo(UInteractiveGizmo* InGizmo)
+{
+	const bool bHasGizmo = ActiveGizmos.ContainsByPredicate([InGizmo](const FActiveGizmo& ActiveGizmo)
+	{
+		return ActiveGizmo.Gizmo == InGizmo;
+	});
+	if (bHasGizmo)
+	{
+		OnGizmosParametersChangedDelegate().RemoveAll(InGizmo);
+	}
+	
+	return Super::DestroyGizmo(InGizmo);
+}
+
 // @todo move this to a gizmo context object
 bool UEditorInteractiveGizmoManager::GetShowEditorGizmos()
 {
@@ -204,6 +233,7 @@ void UEditorInteractiveGizmoManager::UpdateActiveEditorGizmos()
 		return;
 	}
 	
+	
 	const bool bEditorModeToolsSupportsWidgetDrawing = EditorModeManager ? EditorModeManager->GetShowWidget() : true;
 	const bool bNewShowEditorGizmos = bEditorModeToolsSupportsWidgetDrawing && bEnableEditorGizmos;
 
diff --git a/Engine/Source/Editor/Experimental/EditorInteractiveToolsFramework/Public/EditorGizmos/EditorTransformGizmoUtil.h b/Engine/Source/Editor/Experimental/EditorInteractiveToolsFramework/Public/EditorGizmos/EditorTransformGizmoUtil.h
--- a/Engine/Source/Editor/Experimental/EditorInteractiveToolsFramework/Public/EditorGizmos/EditorTransformGizmoUtil.h
+++ b/Engine/Source/Editor/Experimental/EditorInteractiveToolsFramework/Public/EditorGizmos/EditorTransformGizmoUtil.h
@@ -95,10 +95,10 @@ private:
 
 	EDITORINTERACTIVETOOLSFRAMEWORK_API void UpdateGizmo(const TArray<FEditorViewportClient*>& InViewportClients) const;
 	
-	EDITORINTERACTIVETOOLSFRAMEWORK_API void InitializeCVarBinding();
+	EDITORINTERACTIVETOOLSFRAMEWORK_API void InitializeGizmoManagerBinding();
 	EDITORINTERACTIVETOOLSFRAMEWORK_API void InitializeViewportsBinding();
 	
-	EDITORINTERACTIVETOOLSFRAMEWORK_API void RemoveCVarBinding();
+	EDITORINTERACTIVETOOLSFRAMEWORK_API void RemoveGizmoManagerBinding();
 	EDITORINTERACTIVETOOLSFRAMEWORK_API void RemoveViewportsBinding();
 
 	typedef FName FEditorModeID;
@@ -108,7 +108,7 @@ private:
 	
 	FEditorModeTools* ModeTools = nullptr;
 	FDelegateHandle ViewportClientsChangedHandle;
-	FDelegateHandle UseLegacyChangedHandled;
+	FDelegateHandle UseNewGizmosChangedHandled;
 	bool bGizmosRegistered = false;
 
 	TSharedPtr<FEditorTransformGizmoDataBinder> DataBinder;
diff --git a/Engine/Source/Editor/Experimental/EditorInteractiveToolsFramework/Public/EditorGizmos/TransformGizmo.h b/Engine/Source/Editor/Experimental/EditorInteractiveToolsFramework/Public/EditorGizmos/TransformGizmo.h
--- a/Engine/Source/Editor/Experimental/EditorInteractiveToolsFramework/Public/EditorGizmos/TransformGizmo.h
+++ b/Engine/Source/Editor/Experimental/EditorInteractiveToolsFramework/Public/EditorGizmos/TransformGizmo.h
@@ -34,7 +34,7 @@ class IGizmoStateTarget;
 class IGizmoTransformSource;
 class IToolContextTransactionProvider;
 class IToolsContextRenderAPI;
-class UClickDragInputBehavior;
+class UMultiButtonClickDragBehavior;
 class UGizmoConstantFrameAxisSource;
 class UGizmoElementArrow;
 class UGizmoElementBase;
@@ -93,6 +93,7 @@ enum class ETransformGizmoPartIdentifier
 	Max
 };
 
+UENUM()
 namespace EAxisRotateMode
 {
 	enum Type
@@ -102,6 +103,19 @@ namespace EAxisRotateMode
 	};
 }
 
+USTRUCT()
+struct FGizmosParameters
+{
+	GENERATED_BODY()
+
+	/** Rotate mode used when manipulating axis rotation handles. */
+	UPROPERTY(EditAnywhere, Category = NewTRSGizmo)
+	TEnumAsByte<EAxisRotateMode::Type> RotateMode = EAxisRotateMode::Pull;
+	
+	/** Property used to define whether the Y axis is indirectly controlled via CTRL+MMB or CTRL+RMB (and Z via CTRL+RMB or CTRL+LMB+RMB). */
+	UPROPERTY(EditAnywhere, Category = NewTRSGizmo)
+	bool bCtrlMiddleDoesY = true;
+};
 
 /**
  * UTransformGizmo provides standard Transformation Gizmo interactions,
@@ -220,6 +234,11 @@ public:
 	 * Handle widget mode changed.
 	 */
 	void HandleWidgetModeChanged(UE::Widget::EWidgetMode InWidgetMode);
+
+	/**
+	 * Handle user parameters changes
+	 */
+	void OnParametersChanged(const FGizmosParameters& InParameters);
 	
 public:
 
@@ -231,9 +250,9 @@ public:
 	UPROPERTY()
 	TObjectPtr<UGizmoElementHitMultiTarget> HitTarget;
 
-	/** The mouse click behavior of the gizmo is accessible so that it can be modified to use different mouse keys. */
+	/** The multi button mouse click behavior is accessible so that it can be modified to use different mouse keys. */
 	UPROPERTY()
-	TObjectPtr<UClickDragInputBehavior> MouseBehavior;
+	TObjectPtr<UMultiButtonClickDragBehavior> MultiIndirectClickDragBehavior;
 
 	/** Transform Gizmo Source */
 	UPROPERTY()
@@ -872,8 +891,13 @@ protected:
 	UPROPERTY()
 	bool bCtrlMiddleDoesY = true;
 
+	/** Default rotate mode used when using axis rotation handles. */
+	UPROPERTY()
+	TEnumAsByte<EAxisRotateMode::Type> DefaultRotateMode = EAxisRotateMode::Pull;
+
+	/** Actual rotate mode used (based on view dependant information). */
 	TEnumAsByte<EAxisRotateMode::Type> RotateMode = EAxisRotateMode::Pull;
-	
+
 private:
 	/** Debug attributes to display the pull direction */
 	bool bDebugRotate = false;
diff --git a/Engine/Source/Editor/Experimental/EditorInteractiveToolsFramework/Public/EditorInteractiveGizmoManager.h b/Engine/Source/Editor/Experimental/EditorInteractiveToolsFramework/Public/EditorInteractiveGizmoManager.h
--- a/Engine/Source/Editor/Experimental/EditorInteractiveToolsFramework/Public/EditorInteractiveGizmoManager.h
+++ b/Engine/Source/Editor/Experimental/EditorInteractiveToolsFramework/Public/EditorInteractiveGizmoManager.h
@@ -26,6 +26,7 @@ class UObject;
 class UTypedElementSelectionSet;
 struct FToolBuilderState;
 class UTransformGizmo;
+struct FGizmosParameters;
 
 USTRUCT()
 struct FActiveEditorGizmo
@@ -126,6 +127,9 @@ public:
 	virtual UInteractiveGizmo* CreateGizmo(
 		const FString& BuilderIdentifier, const FString& InstanceIdentifier = FString(), void* Owner = nullptr) override;
 
+	/** Shutdown and remove a Gizmo (UInteractiveGizmoManager override) */
+	virtual bool DestroyGizmo(UInteractiveGizmo* InGizmo) override;
+
 	/** instance/builder identifiers for transform gizmo */
 	static const FString& TransformInstanceIdentifier();
 	static const FString& TransformBuilderIdentifier();
@@ -134,7 +138,29 @@ public:
 	 * Returns true if the new TRS gizmos are used.
 	 */
 	static bool UsesNewTRSGizmos();
-	
+
+	/**
+	 * Updates the current New TRS Gizmo state and notifies that change using OnUsesNewTRSGizmosChangedDelegate  
+	 */
+	static void SetUsesNewTRSGizmos(const bool bUseNewTRSGizmos);
+
+	/**
+	 * Delegate to notify from bUseNewTRSGizmos changes  
+	 */
+	DECLARE_MULTICAST_DELEGATE_OneParam(FOnUsesNewTRSGizmosChanged, const bool bUseNewTRSGizmos);
+	static FOnUsesNewTRSGizmosChanged& OnUsesNewTRSGizmosChangedDelegate();
+
+	/**
+	 * Notifies FGizmosParameters changes across bound transform gizmos
+	 */
+	static void SetGizmosParameters(const FGizmosParameters& InParameters);
+
+	/**
+	 * Delegate to notify from FGizmosParameters changes  
+	 */
+	DECLARE_MULTICAST_DELEGATE_OneParam(FOnGizmosParametersChanged, const FGizmosParameters& InParameters);
+    static FOnGizmosParametersChanged& OnGizmosParametersChangedDelegate();
+
 protected:
 
 	/**
