New TRS Gizmo: set default rotate to arc instead of pull

diff --git a/Engine/Source/Editor/Experimental/EditorInteractiveToolsFramework/Public/EditorGizmos/TransformGizmo.h b/Engine/Source/Editor/Experimental/EditorInteractiveToolsFramework/Public/EditorGizmos/TransformGizmo.h
--- a/Engine/Source/Editor/Experimental/EditorInteractiveToolsFramework/Public/EditorGizmos/TransformGizmo.h
+++ b/Engine/Source/Editor/Experimental/EditorInteractiveToolsFramework/Public/EditorGizmos/TransformGizmo.h
@@ -110,7 +110,7 @@ struct FGizmosParameters
 
 	/** Rotate mode used when manipulating axis rotation handles. */
 	UPROPERTY(EditAnywhere, Category = NewTRSGizmo)
-	TEnumAsByte<EAxisRotateMode::Type> RotateMode = EAxisRotateMode::Pull;
+	TEnumAsByte<EAxisRotateMode::Type> RotateMode = EAxisRotateMode::Arc;
 	
 	/** Property used to define whether the Y axis is indirectly controlled via CTRL+MMB or CTRL+RMB (and Z via CTRL+RMB or CTRL+LMB+RMB). */
 	UPROPERTY(EditAnywhere, Category = NewTRSGizmo)
@@ -893,10 +893,10 @@ protected:
 
 	/** Default rotate mode used when using axis rotation handles. */
 	UPROPERTY()
-	TEnumAsByte<EAxisRotateMode::Type> DefaultRotateMode = EAxisRotateMode::Pull;
+	TEnumAsByte<EAxisRotateMode::Type> DefaultRotateMode = EAxisRotateMode::Arc;
 
 	/** Actual rotate mode used (based on view dependant information). */
-	TEnumAsByte<EAxisRotateMode::Type> RotateMode = EAxisRotateMode::Pull;
+	TEnumAsByte<EAxisRotateMode::Type> RotateMode = EAxisRotateMode::Arc;
 
 private:
 	/** Debug attributes to display the pull direction */
