New TRS Gizmos: make sure settings are propagated when the gizmos are created after the plugin has been loaded

diff --git a/Engine/Plugins/Experimental/Animation/SkeletalMeshModelingTools/Source/SkeletalMeshModelingTools/Private/SkeletalMeshGizmoUtils.cpp b/Engine/Plugins/Experimental/Animation/SkeletalMeshModelingTools/Source/SkeletalMeshModelingTools/Private/SkeletalMeshGizmoUtils.cpp
--- a/Engine/Plugins/Experimental/Animation/SkeletalMeshModelingTools/Source/SkeletalMeshModelingTools/Private/SkeletalMeshGizmoUtils.cpp
+++ b/Engine/Plugins/Experimental/Animation/SkeletalMeshModelingTools/Source/SkeletalMeshModelingTools/Private/SkeletalMeshGizmoUtils.cpp
@@ -76,6 +76,10 @@ UTransformGizmo* UE::SkeletalMeshGizmoUtils::CreateTransformGizmo(UInteractiveTo
 	if (ensure(NewGizmo))
 	{
 		GizmoManager->OnGizmosParametersChangedDelegate().AddUObject(NewGizmo, &UTransformGizmo::OnParametersChanged);
+		if (GizmoManager->GetDefaultGizmosParameters())
+		{
+			NewGizmo->OnParametersChanged(*GizmoManager->GetDefaultGizmosParameters());
+		}
 	}
 	
 	return NewGizmo;
diff --git a/Engine/Source/Editor/Experimental/EditorInteractiveToolsFramework/Private/EditorGizmos/EditorTransformGizmoUtil.cpp b/Engine/Source/Editor/Experimental/EditorInteractiveToolsFramework/Private/EditorGizmos/EditorTransformGizmoUtil.cpp
--- a/Engine/Source/Editor/Experimental/EditorInteractiveToolsFramework/Private/EditorGizmos/EditorTransformGizmoUtil.cpp
+++ b/Engine/Source/Editor/Experimental/EditorInteractiveToolsFramework/Private/EditorGizmos/EditorTransformGizmoUtil.cpp
@@ -202,6 +202,10 @@ UTransformGizmo* UEditorTransformGizmoContextObject::CreateTransformGizmo(
 	{
 		OnGizmoCreated.Broadcast(NewGizmo);
 		InGizmoManager->OnGizmosParametersChangedDelegate().AddUObject(NewGizmo, &UTransformGizmo::OnParametersChanged);
+		if (InGizmoManager->GetDefaultGizmosParameters())
+		{
+			NewGizmo->OnParametersChanged(*InGizmoManager->GetDefaultGizmosParameters());
+		}
 	}
 	
 	return NewGizmo;
diff --git a/Engine/Source/Editor/Experimental/EditorInteractiveToolsFramework/Private/EditorInteractiveGizmoManager.cpp b/Engine/Source/Editor/Experimental/EditorInteractiveToolsFramework/Private/EditorInteractiveGizmoManager.cpp
--- a/Engine/Source/Editor/Experimental/EditorInteractiveToolsFramework/Private/EditorInteractiveGizmoManager.cpp
+++ b/Engine/Source/Editor/Experimental/EditorInteractiveToolsFramework/Private/EditorInteractiveGizmoManager.cpp
@@ -29,6 +29,7 @@ class FCanvas;
 namespace GizmoManagerLocals
 {
 	static bool UseNewTRSWidget = false;
+	static TOptional<FGizmosParameters> OptDefaultParameters;
 	static UEditorInteractiveGizmoManager::FOnUsesNewTRSGizmosChanged OnUsesNewTRSGizmosChanged;
 	static UEditorInteractiveGizmoManager::FOnGizmosParametersChanged OnGizmosParametersChanged;
 }
@@ -54,6 +55,7 @@ UEditorInteractiveGizmoManager::FOnUsesNewTRSGizmosChanged& UEditorInteractiveGi
 
 void UEditorInteractiveGizmoManager::SetGizmosParameters(const FGizmosParameters& InParameters)
 {
+	GizmoManagerLocals::OptDefaultParameters = InParameters;
 	GizmoManagerLocals::OnGizmosParametersChanged.Broadcast(InParameters);
 }
 
@@ -62,6 +64,12 @@ UEditorInteractiveGizmoManager::FOnGizmosParametersChanged& UEditorInteractiveGi
 	return GizmoManagerLocals::OnGizmosParametersChanged;
 }
 
+const TOptional<FGizmosParameters>& UEditorInteractiveGizmoManager::GetDefaultGizmosParameters()
+{
+	static const TOptional<FGizmosParameters> Invalid;
+	return GizmoManagerLocals::OptDefaultParameters.IsSet() ? GizmoManagerLocals::OptDefaultParameters : Invalid;
+}
+
 UEditorInteractiveGizmoManager::UEditorInteractiveGizmoManager() :
 	UInteractiveGizmoManager()
 {
diff --git a/Engine/Source/Editor/Experimental/EditorInteractiveToolsFramework/Public/EditorInteractiveGizmoManager.h b/Engine/Source/Editor/Experimental/EditorInteractiveToolsFramework/Public/EditorInteractiveGizmoManager.h
--- a/Engine/Source/Editor/Experimental/EditorInteractiveToolsFramework/Public/EditorInteractiveGizmoManager.h
+++ b/Engine/Source/Editor/Experimental/EditorInteractiveToolsFramework/Public/EditorInteractiveGizmoManager.h
@@ -161,6 +161,11 @@ public:
 	DECLARE_MULTICAST_DELEGATE_OneParam(FOnGizmosParametersChanged, const FGizmosParameters& InParameters);
     static FOnGizmosParametersChanged& OnGizmosParametersChangedDelegate();
 
+	/**
+	 * Returns the default gizmos parameters if set
+	 */
+    static const TOptional<FGizmosParameters>& GetDefaultGizmosParameters();
+
 protected:
 
 	/**
