New TRS Gizmo: fixed delegate not unbound on shutdown

diff --git a/Engine/Source/Editor/Experimental/EditorInteractiveToolsFramework/Private/EditorGizmos/EditorTransformGizmoUtil.cpp b/Engine/Source/Editor/Experimental/EditorInteractiveToolsFramework/Private/EditorGizmos/EditorTransformGizmoUtil.cpp
--- a/Engine/Source/Editor/Experimental/EditorInteractiveToolsFramework/Private/EditorGizmos/EditorTransformGizmoUtil.cpp
+++ b/Engine/Source/Editor/Experimental/EditorInteractiveToolsFramework/Private/EditorGizmos/EditorTransformGizmoUtil.cpp
@@ -301,7 +301,7 @@ void UEditorTransformGizmoContextObject::InitializeGizmoManagerBinding()
 	// bind UseNewGizmosChangedHandled if needed
 	if (!UseNewGizmosChangedHandled.IsValid())
 	{
-		UEditorInteractiveGizmoManager::OnUsesNewTRSGizmosChangedDelegate().AddLambda(OnGizmoVariableChanged);
+		UseNewGizmosChangedHandled = UEditorInteractiveGizmoManager::OnUsesNewTRSGizmosChangedDelegate().AddLambda(OnGizmoVariableChanged);
 	}
 	
 	// initialize default viewport list change binding
