New TRS Gizmo: reset LastHitPart on widget mode changed

#rnx

diff --git a/Engine/Source/Editor/Experimental/EditorInteractiveToolsFramework/Private/EditorGizmos/TransformGizmo.cpp b/Engine/Source/Editor/Experimental/EditorInteractiveToolsFramework/Private/EditorGizmos/TransformGizmo.cpp
--- a/Engine/Source/Editor/Experimental/EditorInteractiveToolsFramework/Private/EditorGizmos/TransformGizmo.cpp
+++ b/Engine/Source/Editor/Experimental/EditorInteractiveToolsFramework/Private/EditorGizmos/TransformGizmo.cpp
@@ -979,6 +979,8 @@ void UTransformGizmo::HandleWidgetModeChanged(UE::Widget::EWidgetMode InWidgetMo
 			UpdateInteractingState(true, DefaultHitPart, true);
 		}
 	}
+	
+	LastHitPart = ETransformGizmoPartIdentifier::Default;
 }
 
 void UTransformGizmo::OnParametersChanged(const FGizmosParameters& InParameters)
