New TRS Gizmo: fixed early access to level editor mode manager
    - follow-up to 33736099

#rnx

diff --git a/Engine/Source/Editor/Experimental/EditorInteractiveToolsFramework/Private/EditorGizmos/EditorTransformGizmoUtil.cpp b/Engine/Source/Editor/Experimental/EditorInteractiveToolsFramework/Private/EditorGizmos/EditorTransformGizmoUtil.cpp
--- a/Engine/Source/Editor/Experimental/EditorInteractiveToolsFramework/Private/EditorGizmos/EditorTransformGizmoUtil.cpp
+++ b/Engine/Source/Editor/Experimental/EditorInteractiveToolsFramework/Private/EditorGizmos/EditorTransformGizmoUtil.cpp
@@ -13,6 +13,8 @@
 #include "EditorGizmos/EditorTransformGizmoDataBinder.h"
 #include "EditorViewportClient.h"
 #include "EditorModes.h"
+#include "LevelEditor.h"
+#include "Modules/ModuleManager.h"
 #include "Tools/DefaultEdMode.h"
 
 #include UE_INLINE_GENERATED_CPP_BY_NAME(EditorTransformGizmoUtil)
@@ -363,9 +365,15 @@ bool UEditorTransformGizmoContextObject::SwapDefaultMode(const FEditorModeID InC
 	{
 		return false;
 	}
+
+	auto IsLevelEditorModeManager = [this]()
+	{
+		const FLevelEditorModule& LevelEditorModule = FModuleManager::GetModuleChecked<FLevelEditorModule>("LevelEditor");
+		const TSharedPtr<ILevelEditor> FirstLevelEditor = LevelEditorModule.GetFirstLevelEditor();
+		return FirstLevelEditor.IsValid() ? ModeTools == &FirstLevelEditor->GetEditorModeManager() : false;
+	};
 	
-	const bool bIsLevelMode = ModeTools == &GLevelEditorModeTools();
-	if (bIsLevelMode)
+	if (IsLevelEditorModeManager())
 	{
 		return false;
 	}
