New TRS Gizmo: fixed indirect translation not taking pixel to world ratio into account

#rnx

diff --git a/Engine/Source/Editor/Experimental/EditorInteractiveToolsFramework/Private/EditorGizmos/TransformGizmo.cpp b/Engine/Source/Editor/Experimental/EditorInteractiveToolsFramework/Private/EditorGizmos/TransformGizmo.cpp
--- a/Engine/Source/Editor/Experimental/EditorInteractiveToolsFramework/Private/EditorGizmos/TransformGizmo.cpp
+++ b/Engine/Source/Editor/Experimental/EditorInteractiveToolsFramework/Private/EditorGizmos/TransformGizmo.cpp
@@ -1727,10 +1727,12 @@ void UTransformGizmo::OnClickDragTranslateAxis(const FInputDeviceRay& DragPos)
 		const FVector2D XAxisDir = GetScreenProjectedAxis(GizmoViewContext, FVector::XAxisVector, CurrentTransform);
 		const FVector2D YAxisDir = GetScreenProjectedAxis(GizmoViewContext, FVector::YAxisVector, CurrentTransform);
 		const FVector2D ZAxisDir = GetScreenProjectedAxis(GizmoViewContext, FVector::ZAxisVector, CurrentTransform);
+
+		const float PixelToWorldRatio = GizmoRenderingUtil::CalculateLocalPixelToWorldScale(GizmoViewContext, CurrentTransform.GetLocation());
 		
-		FVector Delta((InteractionAxisList == EAxisList::X) ? FVector2D::DotProduct(XAxisDir, DragDir) : 0.0,
-					  (InteractionAxisList == EAxisList::Y) ? FVector2D::DotProduct(YAxisDir, DragDir) : 0.0,
-					  (InteractionAxisList == EAxisList::Z) ? FVector2D::DotProduct(ZAxisDir, DragDir) : 0.0);
+		FVector Delta((InteractionAxisList == EAxisList::X) ? PixelToWorldRatio * FVector2D::DotProduct(XAxisDir, DragDir) : 0.0,
+					  (InteractionAxisList == EAxisList::Y) ? PixelToWorldRatio * FVector2D::DotProduct(YAxisDir, DragDir) : 0.0,
+					  (InteractionAxisList == EAxisList::Z) ? PixelToWorldRatio * FVector2D::DotProduct(ZAxisDir, DragDir) : 0.0);
 		Delta = CurrentRotation * Delta;
 		
 		ApplyTranslateDelta(Delta);
