New TRS Gizmo: fixed deprecated namespace

#rnx

diff --git a/Engine/Source/Editor/Experimental/EditorInteractiveToolsFramework/Private/EditorGizmos/TransformGizmo.cpp b/Engine/Source/Editor/Experimental/EditorInteractiveToolsFramework/Private/EditorGizmos/TransformGizmo.cpp
--- a/Engine/Source/Editor/Experimental/EditorInteractiveToolsFramework/Private/EditorGizmos/TransformGizmo.cpp
+++ b/Engine/Source/Editor/Experimental/EditorInteractiveToolsFramework/Private/EditorGizmos/TransformGizmo.cpp
@@ -1728,7 +1728,7 @@ void UTransformGizmo::OnClickDragTranslateAxis(const FInputDeviceRay& DragPos)
 		const FVector2D YAxisDir = GetScreenProjectedAxis(GizmoViewContext, FVector::YAxisVector, CurrentTransform);
 		const FVector2D ZAxisDir = GetScreenProjectedAxis(GizmoViewContext, FVector::ZAxisVector, CurrentTransform);
 
-		const float PixelToWorldRatio = GizmoRenderingUtil::CalculateLocalPixelToWorldScale(GizmoViewContext, CurrentTransform.GetLocation());
+		const float PixelToWorldRatio = UE::GizmoRenderingUtil::CalculateLocalPixelToWorldScale(GizmoViewContext, CurrentTransform.GetLocation());
 		
 		FVector Delta((InteractionAxisList == EAxisList::X) ? PixelToWorldRatio * FVector2D::DotProduct(XAxisDir, DragDir) : 0.0,
 					  (InteractionAxisList == EAxisList::Y) ? PixelToWorldRatio * FVector2D::DotProduct(YAxisDir, DragDir) : 0.0,
