New TRS Gizmo: ability to switch from tangential to normal projection based on the first mouse drag
    - this is disable by default but can be enabled using 'Gizmos.AutoSwitchPullAxis 1' to get some feedback from users

#rnx

diff --git a/Engine/Source/Editor/Experimental/EditorInteractiveToolsFramework/Private/EditorGizmos/TransformGizmo.cpp b/Engine/Source/Editor/Experimental/EditorInteractiveToolsFramework/Private/EditorGizmos/TransformGizmo.cpp
--- a/Engine/Source/Editor/Experimental/EditorInteractiveToolsFramework/Private/EditorGizmos/TransformGizmo.cpp
+++ b/Engine/Source/Editor/Experimental/EditorInteractiveToolsFramework/Private/EditorGizmos/TransformGizmo.cpp
@@ -51,6 +51,13 @@ static FAutoConsoleVariableRef CVarProjectIndirect(
 	ProjectIndirect,
 	TEXT("Project to the nearest point of the curve when handling indirect rotation.")
 	);
+
+static bool	bAutoSwitchPullAxis = false;
+static FAutoConsoleVariableRef CVarAutoDetect(
+	TEXT("Gizmos.AutoSwitchPullAxis"),
+	bAutoSwitchPullAxis,
+	TEXT("Switch from tangential to normal projection based on the first mouse drag.")
+	);
 }
 
 void UTransformGizmo::SetDisallowNegativeScaling(bool bDisallow)
@@ -370,7 +377,8 @@ void UTransformGizmo::Render(IToolsContextRenderAPI* RenderAPI)
 			{
 				const float Radius = 2.f * GetWorldRadius(RotateAxisOuterRadius);
 				FPrimitiveDrawInterface* PDI = RenderAPI->GetPrimitiveDrawInterface();
-				PDI->DrawLine(DebugClosest - (DebugDirection * Radius), DebugClosest + (DebugDirection * Radius), FLinearColor::Yellow, SDPG_Foreground);
+				PDI->DrawLine(DebugClosest - (DebugDirection * Radius), DebugClosest + (DebugDirection * Radius), FLinearColor::Yellow, SDPG_Foreground, 0.f, 0.01f);
+				PDI->DrawLine(DebugClosest, DebugClosest + (DebugNormalRemoved * Radius), FLinearColor::Red, SDPG_Foreground, 0.f, 0.01f);
 			}
 		}
 	}
@@ -2047,6 +2055,7 @@ void UTransformGizmo::OnClickPressRotateAxis(const FInputDeviceRay& InPressPos)
 	const int32 RotateID = RotateIDs.IndexOfByKey(LastHitPart);
 	if (!ensure(RotateID != INDEX_NONE))
 	{
+		bTrySwitchingToNormalPull = false;
 		bInInteraction = false;
 		return;
 	}
@@ -2077,6 +2086,8 @@ void UTransformGizmo::OnClickPressRotateAxis(const FInputDeviceRay& InPressPos)
 		}
 	}
 
+	bTrySwitchingToNormalPull = GizmoLocals::bAutoSwitchPullAxis && bIndirectManipulation && RotateMode == EAxisRotateMode::Pull;
+	
 	bInInteraction = true;
 	SetModeLastHitPart(EGizmoTransformMode::Rotate, LastHitPart);
 }
@@ -2176,6 +2187,9 @@ FVector2D UTransformGizmo::GetScreenRotateAxisDir(const FInputDeviceRay& InPress
 	const double DotAxis = FMath::Abs( FVector2D::DotProduct(PullProjection, AxisProjection) );
 	const double DotClosest = FMath::Abs( FVector2D::DotProduct(PullProjection, ToClosestProjection) );
 	NormalProjectionToRemove = (DotAxis < DotClosest) ? AxisProjection : ToClosestProjection;
+	DebugNormalRemoved = (DotAxis < DotClosest) ? WorldAxis : ToClosestDirection;
+	DebugNormalSkip = (DotAxis < DotClosest) ? ToClosestDirection : WorldAxis;
+
 
 	// debug
 	{
@@ -2192,7 +2206,7 @@ void UTransformGizmo::OnClickDragRotateAxis(const FInputDeviceRay& DragPos)
 	{
 	case EAxisRotateMode::Pull:
 	{
-		const FQuat DeltaRot = ComputeAxisRotateDelta(InteractionScreenCurrPos, DragPos.ScreenPosition);
+		const FQuat DeltaRot = ComputeAxisRotateDelta(InteractionScreenCurrPos, DragPos);
 		ApplyRotateDelta(DeltaRot);
 		InteractionScreenCurrPos = DragPos.ScreenPosition;
 		break;
@@ -2219,10 +2233,23 @@ void UTransformGizmo::OnClickDragRotateAxis(const FInputDeviceRay& DragPos)
 	}
 }
 
-FQuat UTransformGizmo::ComputeAxisRotateDelta(const FVector2D& InStartPos, const FVector2D& InEndPos)
+FQuat UTransformGizmo::ComputeAxisRotateDelta(const FVector2D& InStartPos, const FInputDeviceRay& InDragPos)
 {
-	FVector2D DragDir = InEndPos - InStartPos;
+	FVector2D DragDir = InDragPos.ScreenPosition - InStartPos;
 
+	if (bTrySwitchingToNormalPull)
+	{
+		const double DotTangent = FVector2D::DotProduct(InteractionScreenAxisDirection, DragDir);
+		const double DotNormal = FVector2D::DotProduct(NormalProjectionToRemove, DragDir);
+		if (FMath::Abs(DotNormal) > FMath::Abs(DotTangent))
+		{
+			::Swap(NormalProjectionToRemove, InteractionScreenAxisDirection);
+			::Swap(DebugDirection, DebugNormalRemoved);
+			InteractionScreenAxisDirection *= FMath::Sign(DotTangent) * FMath::Sign(DotNormal);
+		}
+		bTrySwitchingToNormalPull = false;
+	}
+	
 	const FVector2D DragDirToRemove = NormalProjectionToRemove * FVector2D::DotProduct(DragDir, NormalProjectionToRemove);
 	DragDir -= DragDirToRemove;
 
@@ -2255,6 +2282,7 @@ void UTransformGizmo::OnClickReleaseRotateAxis(const FInputDeviceRay& InReleaseP
 {
 	bInInteraction = false;
 	bDebugRotate = false;
+	bTrySwitchingToNormalPull = false;
 }
 
 void UTransformGizmo::OnClickPressScreenSpaceRotate(const FInputDeviceRay& PressPos)
diff --git a/Engine/Source/Editor/Experimental/EditorInteractiveToolsFramework/Public/EditorGizmos/TransformGizmo.h b/Engine/Source/Editor/Experimental/EditorInteractiveToolsFramework/Public/EditorGizmos/TransformGizmo.h
--- a/Engine/Source/Editor/Experimental/EditorInteractiveToolsFramework/Public/EditorGizmos/TransformGizmo.h
+++ b/Engine/Source/Editor/Experimental/EditorInteractiveToolsFramework/Public/EditorGizmos/TransformGizmo.h
@@ -591,7 +591,7 @@ protected:
 	FVector2D GetScreenRotateAxisDir(const FInputDeviceRay& InPressPos);
 
 	/** Compute rotate delta based on screen-space start/end positions */
-	virtual FQuat ComputeAxisRotateDelta(const FVector2D& InStartPos, const FVector2D& InEndPos);
+	virtual FQuat ComputeAxisRotateDelta(const FVector2D& InStartPos, const FInputDeviceRay& InDragPos);
 
 	/** Prepares data for arc rotation. This will return false if this is not possible (the rotate handle is perpendicular to the view) */
 	bool OnClickPressRotateArc( const FInputDeviceRay& InPressPos,
@@ -898,9 +898,14 @@ protected:
 	/** Actual rotate mode used (based on view dependant information). */
 	TEnumAsByte<EAxisRotateMode::Type> RotateMode = EAxisRotateMode::Arc;
 
+	/** Switch from tangential to normal projection based on the first mouse drag. */
+	bool bTrySwitchingToNormalPull = false;
+
 private:
 	/** Debug attributes to display the pull direction */
 	bool bDebugRotate = false;
 	FVector DebugDirection = FVector::ZeroVector;
 	FVector DebugClosest = FVector::ZeroVector;
+	FVector DebugNormalRemoved = FVector::ZeroVector;
+	FVector DebugNormalSkip = FVector::ZeroVector;
 };
