New TRS Gizmo: indirect screen space rotation is now consistent with other axial rotations

#rnx

diff --git a/Engine/Source/Editor/Experimental/EditorInteractiveToolsFramework/Private/EditorGizmos/TransformGizmo.cpp b/Engine/Source/Editor/Experimental/EditorInteractiveToolsFramework/Private/EditorGizmos/TransformGizmo.cpp
--- a/Engine/Source/Editor/Experimental/EditorInteractiveToolsFramework/Private/EditorGizmos/TransformGizmo.cpp
+++ b/Engine/Source/Editor/Experimental/EditorInteractiveToolsFramework/Private/EditorGizmos/TransformGizmo.cpp
@@ -375,7 +375,8 @@ void UTransformGizmo::Render(IToolsContextRenderAPI* RenderAPI)
 		{
 			if (bDebugRotate)
 			{
-				const float Radius = 2.f * GetWorldRadius(RotateAxisOuterRadius);
+				const float RotateRadius = InteractionAxisList == EAxisList::Screen ? RotateScreenSpaceRadius : RotateAxisOuterRadius;
+				const float Radius = 2.f * GetWorldRadius(RotateRadius);
 				FPrimitiveDrawInterface* PDI = RenderAPI->GetPrimitiveDrawInterface();
 				PDI->DrawLine(DebugClosest - (DebugDirection * Radius), DebugClosest + (DebugDirection * Radius), FLinearColor::Yellow, SDPG_Foreground, 0.f, 0.01f);
 				PDI->DrawLine(DebugClosest, DebugClosest + (DebugNormalRemoved * Radius), FLinearColor::Red, SDPG_Foreground, 0.f, 0.01f);
@@ -2120,7 +2121,8 @@ FVector2D UTransformGizmo::GetScreenRotateAxisDir(const FInputDeviceRay& InPress
 
 	static const TArray RotateIDs({	ETransformGizmoPartIdentifier::RotateXAxis,
 									ETransformGizmoPartIdentifier::RotateYAxis,
-									ETransformGizmoPartIdentifier::RotateZAxis});
+									ETransformGizmoPartIdentifier::RotateZAxis,
+									ETransformGizmoPartIdentifier::RotateScreenSpace});
 	const int32 RotateID = RotateIDs.IndexOfByKey(LastHitPart);
 	if (!ensure(RotateID != INDEX_NONE))
 	{
@@ -2131,7 +2133,7 @@ FVector2D UTransformGizmo::GetScreenRotateAxisDir(const FInputDeviceRay& InPress
 
 	// store world origin and axis
 	static const TArray RotateAxis({FVector::XAxisVector, FVector::YAxisVector, -FVector::ZAxisVector});
-	const FVector WorldAxis = GetWorldAxis(RotateAxis[RotateID]);
+	const FVector WorldAxis = RotateAxis.IsValidIndex(RotateID) ? GetWorldAxis(RotateAxis[RotateID]) : GizmoViewContext->GetViewDirection();
 	const FVector WorldOrigin = CurrentTransform.GetLocation();
 
 	// compute axis / view direction projection: is the rotation plane nearly perpendicular to the view plane?
@@ -2285,7 +2287,7 @@ void UTransformGizmo::OnClickReleaseRotateAxis(const FInputDeviceRay& InReleaseP
 	bTrySwitchingToNormalPull = false;
 }
 
-void UTransformGizmo::OnClickPressScreenSpaceRotate(const FInputDeviceRay& PressPos)
+void UTransformGizmo::OnClickPressScreenSpaceRotate(const FInputDeviceRay& InPressPos)
 {
 	check(GizmoViewContext);
 	
@@ -2295,28 +2297,70 @@ void UTransformGizmo::OnClickPressScreenSpaceRotate(const FInputDeviceRay& Press
 	InteractionPlanarAxisY = GizmoViewContext->GetViewRight();
 	InteractionAxisList = EAxisList::Screen;
 
-	float HitDepth;
-	if (GetRayParamIntersectionWithInteractionPlane(PressPos, HitDepth))
+	bTrySwitchingToNormalPull = false;
+	bInInteraction = false;
+	bDebugRotate = false;
+
+	if (bIndirectManipulation)
+	{
+		InteractionScreenAxisDirection = GetScreenRotateAxisDir(InPressPos);
+		InteractionScreenStartPos = InteractionScreenCurrPos = InPressPos.ScreenPosition;
+		bTrySwitchingToNormalPull = GizmoLocals::bAutoSwitchPullAxis;
+		bDebugRotate = true;
+	}
+	else
 	{
-		FVector HitPoint = PressPos.WorldRay.Origin + PressPos.WorldRay.Direction * HitDepth;
+		float HitDepth;
+		if (!GetRayParamIntersectionWithInteractionPlane(InPressPos, HitDepth))
+		{
+			return;
+		}
+		const FVector HitPoint = InPressPos.WorldRay.Origin + InPressPos.WorldRay.Direction * HitDepth;
 		InteractionStartAngle = GizmoMath::ComputeAngleInPlane(HitPoint,
 			InteractionPlanarOrigin, InteractionPlanarNormal, InteractionPlanarAxisX, InteractionPlanarAxisY);
 		InteractionCurrAngle = InteractionStartAngle;
-
-		bInInteraction = true;
-		
-		SetModeLastHitPart(EGizmoTransformMode::Rotate, LastHitPart);
 	}
+	
+	bInInteraction = true;
+	SetModeLastHitPart(EGizmoTransformMode::Rotate, LastHitPart);
 }
 
-void UTransformGizmo::OnClickDragScreenSpaceRotate(const FInputDeviceRay& DragPos)
+void UTransformGizmo::OnClickDragScreenSpaceRotate(const FInputDeviceRay& InDragPos)
 {
 	check(GizmoViewContext);
+
+	if (bIndirectManipulation)
+	{
+		FVector2D DragDir = InDragPos.ScreenPosition - InteractionScreenCurrPos;
+
+		if (bTrySwitchingToNormalPull)
+		{
+			const double DotTangent = FVector2D::DotProduct(InteractionScreenAxisDirection, DragDir);
+			const double DotNormal = FVector2D::DotProduct(NormalProjectionToRemove, DragDir);
+			if (FMath::Abs(DotNormal) > FMath::Abs(DotTangent))
+			{
+				::Swap(NormalProjectionToRemove, InteractionScreenAxisDirection);
+				::Swap(DebugDirection, DebugNormalRemoved);
+				InteractionScreenAxisDirection *= FMath::Sign(DotTangent) * FMath::Sign(DotNormal);
+			}
+			bTrySwitchingToNormalPull = false;
+		}
+	
+		const FVector2D DragDirToRemove = NormalProjectionToRemove * FVector2D::DotProduct(DragDir, NormalProjectionToRemove);
+		DragDir -= DragDirToRemove;
+
+		const double Delta = FVector2D::DotProduct(InteractionScreenAxisDirection, DragDir);
+		const FQuat DeltaRotAxis(InteractionPlanarNormal, FMath::DegreesToRadians(Delta));
+
+		ApplyRotateDelta(DeltaRotAxis);
+		InteractionScreenCurrPos = InDragPos.ScreenPosition;
+		return;
+	}
 	
 	float HitDepth;
-	if (GetRayParamIntersectionWithInteractionPlane(DragPos, HitDepth))
+	if (GetRayParamIntersectionWithInteractionPlane(InDragPos, HitDepth))
 	{
-		const FVector HitPoint = DragPos.WorldRay.Origin + DragPos.WorldRay.Direction * HitDepth;
+		const FVector HitPoint = InDragPos.WorldRay.Origin + InDragPos.WorldRay.Direction * HitDepth;
 		const float HitAngle = GizmoMath::ComputeAngleInPlane(HitPoint,
 			InteractionPlanarOrigin, InteractionPlanarNormal, InteractionPlanarAxisX, InteractionPlanarAxisY);
 
@@ -2336,6 +2380,8 @@ FQuat UTransformGizmo::ComputeAngularRotateDelta(double InStartAngle, double InE
 void UTransformGizmo::OnClickReleaseScreenSpaceRotate(const FInputDeviceRay& InReleasePos)
 {
 	bInInteraction = false;
+	bDebugRotate = false;
+	bTrySwitchingToNormalPull = false;
 }
 
 namespace ArcBallLocals
diff --git a/Engine/Source/Editor/Experimental/EditorInteractiveToolsFramework/Public/EditorGizmos/TransformGizmo.h b/Engine/Source/Editor/Experimental/EditorInteractiveToolsFramework/Public/EditorGizmos/TransformGizmo.h
--- a/Engine/Source/Editor/Experimental/EditorInteractiveToolsFramework/Public/EditorGizmos/TransformGizmo.h
+++ b/Engine/Source/Editor/Experimental/EditorInteractiveToolsFramework/Public/EditorGizmos/TransformGizmo.h
@@ -605,7 +605,7 @@ protected:
 	virtual void OnClickPressScreenSpaceRotate(const FInputDeviceRay& PressPos);
 
 	/** Handle click drag for screen-space rotate */
-	virtual void OnClickDragScreenSpaceRotate(const FInputDeviceRay& DragPos);
+	virtual void OnClickDragScreenSpaceRotate(const FInputDeviceRay& InDragPos);
 
 	/** Handle click release for screen-space rotate */
 	virtual void OnClickReleaseScreenSpaceRotate(const FInputDeviceRay& ReleasePos);
