New TRS Gizmo: set AutoSwitchPullAxis  to true by default
    - follow up to 35055294

#rnx

diff --git a/Engine/Source/Editor/Experimental/EditorInteractiveToolsFramework/Private/EditorGizmos/TransformGizmo.cpp b/Engine/Source/Editor/Experimental/EditorInteractiveToolsFramework/Private/EditorGizmos/TransformGizmo.cpp
--- a/Engine/Source/Editor/Experimental/EditorInteractiveToolsFramework/Private/EditorGizmos/TransformGizmo.cpp
+++ b/Engine/Source/Editor/Experimental/EditorInteractiveToolsFramework/Private/EditorGizmos/TransformGizmo.cpp
@@ -52,12 +52,6 @@ static FAutoConsoleVariableRef CVarProjectIndirect(
 	TEXT("Project to the nearest point of the curve when handling indirect rotation.")
 	);
 
-static bool	bAutoSwitchPullAxis = false;
-static FAutoConsoleVariableRef CVarAutoDetect(
-	TEXT("Gizmos.AutoSwitchPullAxis"),
-	bAutoSwitchPullAxis,
-	TEXT("Switch from tangential to normal projection based on the first mouse drag.")
-	);
 }
 
 void UTransformGizmo::SetDisallowNegativeScaling(bool bDisallow)
@@ -2087,7 +2081,7 @@ void UTransformGizmo::OnClickPressRotateAxis(const FInputDeviceRay& InPressPos)
 		}
 	}
 
-	bTrySwitchingToNormalPull = GizmoLocals::bAutoSwitchPullAxis && bIndirectManipulation && RotateMode == EAxisRotateMode::Pull;
+	bTrySwitchingToNormalPull = bIndirectManipulation && RotateMode == EAxisRotateMode::Pull;
 	
 	bInInteraction = true;
 	SetModeLastHitPart(EGizmoTransformMode::Rotate, LastHitPart);
@@ -2305,7 +2299,7 @@ void UTransformGizmo::OnClickPressScreenSpaceRotate(const FInputDeviceRay& InPre
 	{
 		InteractionScreenAxisDirection = GetScreenRotateAxisDir(InPressPos);
 		InteractionScreenStartPos = InteractionScreenCurrPos = InPressPos.ScreenPosition;
-		bTrySwitchingToNormalPull = GizmoLocals::bAutoSwitchPullAxis;
+		bTrySwitchingToNormalPull = true;
 		bDebugRotate = true;
 	}
 	else
