New TRS Gizmo: do not render in game mode

#rnx

diff --git a/Engine/Source/Editor/Experimental/EditorInteractiveToolsFramework/Private/EditorGizmos/EditorTransformGizmo.cpp b/Engine/Source/Editor/Experimental/EditorInteractiveToolsFramework/Private/EditorGizmos/EditorTransformGizmo.cpp
--- a/Engine/Source/Editor/Experimental/EditorInteractiveToolsFramework/Private/EditorGizmos/EditorTransformGizmo.cpp
+++ b/Engine/Source/Editor/Experimental/EditorInteractiveToolsFramework/Private/EditorGizmos/EditorTransformGizmo.cpp
@@ -17,8 +17,13 @@ DEFINE_LOG_CATEGORY_STATIC(LogEditorTransformGizmo, Log, All);
 void UEditorTransformGizmo::Render(IToolsContextRenderAPI* RenderAPI)
 {
 	const FSceneView* SceneView = RenderAPI ? RenderAPI->GetSceneView() : nullptr;
-	const bool bEngineShowFlagsModeWidget = SceneView && SceneView->Family &&
-											SceneView->Family->EngineShowFlags.ModeWidgets;
+	const FSceneViewFamily* SceneViewFamily = SceneView ? SceneView->Family : nullptr;
+	if (!SceneViewFamily)
+	{
+		return;
+	}
+
+	const bool bEngineShowFlagsModeWidget = SceneViewFamily->EngineShowFlags.ModeWidgets && !SceneViewFamily->EngineShowFlags.Game;
 	if (bEngineShowFlagsModeWidget)
 	{
 		Super::Render(RenderAPI);
