New TRS Gizmo: fixed interaction with hidden gizmo
    - indirect manipulation is still enabled even if hidden

#rnx

diff --git a/Engine/Source/Editor/Experimental/EditorInteractiveToolsFramework/Private/EditorGizmos/EditorTransformGizmoSource.cpp b/Engine/Source/Editor/Experimental/EditorInteractiveToolsFramework/Private/EditorGizmos/EditorTransformGizmoSource.cpp
--- a/Engine/Source/Editor/Experimental/EditorInteractiveToolsFramework/Private/EditorGizmos/EditorTransformGizmoSource.cpp
+++ b/Engine/Source/Editor/Experimental/EditorInteractiveToolsFramework/Private/EditorGizmos/EditorTransformGizmoSource.cpp
@@ -73,6 +73,13 @@ bool UEditorTransformGizmoSource::GetVisible() const
 		{
 			return false;
 		}
+
+		const bool bEngineShowFlagsModeWidget = ViewportClient->EngineShowFlags.ModeWidgets && !ViewportClient->EngineShowFlags.Game;
+		if (!bEngineShowFlagsModeWidget)
+		{
+			return false;
+		}
+		
 		return CanInteract();
 	}
 	return false;
diff --git a/Engine/Source/Editor/Experimental/EditorInteractiveToolsFramework/Private/EditorGizmos/TransformGizmo.cpp b/Engine/Source/Editor/Experimental/EditorInteractiveToolsFramework/Private/EditorGizmos/TransformGizmo.cpp
--- a/Engine/Source/Editor/Experimental/EditorInteractiveToolsFramework/Private/EditorGizmos/TransformGizmo.cpp
+++ b/Engine/Source/Editor/Experimental/EditorInteractiveToolsFramework/Private/EditorGizmos/TransformGizmo.cpp
@@ -516,7 +516,7 @@ FInputRayHit UTransformGizmo::CanBeginClickDragSequence(const FInputDeviceRay& P
 {
 	FInputRayHit RayHit;
 		
-	if (CanInteract() && HitTarget)
+	if (IsVisible() && CanInteract() && HitTarget)
 	{
 		RayHit = HitTarget->IsHit(PressPos);
 		ETransformGizmoPartIdentifier HitPart;
