New TRS Gizmo: expose focussed / hovered context to the visibility / interaction functions

diff --git a/Engine/Source/Editor/Experimental/EditorInteractiveToolsFramework/Private/EditorGizmos/EditorTransformGizmoSource.cpp b/Engine/Source/Editor/Experimental/EditorInteractiveToolsFramework/Private/EditorGizmos/EditorTransformGizmoSource.cpp
--- a/Engine/Source/Editor/Experimental/EditorInteractiveToolsFramework/Private/EditorGizmos/EditorTransformGizmoSource.cpp
+++ b/Engine/Source/Editor/Experimental/EditorInteractiveToolsFramework/Private/EditorGizmos/EditorTransformGizmoSource.cpp
@@ -65,29 +65,32 @@ float UEditorTransformGizmoSource::GetGizmoScale() const
 	return GetModeTools().GetWidgetScale();
 }
 
-bool UEditorTransformGizmoSource::GetVisible() const
+bool UEditorTransformGizmoSource::GetVisible(const EViewportContext InViewportContext) const
 {
-	if (const FEditorViewportClient* ViewportClient = GetViewportClient()) 
+	if (const FEditorViewportClient* ViewportClient = GetViewportClient(InViewportContext)) 
 	{
 		if (!ViewportClient->GetShowWidget())
 		{
 			return false;
 		}
 
-		const bool bEngineShowFlagsModeWidget = ViewportClient->EngineShowFlags.ModeWidgets && !ViewportClient->EngineShowFlags.Game;
-		if (!bEngineShowFlagsModeWidget)
+		if (InViewportContext == EViewportContext::Hovered)
 		{
-			return false;
+			const bool bEngineShowFlagsModeWidget = ViewportClient->EngineShowFlags.ModeWidgets && !ViewportClient->EngineShowFlags.Game;
+			if (!bEngineShowFlagsModeWidget)
+			{
+				return false;
+			}
 		}
 		
-		return CanInteract();
+		return CanInteract(InViewportContext);
 	}
 	return false;
 }
 
-bool UEditorTransformGizmoSource::CanInteract() const
+bool UEditorTransformGizmoSource::CanInteract(const EViewportContext InViewportContext) const
 {
-	if (const FEditorViewportClient* ViewportClient = GetViewportClient())
+	if (const FEditorViewportClient* ViewportClient = GetViewportClient(InViewportContext))
 	{
 		const FEditorModeTools& ModeTools = GetModeTools();
 		if (ModeTools.GetShowWidget() && ModeTools.UsesTransformWidget())
@@ -130,7 +133,14 @@ const FEditorModeTools& UEditorTransformGizmoSource::GetModeTools() const
 	return GLevelEditorModeTools();
 }
 
-const FEditorViewportClient* UEditorTransformGizmoSource::GetViewportClient() const
+const FEditorViewportClient* UEditorTransformGizmoSource::GetViewportClient(const EViewportContext InViewportContext) const
 {
+	if (InViewportContext == EViewportContext::Hovered)
+	{
+		if (FEditorViewportClient* Hovered = GetModeTools().GetHoveredViewportClient())
+		{
+			return Hovered;
+		}
+	}
 	return GetModeTools().GetFocusedViewportClient();
 }
\ No newline at end of file
diff --git a/Engine/Source/Editor/Experimental/EditorInteractiveToolsFramework/Private/EditorGizmos/TransformGizmo.cpp b/Engine/Source/Editor/Experimental/EditorInteractiveToolsFramework/Private/EditorGizmos/TransformGizmo.cpp
--- a/Engine/Source/Editor/Experimental/EditorInteractiveToolsFramework/Private/EditorGizmos/TransformGizmo.cpp
+++ b/Engine/Source/Editor/Experimental/EditorInteractiveToolsFramework/Private/EditorGizmos/TransformGizmo.cpp
@@ -357,7 +357,7 @@ FTransform UTransformGizmo::GetGizmoTransform() const
 
 void UTransformGizmo::Render(IToolsContextRenderAPI* RenderAPI)
 {
-	if (IsVisible() && GizmoElementRoot && RenderAPI)
+	if (CanInteract() && GizmoElementRoot && RenderAPI)
 	{
 		CurrentTransform = ActiveTarget->GetTransform();
 
@@ -413,7 +413,7 @@ void UTransformGizmo::OnEndHover()
 
 FInputRayHit UTransformGizmo::UpdateHoveredPart(const FInputDeviceRay& PressPos)
 {
-	if (!HitTarget || !IsVisible())
+	if (!HitTarget || !IsVisible(EViewportContext::Hovered))
 	{
 		return FInputRayHit();
 	}
@@ -516,7 +516,7 @@ FInputRayHit UTransformGizmo::CanBeginClickDragSequence(const FInputDeviceRay& P
 {
 	FInputRayHit RayHit;
 		
-	if (IsVisible() && CanInteract() && HitTarget)
+	if (IsVisible(EViewportContext::Hovered) && CanInteract() && HitTarget)
 	{
 		RayHit = HitTarget->IsHit(PressPos);
 		ETransformGizmoPartIdentifier HitPart;
@@ -846,21 +846,21 @@ void UTransformGizmo::UpdateCameraAxisSource()
 	}
 }
 
-bool UTransformGizmo::IsVisible() const
+bool UTransformGizmo::IsVisible(const EViewportContext InViewportContext) const
 {
 	if (TransformGizmoSource)
 	{
-		return bVisible && TransformGizmoSource->GetVisible();
+		return bVisible && TransformGizmoSource->GetVisible(InViewportContext);
 	}	
 	return bVisible;
 }
 
-bool UTransformGizmo::CanInteract() const
+bool UTransformGizmo::CanInteract(const EViewportContext InViewportContext) const
 {
 	const bool bValidMode = CurrentMode > EGizmoTransformMode::None && CurrentMode < EGizmoTransformMode::Max;
 	if (TransformGizmoSource)
 	{
-		return bValidMode && TransformGizmoSource->CanInteract(); 
+		return bValidMode && TransformGizmoSource->CanInteract(InViewportContext);
 	}
 	return bValidMode && bVisible;
 }
diff --git a/Engine/Source/Editor/Experimental/EditorInteractiveToolsFramework/Public/EditorGizmos/EditorTransformGizmoSource.h b/Engine/Source/Editor/Experimental/EditorInteractiveToolsFramework/Public/EditorGizmos/EditorTransformGizmoSource.h
--- a/Engine/Source/Editor/Experimental/EditorInteractiveToolsFramework/Public/EditorGizmos/EditorTransformGizmoSource.h
+++ b/Engine/Source/Editor/Experimental/EditorInteractiveToolsFramework/Public/EditorGizmos/EditorTransformGizmoSource.h
@@ -62,13 +62,13 @@ public:
 	/**
 	 * Whether the gizmo is visible
 	 */
-	virtual bool GetVisible() const override;
+	virtual bool GetVisible(const EViewportContext InViewportContext = EViewportContext::Focused) const override;
 
 	/* 
 	 * Returns whether the gizmo can interact.
 	 * Note that this can be true even if the gizmo is hidden to support indirect manipulation in game mode.
 	 */
-	virtual bool CanInteract() const override;
+	virtual bool CanInteract(const EViewportContext InViewportContext = EViewportContext::Focused) const override;
 
 	/**
  	 * Get current scale type
@@ -83,7 +83,7 @@ protected:
 
 	const FEditorModeTools& GetModeTools() const;
 
-	const FEditorViewportClient* GetViewportClient() const;
+	const FEditorViewportClient* GetViewportClient(const EViewportContext InViewportContext = EViewportContext::Focused) const;
 	
 	TWeakObjectPtr<const UEditorTransformGizmoContextObject> WeakContext = nullptr;
 };
diff --git a/Engine/Source/Editor/Experimental/EditorInteractiveToolsFramework/Public/EditorGizmos/TransformGizmo.h b/Engine/Source/Editor/Experimental/EditorInteractiveToolsFramework/Public/EditorGizmos/TransformGizmo.h
--- a/Engine/Source/Editor/Experimental/EditorInteractiveToolsFramework/Public/EditorGizmos/TransformGizmo.h
+++ b/Engine/Source/Editor/Experimental/EditorInteractiveToolsFramework/Public/EditorGizmos/TransformGizmo.h
@@ -715,10 +715,10 @@ protected:
 	virtual bool VerifyPartIdentifier(uint32 InPartIdentifier) const;
 
 	// Returns whether the gizmo is visible in the viewport.
-	bool IsVisible() const;
+	bool IsVisible(const EViewportContext InViewportContext = EViewportContext::Focused) const;
 
 	// Returns whether the gizmo can interact. Note that this can be true even if the gizmo is not visible to support indirect manipulation.
-	bool CanInteract() const;
+	bool CanInteract(const EViewportContext InViewportContext = EViewportContext::Focused) const;
 	
 protected:
 
diff --git a/Engine/Source/Editor/Experimental/EditorInteractiveToolsFramework/Public/EditorGizmos/TransformGizmoInterfaces.h b/Engine/Source/Editor/Experimental/EditorInteractiveToolsFramework/Public/EditorGizmos/TransformGizmoInterfaces.h
--- a/Engine/Source/Editor/Experimental/EditorInteractiveToolsFramework/Public/EditorGizmos/TransformGizmoInterfaces.h
+++ b/Engine/Source/Editor/Experimental/EditorInteractiveToolsFramework/Public/EditorGizmos/TransformGizmoInterfaces.h
@@ -29,6 +29,12 @@ enum class EGizmoTransformScaleType : uint8
 	PercentageBased
 };
 
+enum class EViewportContext : uint8
+{
+	Hovered,
+	Focused
+};
+
 UINTERFACE()
 class EDITORINTERACTIVETOOLSFRAMEWORK_API UTransformGizmoSource : public UInterface
 {
@@ -65,12 +71,12 @@ public:
 	/* 
 	 * Returns whether the gizmo should be visible. 
 	 */
-	virtual bool GetVisible() const PURE_VIRTUAL(ITransformGizmoSource::GetVisible, return false;);
+	virtual bool GetVisible(const EViewportContext InViewportContext = EViewportContext::Focused) const PURE_VIRTUAL(ITransformGizmoSource::GetVisible, return false;);
 
 	/* 
 	 * Returns whether the gizmo can interact.
 	 */
-	virtual bool CanInteract() const PURE_VIRTUAL(ITransformGizmoSource::CanInteract, return false;);
+	virtual bool CanInteract(const EViewportContext InViewportContext = EViewportContext::Focused) const PURE_VIRTUAL(ITransformGizmoSource::CanInteract, return false;);
 
 	/*
 	 * Get current scale type
