make interactive tools context also log to console when displaying tool user info or warning messages, since it's easy for the tool messages to be lost (overwritten by other messages, or discarded on tool exit)

+ give the logs a category other than logtemp

diff --git a/Engine/Source/Editor/Experimental/EditorInteractiveToolsFramework/Private/EdModeInteractiveToolsContext.cpp b/Engine/Source/Editor/Experimental/EditorInteractiveToolsFramework/Private/EdModeInteractiveToolsContext.cpp
--- a/Engine/Source/Editor/Experimental/EditorInteractiveToolsFramework/Private/EdModeInteractiveToolsContext.cpp
+++ b/Engine/Source/Editor/Experimental/EditorInteractiveToolsFramework/Private/EdModeInteractiveToolsContext.cpp
@@ -40,6 +40,8 @@
 
 //#define ENABLE_DEBUG_PRINTING
 
+DEFINE_LOG_CATEGORY_STATIC(LogInteractiveToolsContext, Log, All);
+
 class FEdModeToolsContextQueriesImpl : public IToolsContextQueriesAPI
 {
 public:
@@ -219,15 +221,17 @@ public:
 	{
 		if (Level == EToolMessageLevel::UserNotification || Level == EToolMessageLevel::UserMessage)
 		{
+			UE_LOG(LogInteractiveToolsContext, Display, TEXT("%s"), *Message.ToString());
 			ToolsContext->PostToolNotificationMessage(Message);
 		}
 		else if (Level == EToolMessageLevel::UserWarning || Level == EToolMessageLevel::UserError)
 		{
+			UE_LOG(LogInteractiveToolsContext, Warning, TEXT("%s"), *Message.ToString());
 			ToolsContext->PostToolWarningMessage(Message);
 		}
 		else
 		{
-			UE_LOG(LogTemp, Warning, TEXT("%s"), *Message.ToString());
+			UE_LOG(LogInteractiveToolsContext, Warning, TEXT("%s"), *Message.ToString());
 		}
 	}
 
@@ -933,11 +937,11 @@ void UModeManagerInteractiveToolsContext::UpdateStateWithoutRoutingInputKey(FEdi
 bool UModeManagerInteractiveToolsContext::InputKey(FEditorViewportClient* ViewportClient, FViewport* Viewport, FKey Key, EInputEvent Event)
 {
 #ifdef ENABLE_DEBUG_PRINTING
-	if (Event == IE_Pressed) { UE_LOG(LogTemp, Warning, TEXT("PRESSED EVENT")); }
-	else if (Event == IE_Released) { UE_LOG(LogTemp, Warning, TEXT("RELEASED EVENT")); }
-	else if (Event == IE_Repeat) { UE_LOG(LogTemp, Warning, TEXT("REPEAT EVENT")); }
-	else if (Event == IE_Axis) { UE_LOG(LogTemp, Warning, TEXT("AXIS EVENT")); }
-	else if (Event == IE_DoubleClick) { UE_LOG(LogTemp, Warning, TEXT("DOUBLECLICK EVENT")); }
+	if (Event == IE_Pressed) { UE_LOG(LogInteractiveToolsContext, Warning, TEXT("PRESSED EVENT")); }
+	else if (Event == IE_Released) { UE_LOG(LogInteractiveToolsContext, Warning, TEXT("RELEASED EVENT")); }
+	else if (Event == IE_Repeat) { UE_LOG(LogInteractiveToolsContext, Warning, TEXT("REPEAT EVENT")); }
+	else if (Event == IE_Axis) { UE_LOG(LogInteractiveToolsContext, Warning, TEXT("AXIS EVENT")); }
+	else if (Event == IE_DoubleClick) { UE_LOG(LogInteractiveToolsContext, Warning, TEXT("DOUBLECLICK EVENT")); }
 #endif
 
 	// Update the current state, then route result
@@ -1051,7 +1055,7 @@ bool UModeManagerInteractiveToolsContext::InputKey(FEditorViewportClient* Viewpo
 bool UModeManagerInteractiveToolsContext::MouseEnter(FEditorViewportClient* ViewportClient, FViewport* Viewport, int32 x, int32 y)
 {
 #ifdef ENABLE_DEBUG_PRINTING
-	UE_LOG(LogTemp, Warning, TEXT("MOUSE ENTER"));
+	UE_LOG(LogInteractiveToolsContext, Warning, TEXT("MOUSE ENTER"));
 #endif
 
 	CurrentMouseState.Mouse.Position2D = FVector2D(x, y);
@@ -1064,7 +1068,7 @@ bool UModeManagerInteractiveToolsContext::MouseEnter(FEditorViewportClient* View
 bool UModeManagerInteractiveToolsContext::MouseMove(FEditorViewportClient* ViewportClient, FViewport* Viewport, int32 x, int32 y)
 {
 #ifdef ENABLE_DEBUG_PRINTING
-	UE_LOG(LogTemp, Warning, TEXT("HOVER %p"), ViewportClient);
+	UE_LOG(LogInteractiveToolsContext, Warning, TEXT("HOVER %p"), ViewportClient);
 #endif
 
 	CurrentMouseState.Mouse.Position2D = FVector2D(x, y);
@@ -1097,7 +1101,7 @@ bool UModeManagerInteractiveToolsContext::MouseMove(FEditorViewportClient* Viewp
 bool UModeManagerInteractiveToolsContext::MouseLeave(FEditorViewportClient* ViewportClient, FViewport* Viewport)
 {
 #ifdef ENABLE_DEBUG_PRINTING
-	UE_LOG(LogTemp, Warning, TEXT("MOUSE LEAVE"));
+	UE_LOG(LogInteractiveToolsContext, Warning, TEXT("MOUSE LEAVE"));
 #endif
 
 	return false;
@@ -1120,7 +1124,7 @@ bool UModeManagerInteractiveToolsContext::CapturedMouseMove(FEditorViewportClien
 	if (InputRouter->HasActiveMouseCapture())
 	{
 #ifdef ENABLE_DEBUG_PRINTING
-		UE_LOG(LogTemp, Warning, TEXT("CAPTURED MOUSE MOVE"));
+		UE_LOG(LogInteractiveToolsContext, Warning, TEXT("CAPTURED MOUSE MOVE"));
 #endif
 
 		FInputDeviceState InputState = CurrentMouseState;
