Modify FEditor2DScrollBehaviorTarget to work in orthographic viewports

diff --git a/Engine/Source/Editor/Experimental/EditorInteractiveToolsFramework/Private/Behaviors/2DViewportBehaviorTargets.cpp b/Engine/Source/Editor/Experimental/EditorInteractiveToolsFramework/Private/Behaviors/2DViewportBehaviorTargets.cpp
--- a/Engine/Source/Editor/Experimental/EditorInteractiveToolsFramework/Private/Behaviors/2DViewportBehaviorTargets.cpp
+++ b/Engine/Source/Editor/Experimental/EditorInteractiveToolsFramework/Private/Behaviors/2DViewportBehaviorTargets.cpp
@@ -10,47 +10,99 @@ FEditor2DScrollBehaviorTarget::FEditor2DScrollBehaviorTarget(FEditorViewportClie
 
 FInputRayHit FEditor2DScrollBehaviorTarget::CanBeginClickDragSequence(const FInputDeviceRay& PressPos)
 {
-	// Verify that ray is facing the proper direction
-	if (PressPos.WorldRay.Direction.Z * PressPos.WorldRay.Origin.Z < 0)
+	if (ViewportClient->IsOrtho())
 	{
-		return FInputRayHit(static_cast<float>(-PressPos.WorldRay.Origin.Z / PressPos.WorldRay.Direction.Z));
+		return FInputRayHit(TNumericLimits<float>::Max()); // bHit is true. Depth is max to lose the standard tiebreaker
+	}
+	else
+	{
+		// Verify that ray is facing the proper direction
+		if (PressPos.WorldRay.Direction.Z * PressPos.WorldRay.Origin.Z < 0)
+		{
+			return FInputRayHit(static_cast<float>(-PressPos.WorldRay.Origin.Z / PressPos.WorldRay.Direction.Z));
+		}
 	}
 	return FInputRayHit();
 }
 
 void FEditor2DScrollBehaviorTarget::OnClickPress(const FInputDeviceRay& PressPos)
 {
-	if (ensure(PressPos.WorldRay.Direction.Z * PressPos.WorldRay.Origin.Z < 0))
+	if (ViewportClient->IsOrtho())
+	{
+		DragStart = FVector3d(PressPos.ScreenPosition.X, PressPos.ScreenPosition.Y, 0.0);
+		OriginalCameraLocation = ViewportClient->GetViewLocation();
+	}
+	else
 	{
-		// Intersect with XY plane
-		double T = -PressPos.WorldRay.Origin.Z / PressPos.WorldRay.Direction.Z;
-		DragStart = FVector3d(
-			PressPos.WorldRay.Origin.X + T * PressPos.WorldRay.Direction.X,
-			PressPos.WorldRay.Origin.Y + T * PressPos.WorldRay.Direction.Y,
-			0);
-
-		OriginalCameraLocation = (FVector3d)ViewportClient->GetViewLocation();
+		if (ensure(PressPos.WorldRay.Direction.Z * PressPos.WorldRay.Origin.Z < 0))
+		{
+			// Intersect with XY plane
+			double T = -PressPos.WorldRay.Origin.Z / PressPos.WorldRay.Direction.Z;
+			DragStart = FVector3d(
+				PressPos.WorldRay.Origin.X + T * PressPos.WorldRay.Direction.X,
+				PressPos.WorldRay.Origin.Y + T * PressPos.WorldRay.Direction.Y,
+				0);
+
+			OriginalCameraLocation = (FVector3d)ViewportClient->GetViewLocation();
+		}
 	}
 }
 
 void FEditor2DScrollBehaviorTarget::OnClickDrag(const FInputDeviceRay& DragPos)
 {
-	if (ensure(DragPos.WorldRay.Direction.Z * DragPos.WorldRay.Origin.Z < 0))
+	if (ViewportClient->IsPerspective() || ViewportClient->GetViewportType() == LVT_OrthoFreelook)
+	{
+		if (ensure(DragPos.WorldRay.Direction.Z * DragPos.WorldRay.Origin.Z < 0))
+		{
+			// Intersect a ray starting from the original position and using the new
+			// ray direction. I.e., pretend the camera is not moving.
+
+			double T = -OriginalCameraLocation.Z / DragPos.WorldRay.Direction.Z;
+			FVector3d DragEnd = FVector3d(
+				OriginalCameraLocation.X + T * DragPos.WorldRay.Direction.X,
+				OriginalCameraLocation.Y + T * DragPos.WorldRay.Direction.Y,
+				0);
+
+			// We want to make it look like we are sliding the plane such that DragStart
+			// ends up on DragEnd. For that, our camera will be moving the opposite direction.
+			FVector3d CameraDisplacement = DragStart - DragEnd;
+			check(CameraDisplacement.Z == 0);
+			ViewportClient->SetViewLocation((FVector)(OriginalCameraLocation + CameraDisplacement));
+		}
+	}
+	else
 	{
-		// Intersect a ray starting from the original position and using the new
-		// ray direction. I.e., pretend the camera is not moving.
-
-		double T = -OriginalCameraLocation.Z / DragPos.WorldRay.Direction.Z;
-		FVector3d DragEnd = FVector3d(
-			OriginalCameraLocation.X + T * DragPos.WorldRay.Direction.X,
-			OriginalCameraLocation.Y + T * DragPos.WorldRay.Direction.Y,
-			0);
-
-		// We want to make it look like we are sliding the plane such that DragStart
-		// ends up on DragEnd. For that, our camera will be moving the opposite direction.
-		FVector3d CameraDisplacement = DragStart - DragEnd;
-		check(CameraDisplacement.Z == 0);
-		ViewportClient->SetViewLocation((FVector)(OriginalCameraLocation + CameraDisplacement));
+		const float UnitsPerPixel = ViewportClient->GetOrthoUnitsPerPixel(ViewportClient->Viewport) * ViewportClient->GetDPIScale();
+		const double DX = DragPos.ScreenPosition.X - DragStart.X;
+		const double DY = DragPos.ScreenPosition.Y - DragStart.Y;
+
+		FVector CameraDisplacement(0.0, 0.0, 0.0);
+		switch (ViewportClient->GetViewportType())
+		{
+		case LVT_OrthoXY:
+			CameraDisplacement = FVector(-DX * UnitsPerPixel, -DY * UnitsPerPixel, 0.0);
+			break;
+		case LVT_OrthoXZ:
+			CameraDisplacement = FVector(-DX * UnitsPerPixel, 0.0, DY * UnitsPerPixel);
+			break;
+		case LVT_OrthoYZ:
+			CameraDisplacement = FVector(0.0, -DX * UnitsPerPixel, DY * UnitsPerPixel);
+			break;
+		case LVT_OrthoNegativeXY:
+			CameraDisplacement = FVector(DX * UnitsPerPixel, -DY * UnitsPerPixel, 0.0);
+			break;
+		case LVT_OrthoNegativeXZ:
+			CameraDisplacement = FVector(DX * UnitsPerPixel, 0.0, DY * UnitsPerPixel);
+			break;
+		case LVT_OrthoNegativeYZ:
+			CameraDisplacement = FVector(0.0, DX * UnitsPerPixel, DY * UnitsPerPixel);
+			break;
+		default:
+			checkf(false, TEXT("Unexpected viewport type in FEditor2DScrollBehaviorTarget"));
+			break;
+		}
+
+		ViewportClient->SetViewLocation(OriginalCameraLocation + CameraDisplacement);
 	}
 }
 
