New TRS Gizmo: unset parent coord system when switching back to the legacy one

#rnx

diff --git a/Engine/Source/Editor/Experimental/EditorInteractiveToolsFramework/Private/EditorGizmos/EditorTransformGizmoUtil.cpp b/Engine/Source/Editor/Experimental/EditorInteractiveToolsFramework/Private/EditorGizmos/EditorTransformGizmoUtil.cpp
--- a/Engine/Source/Editor/Experimental/EditorInteractiveToolsFramework/Private/EditorGizmos/EditorTransformGizmoUtil.cpp
+++ b/Engine/Source/Editor/Experimental/EditorInteractiveToolsFramework/Private/EditorGizmos/EditorTransformGizmoUtil.cpp
@@ -286,6 +286,12 @@ void UEditorTransformGizmoContextObject::InitializeGizmoManagerBinding()
 		{
 			// swap back default mode
 			(void)SwapDefaultMode(FAssetEdModes::EM_AssetDefault, FBuiltinEditorModes::EM_Default);
+
+			// reset the coordinate system if the legacy gizmo doesn't support it
+			if (ModeTools->GetCoordSystem() > COORD_Local)
+			{
+				ModeTools->SetCoordSystem(COORD_Local);
+			}
 			
 			// remove viewports' binding + gizmos as they are useless
 			const UModeManagerInteractiveToolsContext* ToolsContext = ModeTools->GetInteractiveToolsContext();
