New TRS Gizmo: explicit coordinate system

The explicit mode provides consistent information between what the manipulator displays and the values it drives. This is also reflected in the Anim Details and the Curve Editor, which reflect the same value changes.
In rotation, this corresponds to what might be called "gimbal mode" in other dcc, and takes into account the rotation order if exposed (ControlRig is currently the only place where that's the case)
For translation and scale, it corresponds to a parent-relative coordinate system (or offset space in the case of ControlRig)

The coordinate system is only exposed if the new trs gizmo is active and if New TRS Gizmo > Enable Explicit is set to true in the preferences (false by default for now)

diff --git a/Engine/Plugins/Animation/ControlRig/Source/ControlRigEditor/ControlRigEditor.Build.cs b/Engine/Plugins/Animation/ControlRig/Source/ControlRigEditor/ControlRigEditor.Build.cs
--- a/Engine/Plugins/Animation/ControlRig/Source/ControlRigEditor/ControlRigEditor.Build.cs
+++ b/Engine/Plugins/Animation/ControlRig/Source/ControlRigEditor/ControlRigEditor.Build.cs
@@ -11,8 +11,9 @@ namespace UnrealBuildTool.Rules
 				{
 					"MainFrame",
 					"AppFramework",
-                    "RigVMEditor", 
+                    "RigVMEditor",
 					"CurveEditor",
+                    "InteractiveToolsFramework"
 				}
 			);
 
diff --git a/Engine/Plugins/Animation/ControlRig/Source/ControlRigEditor/Private/EditMode/ControlRigEditMode.cpp b/Engine/Plugins/Animation/ControlRig/Source/ControlRigEditor/Private/EditMode/ControlRigEditMode.cpp
--- a/Engine/Plugins/Animation/ControlRig/Source/ControlRigEditor/Private/EditMode/ControlRigEditMode.cpp
+++ b/Engine/Plugins/Animation/ControlRig/Source/ControlRigEditor/Private/EditMode/ControlRigEditMode.cpp
@@ -2,6 +2,7 @@
 
 #include "EditMode/ControlRigEditMode.h"
 
+#include "AnimationCoreLibrary.h"
 #include "AnimationEditorPreviewActor.h"
 #include "Components/StaticMeshComponent.h"
 #include "EditMode/ControlRigEditModeToolkit.h"
@@ -72,6 +73,7 @@
 #include "Editor/ControlRigViewportToolbarExtensions.h"
 #include "Editor/Sequencer/Private/SSequencer.h"
 #include "Slate/SceneViewport.h"
+#include "EditorGizmos/EditorTransformGizmoUtil.h"
 #include "Tools/BakingHelper.h"
 #include "Sequencer/AnimLayers/AnimLayers.h"
 
@@ -497,6 +499,9 @@ void FControlRigEditMode::Enter()
 	);
 
 	RegisterPendingFocusMode();
+
+	// initialize the gizmo context
+	WeakGizmoContext = GetModeManager()->GetGizmoContext();
 }
 
 //todo get working with Persona
@@ -519,6 +524,12 @@ static void ClearOutAnyActiveTools()
 
 void FControlRigEditMode::Exit()
 {
+	if (WeakGizmoContext.IsValid())
+	{
+		WeakGizmoContext->RotationContext = FRotationContext();
+		WeakGizmoContext.Reset();
+	}
+	
 	UE::ControlRig::RemoveControlRigViewportToolbarExtensions();
 
 	UnregisterPendingFocusMode();
@@ -747,6 +758,8 @@ void FControlRigEditMode::Tick(FEditorViewportClient* ViewportClient, float Delt
 		}
 	}
 	DetailKeyFrameCache->UpdateIfDirty();
+	
+	UpdateRotationContext();
 }
 
 //Hit proxy for FK Rigs and bones.
@@ -2512,14 +2525,72 @@ bool FControlRigEditMode::ComputePivotFromEditedShape(UControlRig* InControlRig,
 	return true;
 }
 
-bool FControlRigEditMode::ComputePivotFromShapeActors(UControlRig* InControlRig, const bool bEachLocalSpace, const bool bIsParentSpace, FTransform& OutTransform) const
+FControlRigEditMode::ETransformSpaceSpace FControlRigEditMode::GetTransformSpace() const
+{
+	switch (GetCoordSystemSpace())
+	{
+		case COORD_Parent:
+			return ETransformSpaceSpace::Parent;
+		case COORD_Explicit:
+			return ETransformSpaceSpace::Explicit;
+		default:
+			break;
+	}
+	return ETransformSpaceSpace::Default;	
+};
+
+FTransform FControlRigEditMode::GetPivotOrientation(const FRigElementKey& InControlKey, URigHierarchy* InHierarchy, const ETransformSpaceSpace InSpace, const FTransform& InComponentTransform) const
+{
+	FRigControlElement* Control = InHierarchy->Find<FRigControlElement>(InControlKey);
+	if (!Control)
+	{
+		return FTransform::Identity;
+	}
+		
+	switch (InSpace)
+	{
+	case ETransformSpaceSpace::Default:
+		return InHierarchy->GetTransform(Control, ERigTransformType::CurrentGlobal);
+	case ETransformSpaceSpace::Parent:
+		{
+			const int32 NumParents = InHierarchy->GetNumberOfParents(Control);			
+			return NumParents > 0 ? InHierarchy->GetParentTransform(Control, ERigTransformType::CurrentGlobal) : InHierarchy->GetTransform(Control, ERigTransformType::CurrentGlobal);
+		}
+	case ETransformSpaceSpace::Explicit:
+		{
+			FRotationContext& RotationContext = GetRotationContext();
+					
+			const bool bRotating = GetModeManager()->GetWidgetMode() == UE::Widget::EWidgetMode::WM_Rotate;
+			const bool bUsePreferredRotationOrder = InHierarchy->GetUsePreferredRotationOrder(Control);
+
+			RotationContext.RotationOrder = bUsePreferredRotationOrder ? InHierarchy->GetControlPreferredEulerRotationOrder(Control) : EEulerRotationOrder::XYZ;
+			RotationContext.Rotation = InHierarchy->GetControlPreferredRotator(Control);
+
+			if (bRotating)
+			{
+				const FTransform Offset = InHierarchy->GetControlOffsetTransform(Control, ERigTransformType::CurrentGlobal);
+				RotationContext.Offset = Offset * InComponentTransform;
+				return InHierarchy->GetTransform(Control, ERigTransformType::CurrentGlobal);
+			}
+
+			RotationContext.Offset = FTransform::Identity;
+			return InHierarchy->GetControlOffsetTransform(Control, ERigTransformType::CurrentGlobal);
+		}
+	default:
+		break;
+	}
+
+	return FTransform::Identity;
+}
+
+bool FControlRigEditMode::ComputePivotFromShapeActors(UControlRig* InControlRig, const bool bEachLocalSpace, const ETransformSpaceSpace InSpace, FTransform& OutTransform) const
 {
 	if (!ensure(!bIsChangingControlShapeTransform))
 	{
 		return false;
 	}
 	
-	const URigHierarchy* Hierarchy = InControlRig ? InControlRig->GetHierarchy() : nullptr;
+	URigHierarchy* Hierarchy = InControlRig ? InControlRig->GetHierarchy() : nullptr;
 	if (!Hierarchy)
 	{
 		return false;
@@ -2540,10 +2611,8 @@ bool FControlRigEditMode::ComputePivotFromShapeActors(UControlRig* InControlRig,
 			if (IsValid(ShapeActor) && ShapeActor->IsSelected())
 			{
 				const FRigElementKey ControlKey = ShapeActor->GetElementKey();
-				bool bGetParentTransform = bIsParentSpace && Hierarchy->GetNumberOfParents(ControlKey);
-	
 				const FTransform ShapeTransform = ShapeActor->GetActorTransform().GetRelativeTransform(ComponentTransform);
-				LastTransform = bGetParentTransform ? Hierarchy->GetParentTransform(ControlKey) : ShapeTransform;
+				LastTransform = GetPivotOrientation(ControlKey, Hierarchy, InSpace, ComponentTransform);
 				PivotLocation += ShapeTransform.GetLocation();
 				
 				++NumSelectedControls;
@@ -2622,7 +2691,6 @@ void FControlRigEditMode::UpdatePivotTransforms()
 {
 	const UControlRigEditModeSettings* Settings = GetSettings();
 	const bool bEachLocalSpace = Settings && Settings->bLocalTransformsInEachLocalSpace;
-	const bool bIsParentSpace = GetCoordSystemSpace() == COORD_Parent;
 
 	PivotTransforms.Reset();
 
@@ -2641,7 +2709,7 @@ void FControlRigEditMode::UpdatePivotTransforms()
 				}
 				else
 				{
-					bAdd = ComputePivotFromShapeActors(ControlRig, bEachLocalSpace, bIsParentSpace, Pivot);			
+					bAdd = ComputePivotFromShapeActors(ControlRig, bEachLocalSpace, GetTransformSpace(), Pivot);			
 				}
 			}
 			else if (AreRigElementSelectedAndMovable(ControlRig))
@@ -2753,8 +2821,7 @@ void FControlRigEditMode::UpdatePivotTransformsIfNeeded(UControlRig* InControlRi
 		{
 			const UControlRigEditModeSettings* Settings = GetSettings();
 			const bool bEachLocalSpace = Settings && Settings->bLocalTransformsInEachLocalSpace;
-			const bool bIsParentSpace = GetCoordSystemSpace() == COORD_Parent;
-			ComputePivotFromShapeActors(InControlRig, bEachLocalSpace, bIsParentSpace, InOutTransform);			
+			ComputePivotFromShapeActors(InControlRig, bEachLocalSpace, GetTransformSpace(), InOutTransform);			
 		}
 	}
 	else if (AreRigElementSelectedAndMovable(InControlRig))
@@ -4913,20 +4980,92 @@ void FControlRigEditMode::SetControlShapeTransform(
 	const FTransform& InToWorldTransform,
 	const FRigControlModifiedContext& InContext,
 	const bool bPrintPython,
+	const FRotator& InDeltaRot,
 	const bool bFixEulerFlips) const
 {
 	UControlRig* ControlRig = InShapeActor->ControlRig.Get();
-	if (!ControlRig)
+	URigHierarchy* Hierarchy = ControlRig ? ControlRig->GetHierarchy() : nullptr;
+	if (!Hierarchy)
+	{
+		return;
+	}
+
+	FRigControlElement* ControlElement = ControlRig->FindControl(InShapeActor->ControlName);
+	if (!ControlElement)
 	{
 		return;
 	}
 
 	static constexpr bool bNotify = true, bUndo = true;
+	
+	const bool bIsExplicitRotation = GetTransformSpace() == ETransformSpaceSpace::Explicit && GetModeManager()->GetWidgetMode() == UE::Widget::EWidgetMode::WM_Rotate;
+	auto ApplyExplicitRotation = [bIsExplicitRotation, ControlRig, Hierarchy, ControlElement, InDeltaRot]()
+	{
+		if (!bIsExplicitRotation)
+		{
+			return false;
+		}
+		
+		if (InDeltaRot.IsZero())
+		{
+			return false;
+		}
+		
+		if (ControlRig->IsAdditive())
+		{
+			return false;
+		}
+
+		if (!Hierarchy->UsesPreferredEulerAngles())
+		{
+			return false;
+		}
+
+		const ERigControlType ControlType = ControlElement->Settings.ControlType;
+		return ControlType == ERigControlType::Rotator || ControlType == ERigControlType::EulerTransform;
+	};
+
+	auto SetControlRotation = [ControlRig, ControlElement, Hierarchy, InContext, bPrintPython](const FTransform& InGlobalTransform, const FRotator& InDeltaRot)
+	{
+		const FName ControlName = ControlElement->GetKey().Name;
+		const FVector DeltaEulerAngle(InDeltaRot.Roll, InDeltaRot.Pitch, InDeltaRot.Yaw);
+		FVector NewEulerAngle = Hierarchy->GetControlSpecifiedEulerAngle(ControlElement);
+		NewEulerAngle += DeltaEulerAngle;
+		
+		const ERigControlType ControlType = ControlElement->Settings.ControlType;
+		if (ControlType == ERigControlType::Rotator)
+		{
+			const FVector3f NewValue(NewEulerAngle.Y, NewEulerAngle.Z, NewEulerAngle.X);
+			
+			Hierarchy->SetControlSpecifiedEulerAngle(ControlElement, NewEulerAngle);
+			ControlRig->SetControlValue<FVector3f>(ControlName, NewValue, /*bNotify*/true, InContext, /*bUndo*/true);
+			Hierarchy->SetControlSpecifiedEulerAngle(ControlElement, NewEulerAngle);
+		}
+		else if (ControlType == ERigControlType::EulerTransform)
+		{
+			const FQuat Quat = Hierarchy->GetControlQuaternion(ControlElement, NewEulerAngle);
+			FRigControlValue NewValue = ControlRig->GetControlValueFromGlobalTransform(ControlName, InGlobalTransform, ERigTransformType::CurrentGlobal);
+			FEulerTransform EulerTransform = NewValue.Get<FRigControlValue::FEulerTransform_Float>().ToTransform();
+			EulerTransform.Rotation = FRotator(Quat);
+
+			Hierarchy->SetControlSpecifiedEulerAngle(ControlElement, NewEulerAngle);
+			ControlRig->SetControlValue<FRigControlValue::FEulerTransform_Float>(ControlName, EulerTransform, /*bNotify*/true, InContext, /*bUndo*/true, bPrintPython);
+			Hierarchy->SetControlSpecifiedEulerAngle(ControlElement, NewEulerAngle);			
+		}
+	};
+
 	if (AreEditingControlRigDirectly())
 	{
-		// assumes it's attached to actor
-		ControlRig->SetControlGlobalTransform(
-			InShapeActor->ControlName, InGlobalTransform, bNotify, InContext, bUndo, bPrintPython, bFixEulerFlips);
+		if (ApplyExplicitRotation())
+		{
+			SetControlRotation(InGlobalTransform, InDeltaRot);
+		}
+		else
+		{
+			// assumes it's attached to actor
+			ControlRig->SetControlGlobalTransform(
+				InShapeActor->ControlName, InGlobalTransform, bNotify, InContext, bUndo, bPrintPython, bFixEulerFlips);
+		}
 		return;
 	}
 
@@ -4953,9 +5092,16 @@ void FControlRigEditMode::SetControlShapeTransform(
 	// no need to compensate for constraints here, this will be done after when setting the control in the constraint space
 	{
 		TGuardValue<bool> CompensateGuard(FMovieSceneConstraintChannelHelper::bDoNotCompensate, true);
-		ControlRig->SetControlGlobalTransform(
-			InShapeActor->ControlName, InGlobalTransform, bNotify, InContext, bUndo, bPrintPython, bFixEulerFlips);
-		EvaluateRigIfAdditive();
+		if (ApplyExplicitRotation())
+		{
+			SetControlRotation(InGlobalTransform, InDeltaRot);
+		}
+		else
+		{
+			ControlRig->SetControlGlobalTransform(
+					InShapeActor->ControlName, InGlobalTransform, bNotify, InContext, bUndo, bPrintPython, bFixEulerFlips);
+			EvaluateRigIfAdditive();
+		}
 	}
 
 	if (!bNeedsConstraintPostProcess)
@@ -4980,9 +5126,11 @@ void FControlRigEditMode::SetControlShapeTransform(
 	ControlRig->SetControlLocalTransform(InShapeActor->ControlName, LocalTransform, bNotify, Context, bUndo, bFixEulerFlips);
 	EvaluateRigIfAdditive();
 
-	TGuardValue<bool> CompensateGuard(FMovieSceneConstraintChannelHelper::bDoNotCompensate, true);
-	const FConstraintsManagerController& Controller = FConstraintsManagerController::Get(ControlRig->GetWorld());
-	Controller.EvaluateAllConstraints();
+	{
+		TGuardValue<bool> CompensateGuard(FMovieSceneConstraintChannelHelper::bDoNotCompensate, true);
+		const FConstraintsManagerController& Controller = FConstraintsManagerController::Get(ControlRig->GetWorld());
+		Controller.EvaluateAllConstraints();
+	}
 }
 
 FTransform FControlRigEditMode::GetControlShapeTransform(const AControlRigShapeActor* ShapeActor)
@@ -5161,11 +5309,11 @@ void FControlRigEditMode::MoveControlShape(AControlRigShapeActor* ShapeActor, co
 		const FTransform ToWorldTransformCopy = ToWorldTransform;
 		
 		const TFunction<void()> SetControlShapeTask = [this, ShapeActor, ControlRig, ControlElement, NewTransform, ToWorldTransformCopy, Context,
-			bPrintPythonCommands, bIsTransientControl, bRotation, bQueueTasks]()
+			bPrintPythonCommands, bIsTransientControl, bRotation, InRot, bQueueTasks]()
 		{
 			//fix flips and do rotation orders only if not additive
 			const bool bFixEulerFlips = !ControlRig->IsAdditive() && bRotation;
-			SetControlShapeTransform(ShapeActor, NewTransform, ToWorldTransformCopy, Context, bPrintPythonCommands, bFixEulerFlips);
+			SetControlShapeTransform(ShapeActor, NewTransform, ToWorldTransformCopy, Context, bPrintPythonCommands, InRot, bFixEulerFlips);
 			NotifyDrivenControls(ControlRig, ShapeActor->GetElementKey(),Context);
 
 			if (ControlElement && !bIsTransientControl)
@@ -5978,6 +6126,60 @@ UControlRigEditModeSettings* FControlRigEditMode::GetMutableSettings() const
 	return WeakSettings.Get();
 }
 
+FRotationContext& FControlRigEditMode::GetRotationContext() const
+{
+	static FRotationContext DefaultContext;
+	return WeakGizmoContext.IsValid() ? WeakGizmoContext->RotationContext : DefaultContext;
+}
+
+void FControlRigEditMode::UpdateRotationContext()
+{
+	if (!WeakGizmoContext.IsValid())
+	{
+		return;
+	}
+	
+	FRotationContext& RotationContext = GetRotationContext();
+	RotationContext = FRotationContext();
+			
+	const bool bIsExplicitRotation = GetCoordSystemSpace() == COORD_Explicit; 
+	const bool bRotating = GetModeManager()->GetWidgetMode() == UE::Widget::EWidgetMode::WM_Rotate;
+	RotationContext.bUseExplicitRotator = bIsExplicitRotation && bRotating;
+
+	if (!RotationContext.bUseExplicitRotator)
+	{
+		return;
+	}
+	
+	for (const auto& [ControlRig, Shapes] : ControlRigShapeActors)
+	{
+		if (const URigHierarchy* Hierarchy = ControlRig ? ControlRig->GetHierarchy() : nullptr)
+		{
+			if (Hierarchy->UsesPreferredEulerAngles())
+			{
+				for (const AControlRigShapeActor* ShapeActor : Shapes)
+				{
+					if (ShapeActor->IsEnabled() && ShapeActor->IsSelected())
+					{
+						FRigControlElement* Control = ControlRig->FindControl(ShapeActor->ControlName);
+						const bool bUsePreferredRotationOrder = Hierarchy->GetUsePreferredRotationOrder(Control);
+
+						RotationContext.RotationOrder = bUsePreferredRotationOrder ? Hierarchy->GetControlPreferredEulerRotationOrder(Control) : EEulerRotationOrder::XYZ;
+						RotationContext.Rotation = Hierarchy->GetControlPreferredRotator(Control);
+
+						const FTransform Offset = Hierarchy->GetControlOffsetTransform(Control, ERigTransformType::CurrentGlobal);
+						const FTransform ComponentTransform = GetHostingSceneComponentTransform(ControlRig);
+						RotationContext.Offset = Offset * ComponentTransform;
+
+						// only get the first rotation order
+						return;
+					}
+				}
+			}
+		}
+	}
+}
+
 /**
 * FDetailKeyFrameCacheAndHandler
 */
diff --git a/Engine/Plugins/Animation/ControlRig/Source/ControlRigEditor/Public/EditMode/ControlRigEditMode.h b/Engine/Plugins/Animation/ControlRig/Source/ControlRigEditor/Public/EditMode/ControlRigEditMode.h
--- a/Engine/Plugins/Animation/ControlRig/Source/ControlRigEditor/Public/EditMode/ControlRigEditMode.h
+++ b/Engine/Plugins/Animation/ControlRig/Source/ControlRigEditor/Public/EditMode/ControlRigEditMode.h
@@ -49,6 +49,8 @@ enum class EMovieSceneDataChangeType;
 struct FMovieSceneChannelMetaData;
 class UMovieSceneSection;
 class UControlRigEditModeSettings;
+class UEditorTransformGizmoContextObject;
+struct FRotationContext;
 
 DECLARE_DELEGATE_RetVal_ThreeParams(FTransform, FOnGetRigElementTransform, const FRigElementKey& /*RigElementKey*/, bool /*bLocal*/, bool /*bOnDebugInstance*/);
 DECLARE_DELEGATE_ThreeParams(FOnSetRigElementTransform, const FRigElementKey& /*RigElementKey*/, const FTransform& /*Transform*/, bool /*bLocal*/);
@@ -234,6 +236,14 @@ public:
 
 protected:
 
+	enum class ETransformSpaceSpace
+	{
+		Default,
+		Parent,
+		Explicit
+	};
+	ETransformSpaceSpace GetTransformSpace() const;
+	
 	// shape related functions wrt enable/selection
 	/** Get the node name from the property path */
 	AControlRigShapeActor* GetControlShapeFromControlName(UControlRig* InControlRig,const FName& InControlName) const;
@@ -248,8 +258,9 @@ protected:
 	/** Updates cached pivot transforms */
 	void UpdatePivotTransforms();
 	bool ComputePivotFromEditedShape(UControlRig* InControlRig, FTransform& OutTransform) const;
-	bool ComputePivotFromShapeActors(UControlRig* InControlRig, const bool bEachLocalSpace, const bool bIsParentSpace, FTransform& OutTransform) const;
+	bool ComputePivotFromShapeActors(UControlRig* InControlRig, const bool bEachLocalSpace, const ETransformSpaceSpace InSpace, FTransform& OutTransform) const;
 	bool ComputePivotFromElements(UControlRig* InControlRig, FTransform& OutTransform) const;
+	FTransform GetPivotOrientation(const FRigElementKey& InControlKey, URigHierarchy* InHierarchy, const ETransformSpaceSpace InSpace, const FTransform& InComponentTransform) const;
 	
 	/** Get the current coordinate system space */
 	ECoordSystem GetCoordSystemSpace() const;
@@ -536,7 +547,7 @@ private:
 	bool CheckMovieSceneSig();
 	void SetControlShapeTransform( const AControlRigShapeActor* InShapeActor, const FTransform& InGlobalTransform,
 		const FTransform& InToWorldTransform, const FRigControlModifiedContext& InContext, const bool bPrintPython,
-		const bool bFixEulerFlips = true) const;
+		const FRotator& InRot, const bool bFixEulerFlips = true) const;
 	static FTransform GetControlShapeTransform(const AControlRigShapeActor* ShapeActor);
 
 	void ChangeControlShapeTransform(AControlRigShapeActor* ShapeActor, const bool bTranslation, FVector& InDrag,
@@ -659,6 +670,10 @@ private:
 
 	//pose used by the hotkeys
 	TObjectPtr<UControlRigPoseAsset> StoredPose;
+	
+	TWeakObjectPtr<UEditorTransformGizmoContextObject> WeakGizmoContext;
+	FRotationContext& GetRotationContext() const;
+	void UpdateRotationContext();
 
 	friend class FControlRigEditorModule;
 	friend class FControlRigEditor;
diff --git a/Engine/Source/Editor/EditorFramework/Public/UnrealWidgetFwd.h b/Engine/Source/Editor/EditorFramework/Public/UnrealWidgetFwd.h
--- a/Engine/Source/Editor/EditorFramework/Public/UnrealWidgetFwd.h
+++ b/Engine/Source/Editor/EditorFramework/Public/UnrealWidgetFwd.h
@@ -11,6 +11,7 @@ enum ECoordSystem
 	COORD_World,
 	COORD_Local,
 	COORD_Parent,
+	COORD_Explicit,
 	COORD_Max,
 };
 
diff --git a/Engine/Source/Editor/Experimental/EditorInteractiveToolsFramework/EditorInteractiveToolsFramework.Build.cs b/Engine/Source/Editor/Experimental/EditorInteractiveToolsFramework/EditorInteractiveToolsFramework.Build.cs
--- a/Engine/Source/Editor/Experimental/EditorInteractiveToolsFramework/EditorInteractiveToolsFramework.Build.cs
+++ b/Engine/Source/Editor/Experimental/EditorInteractiveToolsFramework/EditorInteractiveToolsFramework.Build.cs
@@ -35,8 +35,8 @@ public class EditorInteractiveToolsFramework : ModuleRules
 				"MeshDescription",
 				"StaticMeshDescription",
                 "EditorSubsystem",
-                "TypedElementRuntime"
-
+                "TypedElementRuntime",
+                "AnimationCore"
 				// ... add private dependencies that you statically link with here ...	
 			}
             );
diff --git a/Engine/Source/Editor/Experimental/EditorInteractiveToolsFramework/Private/EditorGizmos/EditorTransformGizmoSource.cpp b/Engine/Source/Editor/Experimental/EditorInteractiveToolsFramework/Private/EditorGizmos/EditorTransformGizmoSource.cpp
--- a/Engine/Source/Editor/Experimental/EditorInteractiveToolsFramework/Private/EditorGizmos/EditorTransformGizmoSource.cpp
+++ b/Engine/Source/Editor/Experimental/EditorInteractiveToolsFramework/Private/EditorGizmos/EditorTransformGizmoSource.cpp
@@ -117,6 +117,12 @@ EGizmoTransformScaleType UEditorTransformGizmoSource::GetScaleType() const
 	return EGizmoTransformScaleType::Default;
 }
 
+const FRotationContext& UEditorTransformGizmoSource::GetRotationContext() const
+{
+	static const FRotationContext DefaultContext;
+	return WeakContext.IsValid() ? WeakContext->RotationContext : DefaultContext;
+}
+
 UEditorTransformGizmoSource* UEditorTransformGizmoSource::CreateNew(UObject* Outer, const UEditorTransformGizmoContextObject* InContext)
 {
 	UEditorTransformGizmoSource* NewSource = NewObject<UEditorTransformGizmoSource>(Outer);
diff --git a/Engine/Source/Editor/Experimental/EditorInteractiveToolsFramework/Private/EditorGizmos/EditorTransformGizmoUtil.cpp b/Engine/Source/Editor/Experimental/EditorInteractiveToolsFramework/Private/EditorGizmos/EditorTransformGizmoUtil.cpp
--- a/Engine/Source/Editor/Experimental/EditorInteractiveToolsFramework/Private/EditorGizmos/EditorTransformGizmoUtil.cpp
+++ b/Engine/Source/Editor/Experimental/EditorInteractiveToolsFramework/Private/EditorGizmos/EditorTransformGizmoUtil.cpp
@@ -311,9 +311,32 @@ void UEditorTransformGizmoContextObject::InitializeGizmoManagerBinding()
 	{
 		UseNewGizmosChangedHandled = UEditorInteractiveGizmoManager::OnUsesNewTRSGizmosChangedDelegate().AddLambda(OnGizmoVariableChanged);
 	}
-	
+
 	// initialize default viewport list change binding
 	OnGizmoVariableChanged(UEditorInteractiveGizmoManager::UsesNewTRSGizmos());
+
+	// bind GizmosParametersChangedHandle if needed
+	auto OnGizmosParametersChanged = [this](const FGizmosParameters& InParameters)
+	{
+		if (!InParameters.bEnableExplicit)
+		{
+			// reset the coordinate system if explicit mode is not supported
+			if (ModeTools->GetCoordSystem() == COORD_Explicit)
+			{
+				ModeTools->SetCoordSystem(COORD_Local);
+			}
+		}
+	};
+	
+	if (!GizmosParametersChangedHandle.IsValid())
+	{
+		GizmosParametersChangedHandle = UEditorInteractiveGizmoManager::OnGizmosParametersChangedDelegate().AddLambda(OnGizmosParametersChanged);
+	}
+	
+	if (UEditorInteractiveGizmoManager::GetDefaultGizmosParameters())
+	{
+		OnGizmosParametersChanged(*UEditorInteractiveGizmoManager::GetDefaultGizmosParameters());
+	}
 }
 
 void UEditorTransformGizmoContextObject::InitializeViewportsBinding()
@@ -350,6 +373,12 @@ void UEditorTransformGizmoContextObject::RemoveGizmoManagerBinding()
 		UEditorInteractiveGizmoManager::OnUsesNewTRSGizmosChangedDelegate().Remove(UseNewGizmosChangedHandled);
 		UseNewGizmosChangedHandled.Reset();
 	}
+	
+	if (GizmosParametersChangedHandle.IsValid())
+	{
+		UEditorInteractiveGizmoManager::OnGizmosParametersChangedDelegate().Remove(GizmosParametersChangedHandle);
+		GizmosParametersChangedHandle.Reset();
+	}
 }
 
 void UEditorTransformGizmoContextObject::RemoveViewportsBinding()
diff --git a/Engine/Source/Editor/Experimental/EditorInteractiveToolsFramework/Private/EditorGizmos/GizmoElementGimbal.cpp b/Engine/Source/Editor/Experimental/EditorInteractiveToolsFramework/Private/EditorGizmos/GizmoElementGimbal.cpp
new file mode 100644
--- /dev/null
+++ b/Engine/Source/Editor/Experimental/EditorInteractiveToolsFramework/Private/EditorGizmos/GizmoElementGimbal.cpp
@@ -0,0 +1,87 @@
+// Copyright Epic Games, Inc. All Rights Reserved.
+
+#include "EditorGizmos/GizmoElementGimbal.h"
+
+#include "EditorGizmos/GizmoRotationUtil.h"
+
+#include UE_INLINE_GENERATED_CPP_BY_NAME(GizmoElementGimbal)
+
+void UGizmoElementGimbal::Render(IToolsContextRenderAPI* InRenderAPI, const FRenderTraversalState& InRenderState)
+{
+	using namespace UE::GizmoRotationUtil;
+	
+	if (Elements.Num() != 3)
+	{
+		return;
+	}
+	
+	FRenderTraversalState CurrentRenderState(InRenderState);
+	const bool bVisible = UpdateRenderState(InRenderAPI, FVector::ZeroVector, CurrentRenderState);
+	if (!bVisible)
+	{
+		return;
+	}
+
+	// decompose rotations
+	FRotationDecomposition Decomposition;
+	DecomposeRotations(CurrentRenderState.LocalToWorldTransform, RotationContext, Decomposition);
+	
+	int32 Index = 0;
+	for (UGizmoElementBase* Element : Elements)
+	{
+		CurrentRenderState.LocalToWorldTransform.SetRotation(Decomposition.R[Index++]);
+
+		if (Element)
+		{
+			Element->Render(InRenderAPI, CurrentRenderState);
+		}
+	}
+}
+
+FInputRayHit UGizmoElementGimbal::LineTrace(const UGizmoViewContext* InViewContext,
+	const FLineTraceTraversalState& InLineTraceState, const FVector& InRayOrigin, const FVector& InRayDirection)
+{
+	using namespace UE::GizmoRotationUtil;
+	
+	FInputRayHit Hit;
+	if (Elements.Num() != 3)
+	{
+		return Hit;
+	}
+
+	FLineTraceTraversalState CurrentLineTraceState(InLineTraceState);
+	const bool bHittable = UpdateLineTraceState(InViewContext, FVector::ZeroVector, CurrentLineTraceState);
+
+	if (!bHittable)
+	{
+		return Hit;
+	}
+	
+	// decompose rotations
+    FRotationDecomposition Decomposition;
+    DecomposeRotations(CurrentLineTraceState.LocalToWorldTransform, RotationContext, Decomposition);
+
+	int32 Index = 0;
+	for (UGizmoElementBase* Element : Elements)
+	{
+		CurrentLineTraceState.LocalToWorldTransform.SetRotation(Decomposition.R[Index++]);
+		if (Element)
+		{
+			FInputRayHit NewHit = Element->LineTrace(InViewContext, CurrentLineTraceState, InRayOrigin, InRayDirection);
+			if (!Hit.bHit || NewHit.HitDepth < Hit.HitDepth)
+			{
+				Hit = NewHit;
+			}
+		}
+	}
+	
+	return Hit;
+}
+
+void UGizmoElementGimbal::Add(UGizmoElementBase* InElement)
+{
+	if (Elements.Num() < 3)
+	{
+		UGizmoElementGroup::Add(InElement);
+	}
+}
diff --git a/Engine/Source/Editor/Experimental/EditorInteractiveToolsFramework/Private/EditorGizmos/GizmoRotationUtil.cpp b/Engine/Source/Editor/Experimental/EditorInteractiveToolsFramework/Private/EditorGizmos/GizmoRotationUtil.cpp
new file mode 100644
--- /dev/null
+++ b/Engine/Source/Editor/Experimental/EditorInteractiveToolsFramework/Private/EditorGizmos/GizmoRotationUtil.cpp
@@ -0,0 +1,283 @@
+﻿// Fill out your copyright notice in the Description page of Project Settings.
+
+
+#include "EditorGizmos/GizmoRotationUtil.h"
+#include "AnimationCoreLibrary.h"
+#include "Components/SceneComponent.h"
+
+#include "EditorGizmos/TransformGizmoInterfaces.h"
+
+#include "Elements/Actor/ActorElementWorldInterface.h"
+#include "Elements/Actor/ActorElementEditorWorldInterface.h"
+#include "Elements/Actor/ActorElementData.h"
+
+#include "Elements/Component/ComponentElementWorldInterface.h"
+#include "Elements/Component/ComponentElementEditorWorldInterface.h"
+#include "Elements/Component/ComponentElementData.h"
+
+#include "GameFramework/Actor.h"
+
+namespace UE::GizmoRotationUtil
+{
+	
+void DecomposeRotations(const FTransform& InTransform, const FRotationContext& InRotationContext, FRotationDecomposition& OutDecomposition)
+{
+	const FQuat QX = FQuat(FRotator(0.0, 0.0, InRotationContext.Rotation.Roll));
+	const FQuat QY = FQuat(FRotator(InRotationContext.Rotation.Pitch, 0.0, 0.0));
+	const FQuat QZ = FQuat(FRotator(0.0, InRotationContext.Rotation.Yaw, 0.0));
+		
+	FQuat& RX = OutDecomposition.R[0];
+	FQuat& RY = OutDecomposition.R[1];
+	FQuat& RZ = OutDecomposition.R[2];
+
+	switch (InRotationContext.RotationOrder)
+	{
+	case EEulerRotationOrder::XYZ:
+		RZ = QZ;
+		RY = QZ * QY;
+		RX = QZ * QY * QX;
+		break;
+	case EEulerRotationOrder::XZY:
+		RY = QY;
+		RZ = QY * QZ;
+		RX = QY * QZ * QX;
+		break;
+	case EEulerRotationOrder::YXZ:
+		RZ = QZ;
+		RX = QZ * QX;
+		RY = QZ * QX * QY;
+		break;
+	case EEulerRotationOrder::YZX:
+		RX = QX;
+		RZ = QX * QZ;
+		RY = QX * QZ * QY;
+		break;
+	case EEulerRotationOrder::ZXY:
+		RY = QY;
+		RX = QY * QX;
+		RZ = QY * QX * QZ;
+		break;
+	case EEulerRotationOrder::ZYX:
+		RX = QX;
+		RY = QX * QY;
+		RZ = QX * QY * QZ;
+		break;
+	default:
+		RX = QX;
+		RZ = QX * QZ;
+		RY = QX * QZ * QY;
+		break;
+	}
+
+	const FVector GlobalScale = InTransform.GetScale3D();
+	const FVector OffsetScale = InRotationContext.Offset.GetScale3D();
+	if (FTransform::AnyHasNegativeScale(GlobalScale, OffsetScale))
+	{
+		// use FTransform instead to compensate for scaling
+		const FTransform TRX(RX), TRY(RY), TRZ(RZ);
+		
+		// compute global orientations
+		const FTransform GlobalRX = TRX * InRotationContext.Offset;
+		const FTransform GlobalRY = TRY * InRotationContext.Offset;
+		const FTransform GlobalRZ = TRZ * InRotationContext.Offset;
+		
+		// switch them back to gizmo space
+		const FTransform TransformNotRotation(FQuat::Identity, InTransform.GetTranslation(), GlobalScale);
+		RX = GlobalRX.GetRelativeTransform(TransformNotRotation).GetRotation();
+		RY = GlobalRY.GetRelativeTransform(TransformNotRotation).GetRotation();
+		RZ = GlobalRZ.GetRelativeTransform(TransformNotRotation).GetRotation();	
+	}
+	else
+	{
+		const FQuat RotationOffset = InRotationContext.Offset.GetRotation();
+		RX = RotationOffset * RX;
+		RY = RotationOffset * RY;
+		RZ = RotationOffset * RZ;
+	}
+}
+
+FVector GetRotationAxis(const FTransform& InTransform, const FRotationContext& InRotationContext, const int32 InAxis)
+{
+	static const TArray RotateAxis({FVector::XAxisVector, FVector::YAxisVector, -FVector::ZAxisVector});
+	if (!ensure(RotateAxis.IsValidIndex(InAxis)))
+	{
+		return FVector::Zero();
+	}
+
+	FRotationDecomposition Decomposition;
+	DecomposeRotations(InTransform, InRotationContext, Decomposition);
+
+	// handle negative scaling
+	FVector ScaleSign = InRotationContext.Offset.GetScale3D().GetSignVector();
+	if (ScaleSign[0] < 0.0 || ScaleSign[1] < 0.0 || ScaleSign[2] < 0.0)
+	{
+		const double GlobalScale = ScaleSign[0] * ScaleSign[1] * ScaleSign[2];
+		ScaleSign *= GlobalScale;
+		return Decomposition.R[InAxis] * (RotateAxis[InAxis] * ScaleSign);
+	}
+
+	return Decomposition.R[InAxis] * (RotateAxis[InAxis] * ScaleSign);
+}
+
+FRelativeTransformInterfaceRegistry& FRelativeTransformInterfaceRegistry::Get()
+{
+	static FRelativeTransformInterfaceRegistry Singleton;
+	return Singleton;
+}
+	
+IRelativeTransformInterface* FRelativeTransformInterfaceRegistry::FindRelativeTransformInterface(const TTypedElement<ITypedElementWorldInterface>& InElement) const
+{
+	if (UObject* ObjectWorldInterface = Cast<UObject>(InElement.GetInterface()))
+	{
+		return FindRelativeTransformInterface(ObjectWorldInterface->GetClass());
+	}
+	return nullptr;
+}
+	
+IRelativeTransformInterface* FRelativeTransformInterfaceRegistry::FindRelativeTransformInterface(const UClass* InClass) const
+{
+	const TUniquePtr<IRelativeTransformInterface>* Interface = WorldInterfaceToRelativeTransformInterface.Find(InClass);
+	ensureMsgf(Interface, TEXT("No relative transform interface found for class %s. Did you call RegisterRelativeTransformInterface<> for that class?"), *InClass->GetName());
+	return Interface ? Interface->Get() : nullptr;
+}
+
+void FRelativeTransformInterfaceRegistry::RegisterDefaultInterfaces()
+{
+	FRelativeTransformInterfaceRegistry& Registry = Get();
+	
+	Registry.RegisterRelativeTransformInterface<UActorElementWorldInterface>(MakeUnique<FActorRelativeTransformInterface>());
+	Registry.RegisterRelativeTransformInterface<UActorElementEditorWorldInterface>(MakeUnique<FActorRelativeTransformInterface>());
+	
+	Registry.RegisterRelativeTransformInterface<UComponentElementWorldInterface>(MakeUnique<FComponentRelativeTransformInterface>());
+	Registry.RegisterRelativeTransformInterface<UComponentElementEditorWorldInterface>(MakeUnique<FComponentRelativeTransformInterface>());
+}
+
+bool GetRelativeTransform(const TTypedElement<ITypedElementWorldInterface>& InElement, FTransform& OutTransform, FRotationContext& InOutRotationContext)
+{
+	InOutRotationContext.RotationOrder = EEulerRotationOrder::XYZ;
+	
+	if (InElement.GetWorldTransform(OutTransform))
+	{
+		if (InOutRotationContext.bUseExplicitRotator)
+		{
+			bool bComputeRelativeFromTransform = true;
+
+			if (UObject* ObjectWorldInterface = Cast<UObject>(InElement.GetInterface()))
+			{
+				const FRelativeTransformInterfaceRegistry& Registry = FRelativeTransformInterfaceRegistry::Get();
+				if (IRelativeTransformInterface* Interface = Registry.FindRelativeTransformInterface(ObjectWorldInterface->GetClass()))
+				{
+					FEulerTransform RelativeEulerTransform;
+					if (Interface->GetRelativeTransform(InElement, RelativeEulerTransform))
+					{
+						// explicit rotator value
+						InOutRotationContext.Rotation = RelativeEulerTransform.Rotator();
+						// parent's world 
+						InOutRotationContext.Offset = RelativeEulerTransform.ToFTransform().Inverse() * OutTransform;
+						
+						bComputeRelativeFromTransform = false;
+					}
+				}
+			}
+
+			if (bComputeRelativeFromTransform)
+			{
+				// relative rotation
+				FTransform RelativeTransform;
+				if (InElement.GetRelativeTransform(RelativeTransform))
+				{
+					// rotator from FQuat
+					InOutRotationContext.Rotation = RelativeTransform.Rotator();
+					// parent's world 
+					InOutRotationContext.Offset = RelativeTransform.Inverse() * OutTransform; 
+				}
+			}
+			return true;
+		}
+	
+		// parent space only, leave the rotation context as it is
+		FTransform RelativeTransform;
+		if (InElement.GetRelativeTransform(RelativeTransform))
+		{
+			const FTransform ParentWorld = RelativeTransform.Inverse() * OutTransform;
+			OutTransform.SetRotation(ParentWorld.GetRotation());
+		}
+		return true;
+	}
+	
+	return false;
+}
+	
+void GetComponentRelativeTransform(const USceneComponent* InSceneComponent, FEulerTransform& OutRelativeTransform)
+{
+	OutRelativeTransform.Location = InSceneComponent->GetRelativeLocation();
+	OutRelativeTransform.Rotation = InSceneComponent->GetRelativeRotation();
+	OutRelativeTransform.Scale = InSceneComponent->GetRelativeScale3D();
+}
+
+void SetComponentRelativeTransform(USceneComponent* InSceneComponent, const FEulerTransform& InRelativeTransform)
+{
+	InSceneComponent->SetRelativeLocation_Direct(InRelativeTransform.GetLocation());
+	InSceneComponent->SetRelativeRotationExact(InRelativeTransform.Rotator());
+	InSceneComponent->SetRelativeScale3D_Direct(InRelativeTransform.GetScale3D());
+}
+
+bool FActorRelativeTransformInterface::GetRelativeTransform(const FTypedElementHandle& InElementHandle, FEulerTransform& OutRelativeTransform)
+{
+	if (USceneComponent* SceneComponent = GetSceneComponent(InElementHandle))
+	{
+		GetComponentRelativeTransform(SceneComponent, OutRelativeTransform);
+		return true;	
+	}
+	
+	return false;
+}
+	
+bool FActorRelativeTransformInterface::SetRelativeTransform(const FTypedElementHandle& InElementHandle, const FEulerTransform& InRelativeTransform)
+{
+	if (USceneComponent* SceneComponent = GetSceneComponent(InElementHandle))
+	{
+		SetComponentRelativeTransform(SceneComponent, InRelativeTransform);
+		return true;	
+	}
+	
+	return false;
+}
+
+USceneComponent* FActorRelativeTransformInterface::GetSceneComponent(const FTypedElementHandle& InElementHandle)
+{
+	const AActor* Actor = ActorElementDataUtil::GetActorFromHandle(InElementHandle);
+	return Actor ? Actor->GetRootComponent() : nullptr;
+}
+	
+bool FComponentRelativeTransformInterface::GetRelativeTransform(const FTypedElementHandle& InElementHandle, FEulerTransform& OutRelativeTransform)
+{
+	if (const USceneComponent* SceneComponent = GetSceneComponent(InElementHandle))
+	{
+		OutRelativeTransform.Location = SceneComponent->GetRelativeLocation();
+    	OutRelativeTransform.Rotation = SceneComponent->GetRelativeRotation();
+    	OutRelativeTransform.Scale = SceneComponent->GetRelativeScale3D();
+		return true;
+	}
+
+	return false;
+}
+
+bool FComponentRelativeTransformInterface::SetRelativeTransform(const FTypedElementHandle& InElementHandle, const FEulerTransform& InRelativeTransform)
+{
+	if (USceneComponent* SceneComponent = GetSceneComponent(InElementHandle))
+	{
+		SetComponentRelativeTransform(SceneComponent, InRelativeTransform);
+		return true;	
+	}
+	
+	return false;
+}
+
+USceneComponent* FComponentRelativeTransformInterface::GetSceneComponent(const FTypedElementHandle& InElementHandle)
+{
+	UActorComponent* Component = ComponentElementDataUtil::GetComponentFromHandle(InElementHandle);
+	return Component ? Cast<USceneComponent>(Component) : nullptr;
+}
+	
+}
diff --git a/Engine/Source/Editor/Experimental/EditorInteractiveToolsFramework/Private/EditorGizmos/TransformGizmo.cpp b/Engine/Source/Editor/Experimental/EditorInteractiveToolsFramework/Private/EditorGizmos/TransformGizmo.cpp
--- a/Engine/Source/Editor/Experimental/EditorInteractiveToolsFramework/Private/EditorGizmos/TransformGizmo.cpp
+++ b/Engine/Source/Editor/Experimental/EditorInteractiveToolsFramework/Private/EditorGizmos/TransformGizmo.cpp
@@ -1,6 +1,8 @@
 // Copyright Epic Games, Inc. All Rights Reserved.
 
 #include "EditorGizmos/TransformGizmo.h"
+
+#include "AnimationCoreLibrary.h"
 #include "BaseBehaviors/ClickDragBehavior.h"
 #include "BaseBehaviors/MouseHoverBehavior.h"
 #include "BaseGizmos/AxisSources.h"
@@ -9,6 +11,8 @@
 #include "BaseGizmos/GizmoMath.h"
 #include "BaseGizmos/GizmoRenderingUtil.h"
 #include "BaseGizmos/StateTargets.h"
+#include "EditorGizmos/GizmoElementGimbal.h"
+#include "EditorGizmos/GizmoRotationUtil.h"
 #include "Engine/CollisionProfile.h"
 #include "Engine/World.h"
 #include "Materials/Material.h"
@@ -52,6 +56,34 @@ static FAutoConsoleVariableRef CVarProjectIndirect(
 	TEXT("Project to the nearest point of the curve when handling indirect rotation.")
 	);
 
+//cf. GeometryCore implementation
+void MakePerpVectors(const FVector& Normal, FVector& OutPerp1, FVector& OutPerp2)
+{
+	// Duff et al method, from https://graphics.pixar.com/library/OrthonormalB/paper.pdf
+	if (Normal.Z < 0.0)
+	{
+		double A = 1.0 / (1.0 - Normal.Z);
+		double B = Normal.X * Normal.Y * A;
+		OutPerp1.X = 1.0 - Normal.X * Normal.X * A;
+		OutPerp1.Y = -B;
+		OutPerp1.Z = Normal.X;
+		OutPerp2.X = B;
+		OutPerp2.Y = Normal.Y * Normal.Y * A - 1.0;
+		OutPerp2.Z = -Normal.Y;
+	}
+	else
+	{
+		double A = 1.0 / (1.0 + Normal.Z);
+		double B = -Normal.X * Normal.Y * A;
+		OutPerp1.X = 1.0 - Normal.X * Normal.X * A;
+		OutPerp1.Y = B;
+		OutPerp1.Z = -Normal.X;
+		OutPerp2.X = B;
+		OutPerp2.Y = 1.0 - Normal.Y * Normal.Y * A;
+		OutPerp2.Z = -Normal.Y;
+	}
+}
+	
 }
 
 void UTransformGizmo::SetDisallowNegativeScaling(bool bDisallow)
@@ -82,6 +114,14 @@ void UTransformGizmo::Setup()
 	GizmoElementRoot->SetHoverLineColor(CurrentColor);
 	GizmoElementRoot->SetInteractLineColor(CurrentColor);
 
+	// the main gimbal rotation element that manages the three gimbal rotation axis as a group
+	RotateGimbalElement = NewObject<UGizmoElementGimbal>();
+	RotateGimbalElement->SetHoverMaterial(CurrentAxisMaterial);
+	RotateGimbalElement->SetInteractMaterial(CurrentAxisMaterial);
+	RotateGimbalElement->SetHoverLineColor(CurrentColor);
+	RotateGimbalElement->SetInteractLineColor(CurrentColor);
+	GizmoElementRoot->Add(RotateGimbalElement);
+
 	bInInteraction = false;
 
 	SetModeLastHitPart(EGizmoTransformMode::None, ETransformGizmoPartIdentifier::Default);
@@ -107,9 +147,12 @@ void UTransformGizmo::SetupBehaviors()
 
 void UTransformGizmo::SetupIndirectBehaviors()
 {
+	static constexpr int INDIRECT_GIZMO_PRIORITY = FInputCapturePriority::DEFAULT_GIZMO_PRIORITY-1;
+	
 	// Add middle mouse input behavior for indirect manipulation
 	ULocalClickDragInputBehavior* MiddleClickDragBehavior = NewObject<ULocalClickDragInputBehavior>();
 	MiddleClickDragBehavior->Initialize();
+	MiddleClickDragBehavior->SetDefaultPriority(INDIRECT_GIZMO_PRIORITY);
 	MiddleClickDragBehavior->SetUseMiddleMouseButton();
 	MiddleClickDragBehavior->CanBeginClickDragFunc = [this](const FInputDeviceRay&)
 	{
@@ -151,6 +194,7 @@ void UTransformGizmo::SetupIndirectBehaviors()
 	// Add left/right mouse input behavior for indirect manipulation
 	MultiIndirectClickDragBehavior = NewObject<UMultiButtonClickDragBehavior>();
 	MultiIndirectClickDragBehavior->Initialize();
+	MultiIndirectClickDragBehavior->SetDefaultPriority(INDIRECT_GIZMO_PRIORITY);
 	MultiIndirectClickDragBehavior->EnableButton(EKeys::LeftMouseButton);
 	if (bCtrlMiddleDoesY)
 	{
@@ -222,10 +266,16 @@ void UTransformGizmo::SetupIndirectBehaviors()
 			ETransformGizmoPartIdentifier::ScaleYAxis,
 			ETransformGizmoPartIdentifier::ScaleZAxis};
 
+		static constexpr ETransformGizmoPartIdentifier GimbalIds [4] = {
+			ETransformGizmoPartIdentifier::Default,
+			ETransformGizmoPartIdentifier::RotateXGimbal,
+			ETransformGizmoPartIdentifier::RotateYGimbal,
+			ETransformGizmoPartIdentifier::RotateZGimbal};
+
 		switch (CurrentMode)
 		{
 		case EGizmoTransformMode::Translate: return TranslateIds[InAxis];
-		case EGizmoTransformMode::Rotate: return RotateIds[InAxis];
+		case EGizmoTransformMode::Rotate: return bGimbalRotationMode ? GimbalIds[InAxis] : RotateIds[InAxis];
 		case EGizmoTransformMode::Scale: return ScaleIds[InAxis];
 		default: break;
 		}
@@ -582,6 +632,67 @@ void UTransformGizmo::UpdateMode()
 		EnableMode(CurrentMode, NewAxisToDraw);
 		CurrentAxisToDraw = NewAxisToDraw;
 	}
+	else if (CurrentMode == EGizmoTransformMode::Rotate && bGimbalRotationMode != GetRotationContext().bUseExplicitRotator)
+	{
+		UpdateRotationMode();
+	}
+}
+
+void UTransformGizmo::UpdateRotationMode()
+{
+	EnableMode(EGizmoTransformMode::Rotate, CurrentAxisToDraw);
+
+	if (GetCurrentModeLastHitPart() != ETransformGizmoPartIdentifier::Default)
+	{
+		const ETransformGizmoPartIdentifier PreviousHitPart = GetCurrentModeLastHitPart();
+		UpdateInteractingState(false, PreviousHitPart, true);
+		
+		ETransformGizmoPartIdentifier NewHitPart = ETransformGizmoPartIdentifier::RotateArcball;
+
+		constexpr uint8 RotateId = static_cast<uint8>(EGizmoTransformMode::Rotate);
+		if (bGimbalRotationMode)
+		{
+			switch (LastHitPartPerMode[RotateId])
+			{
+			case ETransformGizmoPartIdentifier::RotateXAxis:
+				NewHitPart = ETransformGizmoPartIdentifier::RotateXGimbal;
+				break;
+			case ETransformGizmoPartIdentifier::RotateYAxis:
+				NewHitPart = ETransformGizmoPartIdentifier::RotateYGimbal;
+				break;
+			case ETransformGizmoPartIdentifier::RotateZAxis:
+				NewHitPart = ETransformGizmoPartIdentifier::RotateZGimbal;
+				break;
+			case ETransformGizmoPartIdentifier::RotateScreenSpace:
+				NewHitPart = ETransformGizmoPartIdentifier::RotateArcball;
+				break;
+			default:
+				break;
+			}
+		}
+		else
+		{
+			switch (LastHitPartPerMode[RotateId])
+			{
+			case ETransformGizmoPartIdentifier::RotateXGimbal:
+				NewHitPart = ETransformGizmoPartIdentifier::RotateXAxis;
+				break;
+			case ETransformGizmoPartIdentifier::RotateYGimbal:
+				NewHitPart = ETransformGizmoPartIdentifier::RotateYAxis;
+				break;
+			case ETransformGizmoPartIdentifier::RotateZGimbal:
+				NewHitPart = ETransformGizmoPartIdentifier::RotateZAxis;
+				break;
+			default:
+				break;
+			}
+		}
+
+		SetModeLastHitPart(CurrentMode, NewHitPart);
+
+		static constexpr bool bInInteracting = true, bIdOnly = true;
+		UpdateInteractingState(bInInteracting, NewHitPart, bIdOnly);
+	}
 }
 
 void UTransformGizmo::EnableMode(EGizmoTransformMode InMode, EAxisList::Type InAxisListToDraw)
@@ -713,72 +824,85 @@ void UTransformGizmo::EnableRotate(EAxisList::Type InAxisListToDraw)
 	const bool bEnableZ = static_cast<uint8>(InAxisListToDraw) & static_cast<uint8>(EAxisList::Z);
 	const bool bEnableAll = bEnableX && bEnableY && bEnableZ;
 
-	const FVector XAxis(1.0f, 0.0f, 0.0f);
-	const FVector YAxis(0.0f, 1.0f, 0.0f);
-	const FVector ZAxis(0.0f, 0.0f, 1.0f);
+	bGimbalRotationMode = (bEnableX || bEnableY || bEnableZ) && GetRotationContext().bUseExplicitRotator;
 
-	if (bEnableX && RotateXAxisElement == nullptr)
+	// default rotation handles
 	{
-		RotateXAxisElement = MakeRotateAxis(ETransformGizmoPartIdentifier::RotateXAxis, YAxis, ZAxis, AxisMaterialX, CurrentAxisMaterial);
-		GizmoElementRoot->Add(RotateXAxisElement);
-	}
+		auto EnableRotateElement = [this](TObjectPtr<UGizmoElementTorus>& InOutElement, const ETransformGizmoPartIdentifier InPartId, const bool bEnableAxis)
+		{
+			if (bEnableAxis)
+			{
+				if (!InOutElement)
+				{
+					InOutElement = MakeDefaultRotateAxis(InPartId);
+					GizmoElementRoot->Add(InOutElement);
+				}
+			}
+	
+			if (InOutElement)
+			{
+				InOutElement->SetEnabled(bEnableAxis && !bGimbalRotationMode);
+			}
+		};
 
-	if (bEnableY && RotateYAxisElement == nullptr)
-	{
-		RotateYAxisElement = MakeRotateAxis(ETransformGizmoPartIdentifier::RotateYAxis, ZAxis, XAxis, AxisMaterialY, CurrentAxisMaterial);
-		GizmoElementRoot->Add(RotateYAxisElement);
+		EnableRotateElement(RotateXAxisElement, ETransformGizmoPartIdentifier::RotateXAxis, bEnableX);
+		EnableRotateElement(RotateYAxisElement, ETransformGizmoPartIdentifier::RotateYAxis, bEnableY);
+		EnableRotateElement(RotateZAxisElement, ETransformGizmoPartIdentifier::RotateZAxis, bEnableZ);
 	}
 
-	if (bEnableZ && RotateZAxisElement == nullptr)
+	// gimbal rotation handles
 	{
-		RotateZAxisElement = MakeRotateAxis(ETransformGizmoPartIdentifier::RotateZAxis, XAxis, YAxis, AxisMaterialZ, CurrentAxisMaterial);
-		GizmoElementRoot->Add(RotateZAxisElement);
+		auto EnableGimbalRotateElement = [this](TObjectPtr<UGizmoElementTorus>& InOutElement, const ETransformGizmoPartIdentifier InPartId, const bool bEnable)
+		{
+			if (bEnable)
+			{
+				if (!InOutElement)
+				{
+					InOutElement = MakeDefaultRotateAxis(InPartId);
+					RotateGimbalElement->Add(InOutElement);
+				}
+			}
+	
+			if (InOutElement)
+			{
+				InOutElement->SetEnabled(bEnable && bGimbalRotationMode);
+			}
+		};
+
+		EnableGimbalRotateElement(RotateXGimbalElement, ETransformGizmoPartIdentifier::RotateXGimbal, bEnableX);
+		EnableGimbalRotateElement(RotateYGimbalElement, ETransformGizmoPartIdentifier::RotateYGimbal, bEnableY);
+		EnableGimbalRotateElement(RotateZGimbalElement, ETransformGizmoPartIdentifier::RotateZGimbal, bEnableZ);
+
+		if (RotateGimbalElement)
+		{
+			RotateGimbalElement->SetEnabled(bEnableAll && bGimbalRotationMode);
+		}
 	}
 
+	// screen space & arc ball handles
 	if (bEnableAll)
 	{
-		if (RotateScreenSpaceElement == nullptr)
+		if (!RotateScreenSpaceElement)
 		{
 			RotateScreenSpaceElement = MakeRotateCircleHandle(ETransformGizmoPartIdentifier::RotateScreenSpace, RotateScreenSpaceRadius, RotateScreenSpaceCircleColor, false);
 			GizmoElementRoot->Add(RotateScreenSpaceElement);
 		}
 
-		if (RotateArcballElement == nullptr)
+		if (!RotateArcballElement)
 		{
 			RotateArcballElement = MakeRotateCircleHandle(ETransformGizmoPartIdentifier::RotateArcball, RotateArcballSphereRadius, RotateArcballCircleColor, true);
 			GizmoElementRoot->Add(RotateArcballElement);
 		}
 	}
-
-	if (RotateXAxisElement)
-	{
-		RotateXAxisElement->SetEnabled(bEnableX);
-	}
-
-	if (RotateYAxisElement)
-	{
-		RotateYAxisElement->SetEnabled(bEnableY);
-	}
-
-	if (RotateZAxisElement)
-	{
-		RotateZAxisElement->SetEnabled(bEnableZ);
-	}
-
+	
 	if (RotateScreenSpaceElement)
-	{
-		RotateScreenSpaceElement->SetEnabled(bEnableAll);
-	}
-
-	if (RotateOuterCircleElement)
-	{
-		RotateOuterCircleElement->SetEnabled(bEnableAll);
-	}
-
-	if (RotateArcballElement)
-	{ 
-		RotateArcballElement->SetEnabled(bEnableAll);
-	}
+    {
+        RotateScreenSpaceElement->SetEnabled(bEnableAll && !bGimbalRotationMode);
+    }
+    if (RotateArcballElement)
+    { 
+        RotateArcballElement->SetEnabled(bEnableAll);
+    }
 }
 
 void UTransformGizmo::EnableScale(EAxisList::Type InAxisListToDraw)
@@ -868,6 +992,12 @@ bool UTransformGizmo::CanInteract(const EViewportContext InViewportContext) cons
 	return bValidMode && bVisible;
 }
 
+const FRotationContext& UTransformGizmo::GetRotationContext() const
+{
+	static const FRotationContext Context;
+	return TransformGizmoSource ? TransformGizmoSource->GetRotationContext() : Context;
+}
+
 void UTransformGizmo::Tick(float DeltaTime)
 {
 	if (PendingDragFunction)
@@ -879,6 +1009,12 @@ void UTransformGizmo::Tick(float DeltaTime)
 	UpdateMode();
 
 	UpdateCameraAxisSource();
+
+	// update gimbal handle's rotation context
+	if (RotateGimbalElement)
+	{
+		RotateGimbalElement->RotationContext = GetRotationContext();
+	}
 }
 
 void UTransformGizmo::SetActiveTarget(UTransformProxy* Target, IToolContextTransactionProvider* TransactionProvider, IGizmoStateTarget* InStateTarget)
@@ -958,7 +1094,7 @@ void UTransformGizmo::HandleWidgetModeChanged(UE::Widget::EWidgetMode InWidgetMo
 	if (CurrentMode != EGizmoTransformMode::None && NewMode == CurrentMode)
 	{
 		const ETransformGizmoPartIdentifier CurrentModeLastHitPart = GetCurrentModeLastHitPart();
-		auto GetModeDefaultHitPart = [NewMode, CurrentModeLastHitPart]()
+		auto GetModeDefaultHitPart = [NewMode, CurrentModeLastHitPart, bGimbal = bGimbalRotationMode]()
 		{
 			const bool bIsRotateArcBall = (CurrentModeLastHitPart == ETransformGizmoPartIdentifier::RotateArcball);
 			switch (NewMode)
@@ -966,7 +1102,7 @@ void UTransformGizmo::HandleWidgetModeChanged(UE::Widget::EWidgetMode InWidgetMo
 			case EGizmoTransformMode::Translate:
 				return ETransformGizmoPartIdentifier::TranslateScreenSpace;
 			case EGizmoTransformMode::Rotate:
-				return bIsRotateArcBall ? ETransformGizmoPartIdentifier::RotateScreenSpace : ETransformGizmoPartIdentifier::RotateArcball;
+				return bGimbal ? ETransformGizmoPartIdentifier::RotateArcball : bIsRotateArcBall ? ETransformGizmoPartIdentifier::RotateScreenSpace : ETransformGizmoPartIdentifier::RotateArcball;
 			case EGizmoTransformMode::Scale:
 				return ETransformGizmoPartIdentifier::ScaleUniform;
 			default:
@@ -1120,6 +1256,30 @@ UGizmoElementRectangle* UTransformGizmo::MakeTranslateScreenSpaceHandle()
 	return RectangleElement;
 }
 
+UGizmoElementTorus* UTransformGizmo::MakeDefaultRotateAxis(const ETransformGizmoPartIdentifier InPartId)
+{
+	switch (InPartId)
+	{
+	case ETransformGizmoPartIdentifier::RotateXAxis:
+	case ETransformGizmoPartIdentifier::RotateXGimbal:
+		return MakeRotateAxis(InPartId, FVector::YAxisVector, FVector::ZAxisVector, AxisMaterialX, CurrentAxisMaterial);
+		
+	case ETransformGizmoPartIdentifier::RotateYAxis:
+	case ETransformGizmoPartIdentifier::RotateYGimbal:
+		return MakeRotateAxis(InPartId, FVector::ZAxisVector, FVector::XAxisVector, AxisMaterialY, CurrentAxisMaterial);
+		
+	case ETransformGizmoPartIdentifier::RotateZAxis:
+	case ETransformGizmoPartIdentifier::RotateZGimbal:
+		return MakeRotateAxis(InPartId, FVector::XAxisVector, FVector::YAxisVector, AxisMaterialZ, CurrentAxisMaterial);
+		
+	default:
+		break;
+	}
+
+	ensure(false);
+	return nullptr;
+}
+
 UGizmoElementTorus* UTransformGizmo::MakeRotateAxis(ETransformGizmoPartIdentifier InPartId, const FVector& TorusAxis0, const FVector& TorusAxis1,
 	UMaterialInterface* InMaterial, UMaterialInterface* InCurrentMaterial)
 {
@@ -1261,7 +1421,7 @@ FQuat UTransformGizmo::RotationSnapFunction(const FQuat& DeltaRotation) const
 	return SnappedDeltaRotation;
 }
 
-FVector UTransformGizmo::GetWorldAxis(const FVector& InAxis)
+FVector UTransformGizmo::GetWorldAxis(const FVector& InAxis) const
 {
 	auto GetCoordinateSystem = [&]()
 	{
@@ -1280,6 +1440,10 @@ FVector UTransformGizmo::GetWorldAxis(const FVector& InAxis)
 	return InAxis;
 }
 
+FVector UTransformGizmo::GetGimbalRotationAxis(const int32 InAxis) const
+{
+	return UE::GizmoRotationUtil::GetRotationAxis(CurrentTransform, GetRotationContext(), InAxis);
+}
 
 void UTransformGizmo::SetupOnClickFunctions()
 {
@@ -1307,6 +1471,9 @@ void UTransformGizmo::SetupOnClickFunctions()
 	OnClickPressFunctions[static_cast<int>(ETransformGizmoPartIdentifier::RotateZAxis)] = &UTransformGizmo::OnClickPressRotateZAxis;
 	OnClickPressFunctions[static_cast<int>(ETransformGizmoPartIdentifier::RotateScreenSpace)] = &UTransformGizmo::OnClickPressScreenSpaceRotate;
 	OnClickPressFunctions[static_cast<int>(ETransformGizmoPartIdentifier::RotateArcball)] = &UTransformGizmo::OnClickPressArcBallRotate;
+	OnClickPressFunctions[static_cast<int>(ETransformGizmoPartIdentifier::RotateXGimbal)] = &UTransformGizmo::OnClickPressGimbalRotateAxis;
+	OnClickPressFunctions[static_cast<int>(ETransformGizmoPartIdentifier::RotateYGimbal)] = &UTransformGizmo::OnClickPressGimbalRotateAxis;
+	OnClickPressFunctions[static_cast<int>(ETransformGizmoPartIdentifier::RotateZGimbal)] = &UTransformGizmo::OnClickPressGimbalRotateAxis;
 
 	OnClickDragFunctions[static_cast<int>(ETransformGizmoPartIdentifier::TranslateXAxis)] = &UTransformGizmo::OnClickDragTranslateAxis;
 	OnClickDragFunctions[static_cast<int>(ETransformGizmoPartIdentifier::TranslateYAxis)] = &UTransformGizmo::OnClickDragTranslateAxis;
@@ -1327,6 +1494,9 @@ void UTransformGizmo::SetupOnClickFunctions()
 	OnClickDragFunctions[static_cast<int>(ETransformGizmoPartIdentifier::RotateZAxis)] = &UTransformGizmo::OnClickDragRotateAxis;
 	OnClickDragFunctions[static_cast<int>(ETransformGizmoPartIdentifier::RotateScreenSpace)] = &UTransformGizmo::OnClickDragScreenSpaceRotate;
 	OnClickDragFunctions[static_cast<int>(ETransformGizmoPartIdentifier::RotateArcball)] = &UTransformGizmo::OnClickDragArcBallRotate;
+	OnClickDragFunctions[static_cast<int>(ETransformGizmoPartIdentifier::RotateXGimbal)] = &UTransformGizmo::OnClickDragGimbalRotateAxis;
+	OnClickDragFunctions[static_cast<int>(ETransformGizmoPartIdentifier::RotateYGimbal)] = &UTransformGizmo::OnClickDragGimbalRotateAxis;
+	OnClickDragFunctions[static_cast<int>(ETransformGizmoPartIdentifier::RotateZGimbal)] = &UTransformGizmo::OnClickDragGimbalRotateAxis;
 
 	OnClickReleaseFunctions[static_cast<int>(ETransformGizmoPartIdentifier::TranslateXAxis)] = &UTransformGizmo::OnClickReleaseTranslateAxis;
 	OnClickReleaseFunctions[static_cast<int>(ETransformGizmoPartIdentifier::TranslateYAxis)] = &UTransformGizmo::OnClickReleaseTranslateAxis;
@@ -1347,6 +1517,9 @@ void UTransformGizmo::SetupOnClickFunctions()
 	OnClickReleaseFunctions[static_cast<int>(ETransformGizmoPartIdentifier::RotateZAxis)] = &UTransformGizmo::OnClickReleaseRotateAxis;
 	OnClickReleaseFunctions[static_cast<int>(ETransformGizmoPartIdentifier::RotateScreenSpace)] = &UTransformGizmo::OnClickReleaseScreenSpaceRotate;
 	OnClickReleaseFunctions[static_cast<int>(ETransformGizmoPartIdentifier::RotateArcball)] = &UTransformGizmo::OnClickReleaseArcBallRotate;
+	OnClickReleaseFunctions[static_cast<int>(ETransformGizmoPartIdentifier::RotateXGimbal)] = &UTransformGizmo::OnClickReleaseRotateAxis;
+	OnClickReleaseFunctions[static_cast<int>(ETransformGizmoPartIdentifier::RotateYGimbal)] = &UTransformGizmo::OnClickReleaseRotateAxis;
+	OnClickReleaseFunctions[static_cast<int>(ETransformGizmoPartIdentifier::RotateZGimbal)] = &UTransformGizmo::OnClickReleaseRotateAxis;
 }
 
 
@@ -2105,7 +2278,7 @@ void UTransformGizmo::OnClickPressRotateZAxis(const FInputDeviceRay& InPressPos)
 	OnClickPressRotateAxis(InPressPos);
 }
 
-FVector2D UTransformGizmo::GetScreenRotateAxisDir(const FInputDeviceRay& InPressPos)
+FVector2D UTransformGizmo::GetWorldToScreenRotateAxisDir(const FInputDeviceRay& InPressPos, const FVector& InWorldAxis)
 {
 	// NOTE that function is not intended to remain here indefinitely, its purpose is to debug closest point computation
 	const auto PrintProjection = [this](const TCHAR* const InMessage)
@@ -2115,31 +2288,20 @@ FVector2D UTransformGizmo::GetScreenRotateAxisDir(const FInputDeviceRay& InPress
 			UE_LOG(LogTransformGizmo, Warning, TEXT("%s"), InMessage);
 		}
 	};
-
-	static const TArray RotateIDs({	ETransformGizmoPartIdentifier::RotateXAxis,
-									ETransformGizmoPartIdentifier::RotateYAxis,
-									ETransformGizmoPartIdentifier::RotateZAxis,
-									ETransformGizmoPartIdentifier::RotateScreenSpace});
-	const int32 RotateID = RotateIDs.IndexOfByKey(LastHitPart);
-	if (!ensure(RotateID != INDEX_NONE))
-	{
-		return FVector2D::ZeroVector;
-	}
 	
 	const FRay& Ray = InPressPos.WorldRay;
 
 	// store world origin and axis
 	static const TArray RotateAxis({FVector::XAxisVector, FVector::YAxisVector, -FVector::ZAxisVector});
-	const FVector WorldAxis = RotateAxis.IsValidIndex(RotateID) ? GetWorldAxis(RotateAxis[RotateID]) : GizmoViewContext->GetViewDirection();
 	const FVector WorldOrigin = CurrentTransform.GetLocation();
 
 	// compute axis / view direction projection: is the rotation plane nearly perpendicular to the view plane?
 	const FVector ViewDirection = GizmoViewContext->IsPerspectiveProjection() ?
 		(WorldOrigin - GizmoViewContext->ViewLocation).GetSafeNormal() :
 		GizmoViewContext->GetViewDirection();
-	const bool bAxisPerpendicularToView = FMath::Abs(FVector::DotProduct(WorldAxis, ViewDirection)) < GizmoLocals::DotThreshold;
+	const bool bAxisPerpendicularToView = FMath::Abs(FVector::DotProduct(InWorldAxis, ViewDirection)) < GizmoLocals::DotThreshold;
 	// compute axis / ray direction projection: is the ray direction parallel to the axis?
-	const bool bRayPerpendicularToAxis = FMath::IsNearlyZero(FVector::DotProduct(WorldAxis, Ray.Direction));
+	const bool bRayPerpendicularToAxis = FMath::IsNearlyZero(FVector::DotProduct(InWorldAxis, Ray.Direction));
 
 	// compute closest point on the rotate handle
 	const bool bUseRayForIndirect = bIndirectManipulation && !GizmoLocals::ProjectIndirect;
@@ -2148,7 +2310,7 @@ FVector2D UTransformGizmo::GetScreenRotateAxisDir(const FInputDeviceRay& InPress
 	FVector QueryPoint = Ray.Origin;
 	if (!bUseRayOrigin)
 	{
-		const FPlane Plane(WorldOrigin, WorldAxis);
+		const FPlane Plane(WorldOrigin, InWorldAxis);
 
 		// if the projection is in front of the camera then use it
 		const double HitDepth = FMath::RayPlaneIntersectionParam(Ray.Origin, Ray.Direction, Plane);
@@ -2170,25 +2332,24 @@ FVector2D UTransformGizmo::GetScreenRotateAxisDir(const FInputDeviceRay& InPress
 	// compute nearest point
 	const float Radius = GetWorldRadius(RotateAxisOuterRadius);
 	FVector ClosestPointOnCircle;
-	GizmoMath::ClosetPointOnCircle(QueryPoint, WorldOrigin, WorldAxis, Radius, ClosestPointOnCircle);
+	GizmoMath::ClosetPointOnCircle(QueryPoint, WorldOrigin, InWorldAxis, Radius, ClosestPointOnCircle);
 
 	// compute world directions
 	const FVector ToClosestDirection = (ClosestPointOnCircle - WorldOrigin).GetSafeNormal();
-	const FVector PullDirection = FVector::CrossProduct(ToClosestDirection, WorldAxis);
+	const FVector PullDirection = FVector::CrossProduct(ToClosestDirection, InWorldAxis);
 
 	// compute screen projections
 	const FTransform ToClosest(ClosestPointOnCircle);
 	const FVector2D PullProjection = GetScreenProjectedAxis(GizmoViewContext, PullDirection, ToClosest);
-	const FVector2D AxisProjection = GetScreenProjectedAxis(GizmoViewContext, WorldAxis, ToClosest);
+	const FVector2D AxisProjection = GetScreenProjectedAxis(GizmoViewContext, InWorldAxis, ToClosest);
 	const FVector2D ToClosestProjection = GetScreenProjectedAxis(GizmoViewContext, ToClosestDirection, ToClosest);
 
 	// compute which projection to remove from drag
 	const double DotAxis = FMath::Abs( FVector2D::DotProduct(PullProjection, AxisProjection) );
 	const double DotClosest = FMath::Abs( FVector2D::DotProduct(PullProjection, ToClosestProjection) );
 	NormalProjectionToRemove = (DotAxis < DotClosest) ? AxisProjection : ToClosestProjection;
-	DebugNormalRemoved = (DotAxis < DotClosest) ? WorldAxis : ToClosestDirection;
-	DebugNormalSkip = (DotAxis < DotClosest) ? ToClosestDirection : WorldAxis;
-
+	DebugNormalRemoved = (DotAxis < DotClosest) ? InWorldAxis : ToClosestDirection;
+	DebugNormalSkip = (DotAxis < DotClosest) ? ToClosestDirection : InWorldAxis;
 
 	// debug
 	{
@@ -2199,6 +2360,39 @@ FVector2D UTransformGizmo::GetScreenRotateAxisDir(const FInputDeviceRay& InPress
 	return PullProjection;
 }
 
+FVector2D UTransformGizmo::GetScreenRotateAxisDir(const FInputDeviceRay& InPressPos)
+{
+	static const TArray RotateIDs({	ETransformGizmoPartIdentifier::RotateXAxis,
+									ETransformGizmoPartIdentifier::RotateYAxis,
+									ETransformGizmoPartIdentifier::RotateZAxis,
+									ETransformGizmoPartIdentifier::RotateScreenSpace});
+	const int32 RotateID = RotateIDs.IndexOfByKey(LastHitPart);
+	if (!ensure(RotateID != INDEX_NONE))
+	{
+		return FVector2D::ZeroVector;
+	}
+
+	// store world origin and axis
+	static const TArray RotateAxis({FVector::XAxisVector, FVector::YAxisVector, -FVector::ZAxisVector});
+	const FVector WorldAxis = RotateAxis.IsValidIndex(RotateID) ? GetWorldAxis(RotateAxis[RotateID]) : GizmoViewContext->GetViewDirection();
+
+	return GetWorldToScreenRotateAxisDir(InPressPos, WorldAxis);
+}
+
+FVector2D UTransformGizmo::GetScreenGimbalRotateAxisDir(const FInputDeviceRay& InPressPos)
+{
+	static const TArray RotateIDs({	ETransformGizmoPartIdentifier::RotateXGimbal,
+									ETransformGizmoPartIdentifier::RotateYGimbal,
+									ETransformGizmoPartIdentifier::RotateZGimbal});
+	const int32 RotateID = RotateIDs.IndexOfByKey(LastHitPart);
+	if (!ensure(RotateID != INDEX_NONE))
+	{
+		return FVector2D::ZeroVector;
+	}
+	
+	return GetWorldToScreenRotateAxisDir(InPressPos, GetGimbalRotationAxis(RotateID));
+}
+
 void UTransformGizmo::OnClickDragRotateAxis(const FInputDeviceRay& DragPos)
 {
 	switch (RotateMode)
@@ -2277,6 +2471,36 @@ FQuat UTransformGizmo::ComputeAxisRotateDelta(const FVector2D& InStartPos, const
 	return DeltaRot.Quaternion();
 }
 
+FQuat UTransformGizmo::ComputeGimbalRotateDelta(const FVector2D& InStartPos, const FVector2D& InEndPos)
+{
+	FVector2D DragDir = InEndPos - InStartPos;
+
+	if (bTrySwitchingToNormalPull)
+	{
+		const double DotTangent = FVector2D::DotProduct(InteractionScreenAxisDirection, DragDir);
+		const double DotNormal = FVector2D::DotProduct(NormalProjectionToRemove, DragDir);
+		if (FMath::Abs(DotNormal) > FMath::Abs(DotTangent))
+		{
+			::Swap(NormalProjectionToRemove, InteractionScreenAxisDirection);
+			::Swap(DebugDirection, DebugNormalRemoved);
+			InteractionScreenAxisDirection *= FMath::Sign(DotTangent) * FMath::Sign(DotNormal);
+		}
+		bTrySwitchingToNormalPull = false;
+	}
+	
+	const FVector2D DragDirToRemove = NormalProjectionToRemove * FVector2D::DotProduct(DragDir, NormalProjectionToRemove);
+	DragDir -= DragDirToRemove;
+
+	const double Delta = FVector2D::DotProduct(InteractionScreenAxisDirection, DragDir) * 0.25;
+	const FVector DeltaRot(
+	 	   InteractionAxisList == EAxisList::X ? Delta : 0.0,
+	 	   InteractionAxisList == EAxisList::Y ? Delta : 0.0,
+	 	   InteractionAxisList == EAxisList::Z ? Delta : 0.0);
+	
+	static constexpr bool bUseUEHandyness = true;
+	return AnimationCore::QuatFromEuler(DeltaRot, GetRotationContext().RotationOrder, bUseUEHandyness);
+}
+
 void UTransformGizmo::OnClickReleaseRotateAxis(const FInputDeviceRay& InReleasePos)
 {
 	bInInteraction = false;
@@ -2284,6 +2508,103 @@ void UTransformGizmo::OnClickReleaseRotateAxis(const FInputDeviceRay& InReleaseP
 	bTrySwitchingToNormalPull = false;
 }
 
+void UTransformGizmo::OnClickPressGimbalRotateAxis(const FInputDeviceRay& InPressPos)
+{
+	static const TArray RotateIDs({	ETransformGizmoPartIdentifier::RotateXGimbal,
+									ETransformGizmoPartIdentifier::RotateYGimbal,
+									ETransformGizmoPartIdentifier::RotateZGimbal});
+
+	const int32 RotateID = RotateIDs.IndexOfByKey(LastHitPart);
+	if (!ensure(RotateID != INDEX_NONE))
+	{
+		bTrySwitchingToNormalPull = false;
+		bInInteraction = false;
+		return;
+	}
+
+	static const TArray AxisList({EAxisList::X, EAxisList::Y, EAxisList::Z});
+	
+	bDebugRotate = true;
+	
+	// initialize pull data
+	InteractionScreenAxisDirection = GetScreenGimbalRotateAxisDir(InPressPos);
+	InteractionAxisDirection = GetGimbalRotationAxis(RotateID);
+	InteractionAxisList = AxisList[RotateID];
+	InteractionScreenStartPos = InteractionScreenCurrPos = InPressPos.ScreenPosition;
+	RotateMode = EAxisRotateMode::Pull;
+	
+	// initialize arc/mixed data
+	const bool bRotatePull = bIndirectManipulation || DefaultRotateMode == EAxisRotateMode::Pull;
+	if (!bRotatePull)
+	{		
+		const FVector WorldPlaneNormal = GetGimbalRotationAxis(RotateID);
+
+		FVector WorldPlaneAxis1, WorldPlaneAxis2;
+		GizmoLocals::MakePerpVectors(WorldPlaneNormal, WorldPlaneAxis1, WorldPlaneAxis2);
+		
+		const bool bCanRotateArc = OnClickPressRotateArc(InPressPos, WorldPlaneNormal, WorldPlaneAxis1, WorldPlaneAxis2);
+		if (bCanRotateArc)
+		{
+			RotateMode = EAxisRotateMode::Arc;
+		}
+	}
+
+	bTrySwitchingToNormalPull = bIndirectManipulation && RotateMode == EAxisRotateMode::Pull;
+	
+	bInInteraction = true;
+	SetModeLastHitPart(EGizmoTransformMode::Rotate, LastHitPart);
+}
+
+void UTransformGizmo::OnClickDragGimbalRotateAxis(const FInputDeviceRay& InDragPos)
+{
+	static const TArray RotateIDs({	ETransformGizmoPartIdentifier::RotateXGimbal,
+									ETransformGizmoPartIdentifier::RotateYGimbal,
+									ETransformGizmoPartIdentifier::RotateZGimbal});
+
+	const int32 RotateID = RotateIDs.IndexOfByKey(LastHitPart);
+	if (!ensure(RotateID != INDEX_NONE))
+	{
+		return;
+	}
+
+	switch (RotateMode)
+	{
+	case EAxisRotateMode::Pull:
+		{
+			const FQuat DeltaRot = ComputeGimbalRotateDelta(InteractionScreenCurrPos, InDragPos.ScreenPosition);
+			ApplyRotateDelta(DeltaRot);
+			InteractionScreenCurrPos = InDragPos.ScreenPosition;
+			break;
+		}
+	case EAxisRotateMode::Arc:
+		{
+			float HitDepth;
+			if (GetRayParamIntersectionWithInteractionPlane(InDragPos, HitDepth))
+			{
+				const FVector HitPoint = InDragPos.WorldRay.Origin + InDragPos.WorldRay.Direction * HitDepth;
+				const float HitAngle = GizmoMath::ComputeAngleInPlane(HitPoint,
+					InteractionPlanarOrigin, InteractionPlanarNormal, InteractionPlanarAxisX, InteractionPlanarAxisY);
+
+				const float DeltaAngle = FMath::RadiansToDegrees(InteractionCurrAngle - HitAngle);
+				const FVector DeltaRot(
+					InteractionAxisList == EAxisList::X ? DeltaAngle : 0.0,
+					InteractionAxisList == EAxisList::Y ? DeltaAngle : 0.0,
+					InteractionAxisList == EAxisList::Z ? DeltaAngle : 0.0);
+	
+				static constexpr bool bUseUEHandyness = true;
+				const FQuat Delta = AnimationCore::QuatFromEuler(DeltaRot, GetRotationContext().RotationOrder, bUseUEHandyness);
+	
+				ApplyRotateDelta(Delta);
+				InteractionCurrAngle = HitAngle;
+			}
+			break;
+		}
+	default:
+		ensure(false);
+		break;
+	}
+}
+
 void UTransformGizmo::OnClickPressScreenSpaceRotate(const FInputDeviceRay& InPressPos)
 {
 	check(GizmoViewContext);
@@ -2434,6 +2755,14 @@ namespace ArcBallLocals
 
 void UTransformGizmo::OnClickPressArcBallRotate(const FInputDeviceRay& PressPos)
 {
+	if (bGimbalRotationMode)
+	{
+		// arc ball rotation is disabled in gimbal mode but is still rendered to ease the visualization 
+		SetModeLastHitPart(EGizmoTransformMode::Rotate, LastHitPart);
+		bInInteraction = false;
+		return;
+	}
+	
 	check(GizmoViewContext);
 	
 	const FVector& RayOrigin = PressPos.WorldRay.Origin;
diff --git a/Engine/Source/Editor/Experimental/EditorInteractiveToolsFramework/Private/EditorInteractiveGizmoManager.cpp b/Engine/Source/Editor/Experimental/EditorInteractiveToolsFramework/Private/EditorInteractiveGizmoManager.cpp
--- a/Engine/Source/Editor/Experimental/EditorInteractiveToolsFramework/Private/EditorInteractiveGizmoManager.cpp
+++ b/Engine/Source/Editor/Experimental/EditorInteractiveToolsFramework/Private/EditorInteractiveGizmoManager.cpp
@@ -70,6 +70,11 @@ const TOptional<FGizmosParameters>& UEditorInteractiveGizmoManager::GetDefaultGi
 	return GizmoManagerLocals::OptDefaultParameters.IsSet() ? GizmoManagerLocals::OptDefaultParameters : Invalid;
 }
 
+bool UEditorInteractiveGizmoManager::IsExplicitModeEnabled()
+{
+	return GizmoManagerLocals::OptDefaultParameters.IsSet() ? GizmoManagerLocals::OptDefaultParameters->bEnableExplicit : false;
+}
+
 UEditorInteractiveGizmoManager::UEditorInteractiveGizmoManager() :
 	UInteractiveGizmoManager()
 {
diff --git a/Engine/Source/Editor/Experimental/EditorInteractiveToolsFramework/Private/EditorInteractiveToolsFrameworkModule.cpp b/Engine/Source/Editor/Experimental/EditorInteractiveToolsFramework/Private/EditorInteractiveToolsFrameworkModule.cpp
--- a/Engine/Source/Editor/Experimental/EditorInteractiveToolsFramework/Private/EditorInteractiveToolsFrameworkModule.cpp
+++ b/Engine/Source/Editor/Experimental/EditorInteractiveToolsFramework/Private/EditorInteractiveToolsFrameworkModule.cpp
@@ -3,6 +3,7 @@
 #include "EditorInteractiveToolsFrameworkModule.h"
 
 #include "ComponentSourceInterfaces.h"
+#include "EditorGizmos/GizmoRotationUtil.h"
 #include "Modules/ModuleManager.h"
 #include "Templates/UniquePtr.h"
 #include "Tools/EditorComponentSourceFactory.h"
@@ -20,6 +21,11 @@ void FEditorInteractiveToolsFrameworkModule::StartupModule()
 	// The Interactive Tools Framework needs to be explicitly loaded, at the very least to make sure that
 	// the InteractiveToolsSelectionStoreSubsystem gets initialized.
 	FModuleManager::Get().LoadModule(TEXT("InteractiveToolsFramework"));
+
+	// register relative transform interface
+	using namespace UE::GizmoRotationUtil;
+	FRelativeTransformInterfaceRegistry& Registry = FRelativeTransformInterfaceRegistry::Get();
+	Registry.RegisterDefaultInterfaces();
 }
 
 void FEditorInteractiveToolsFrameworkModule::ShutdownModule()
diff --git a/Engine/Source/Editor/Experimental/EditorInteractiveToolsFramework/Public/EditorGizmos/EditorTransformGizmoSource.h b/Engine/Source/Editor/Experimental/EditorInteractiveToolsFramework/Public/EditorGizmos/EditorTransformGizmoSource.h
--- a/Engine/Source/Editor/Experimental/EditorInteractiveToolsFramework/Public/EditorGizmos/EditorTransformGizmoSource.h
+++ b/Engine/Source/Editor/Experimental/EditorInteractiveToolsFramework/Public/EditorGizmos/EditorTransformGizmoSource.h
@@ -75,6 +75,9 @@ public:
 	 */
 	virtual EGizmoTransformScaleType GetScaleType() const override;
 
+	/** Get rotation context */
+	virtual const FRotationContext& GetRotationContext() const override;
+	
 	static UEditorTransformGizmoSource* CreateNew(
 		UObject* Outer = (UObject*)GetTransientPackage(),
 		const UEditorTransformGizmoContextObject* InContext = nullptr);
diff --git a/Engine/Source/Editor/Experimental/EditorInteractiveToolsFramework/Public/EditorGizmos/EditorTransformGizmoUtil.h b/Engine/Source/Editor/Experimental/EditorInteractiveToolsFramework/Public/EditorGizmos/EditorTransformGizmoUtil.h
--- a/Engine/Source/Editor/Experimental/EditorInteractiveToolsFramework/Public/EditorGizmos/EditorTransformGizmoUtil.h
+++ b/Engine/Source/Editor/Experimental/EditorInteractiveToolsFramework/Public/EditorGizmos/EditorTransformGizmoUtil.h
@@ -2,11 +2,13 @@
 
 #pragma once
 
+#include "TransformGizmoInterfaces.h"
 #include "UObject/Object.h"
 #include "Containers/UnrealString.h"
 
 #include "EditorTransformGizmoUtil.generated.h"
 
+struct FGizmosParameters;
 class UInteractiveToolsContext;
 class UInteractiveToolManager;
 class UTransformGizmo;
@@ -90,6 +92,8 @@ public:
 	DECLARE_MULTICAST_DELEGATE_OneParam(FOnGizmoCreated, UTransformGizmo*);
 	FOnGizmoCreated& OnGizmoCreatedDelegate();
 
+	UPROPERTY(Transient)
+	FRotationContext RotationContext;
 	
 private:
 
@@ -109,6 +113,7 @@ private:
 	FEditorModeTools* ModeTools = nullptr;
 	FDelegateHandle ViewportClientsChangedHandle;
 	FDelegateHandle UseNewGizmosChangedHandled;
+	FDelegateHandle GizmosParametersChangedHandle;
 	bool bGizmosRegistered = false;
 
 	TSharedPtr<FEditorTransformGizmoDataBinder> DataBinder;
diff --git a/Engine/Source/Editor/Experimental/EditorInteractiveToolsFramework/Public/EditorGizmos/GizmoElementGimbal.h b/Engine/Source/Editor/Experimental/EditorInteractiveToolsFramework/Public/EditorGizmos/GizmoElementGimbal.h
new file mode 100644
--- /dev/null
+++ b/Engine/Source/Editor/Experimental/EditorInteractiveToolsFramework/Public/EditorGizmos/GizmoElementGimbal.h
@@ -0,0 +1,31 @@
+// Copyright Epic Games, Inc. All Rights Reserved.
+
+#pragma once
+
+#include "BaseGizmos/GizmoElementGroup.h"
+#include "TransformGizmoInterfaces.h"
+
+#include "GizmoElementGimbal.generated.h"
+
+/**
+ * Gimbal rotation group object intended to be used as part of 3D Rotation Gizmos
+ * This group expects three rotation sub-elements
+ */
+UCLASS(Transient, MinimalAPI)
+class UGizmoElementGimbal : public UGizmoElementGroup
+{
+	GENERATED_BODY()
+
+public:
+
+	//~ Begin UGizmoElementBase Interface.
+	EDITORINTERACTIVETOOLSFRAMEWORK_API virtual void Render(IToolsContextRenderAPI* InRenderAPI, const FRenderTraversalState& InRenderState) override;
+	EDITORINTERACTIVETOOLSFRAMEWORK_API virtual FInputRayHit LineTrace(const UGizmoViewContext* InViewContext, const FLineTraceTraversalState& InLineTraceState, const FVector& InRayOrigin, const FVector& InRayDirection) override;
+	//~ End UGizmoElementBase Interface.
+
+	// Add object to group.
+	EDITORINTERACTIVETOOLSFRAMEWORK_API virtual void Add(UGizmoElementBase* InElement) override;
+	
+	UPROPERTY()
+	FRotationContext RotationContext;
+};
diff --git a/Engine/Source/Editor/Experimental/EditorInteractiveToolsFramework/Public/EditorGizmos/GizmoRotationUtil.h b/Engine/Source/Editor/Experimental/EditorInteractiveToolsFramework/Public/EditorGizmos/GizmoRotationUtil.h
new file mode 100644
--- /dev/null
+++ b/Engine/Source/Editor/Experimental/EditorInteractiveToolsFramework/Public/EditorGizmos/GizmoRotationUtil.h
@@ -0,0 +1,115 @@
+﻿// Copyright Epic Games, Inc. All Rights Reserved.
+
+#pragma once
+
+#include "Math/Transform.h"
+#include "Elements/Interfaces/TypedElementWorldInterface.h"
+
+struct FRotationContext;
+struct FEulerTransform;
+
+namespace UE::GizmoRotationUtil
+{
+	/**
+	 * FRotationDecomposition is a data storage structure representing the rotations around the X, Y and Z axis  
+	 */	
+	struct FRotationDecomposition
+	{
+		FRotationDecomposition()
+		{
+			R[0] = R[1] = R[2] = FQuat::Identity;
+		}
+		FQuat R[3];
+	};
+
+	/**
+	 * Decomposes InRotationContext.Rotation in respect to the rotation order to output the explicit per axis rotations.
+	 */
+	void DecomposeRotations(const FTransform& InTransform, const FRotationContext& InRotationContext, FRotationDecomposition& OutDecomposition);
+
+	/**
+	 * Returns the explicit rotation axis in respect to InRotationContext.Rotation & InRotationContext.RotationOrder
+	 */
+	FVector GetRotationAxis(const FTransform& InTransform, const FRotationContext& InRotationContext, const int32 InAxis);
+}
+
+namespace UE::GizmoRotationUtil
+{
+	/**
+	 * IRelativeTransformInterface
+	 */
+	struct IRelativeTransformInterface
+	{
+		virtual ~IRelativeTransformInterface() = default;
+
+		virtual bool GetRelativeTransform(const FTypedElementHandle& InElementHandle, FEulerTransform& OutRelativeTransform) = 0;
+		virtual bool SetRelativeTransform(const FTypedElementHandle& InElementHandle, const FEulerTransform& InRelativeTransform) = 0;
+	};
+
+	/**
+	 * FRelativeTransformInterfaceRegistry
+	 */
+	class FRelativeTransformInterfaceRegistry
+	{
+	public:
+		~FRelativeTransformInterfaceRegistry() = default;
+
+		/** Get the singleton registry object */
+		static EDITORINTERACTIVETOOLSFRAMEWORK_API FRelativeTransformInterfaceRegistry& Get();
+
+		/** Register an interface for the TWorldInterface static class */
+		template<typename TWorldInterface>
+		void RegisterRelativeTransformInterface(TUniquePtr<IRelativeTransformInterface>&& InInterface)
+		{
+			UClass* WorldInterfaceClass = TWorldInterface::StaticClass();
+			check(!WorldInterfaceToRelativeTransformInterface.Contains(WorldInterfaceClass));
+			WorldInterfaceToRelativeTransformInterface.Add(WorldInterfaceClass, MoveTemp(InInterface));
+		}
+
+		static void RegisterDefaultInterfaces();
+
+		/** Find the registered interface from the given class. Returns a nullptr if noting registered for that class. */
+		IRelativeTransformInterface* FindRelativeTransformInterface(const TTypedElement<ITypedElementWorldInterface>& InElement) const;
+		IRelativeTransformInterface* FindRelativeTransformInterface(const UClass* InClass) const;
+
+	private:
+		FRelativeTransformInterfaceRegistry() = default;
+		
+		TMap<UClass*, TUniquePtr<IRelativeTransformInterface>> WorldInterfaceToRelativeTransformInterface;
+	};
+
+	/**
+	* Handles get / set relative transform operations UActorElementWorldInterface & UActorElementEditorWorldInterface
+	*/
+	struct FActorRelativeTransformInterface : IRelativeTransformInterface
+	{
+		virtual ~FActorRelativeTransformInterface() override = default;
+
+		virtual bool GetRelativeTransform(const FTypedElementHandle& InElementHandle, FEulerTransform& OutRelativeTransform) override;
+		virtual bool SetRelativeTransform(const FTypedElementHandle& InElementHandle, const FEulerTransform& InRelativeTransform) override;
+
+	private:
+		static USceneComponent* GetSceneComponent(const FTypedElementHandle& InElementHandle);
+	};
+
+	/**
+	* Handles get / set relative transform operations UActorElementWorldInterface & UActorElementEditorWorldInterface
+	*/
+	struct FComponentRelativeTransformInterface : IRelativeTransformInterface
+	{
+		virtual ~FComponentRelativeTransformInterface() override = default;
+		
+		virtual bool GetRelativeTransform(const FTypedElementHandle& InElementHandle, FEulerTransform& OutRelativeTransform) override;
+		virtual bool SetRelativeTransform(const FTypedElementHandle& InElementHandle, const FEulerTransform& InRelativeTransform) override;
+
+	private:
+		static USceneComponent* GetSceneComponent(const FTypedElementHandle& InElementHandle);
+	};
+
+	/**
+ * Returns the relative transform of an object by storing its explicit relative rotation in the provided FRotationContext
+ */
+	EDITORINTERACTIVETOOLSFRAMEWORK_API bool GetRelativeTransform(const TTypedElement<ITypedElementWorldInterface>& InElement, FTransform& OutTransform, FRotationContext& InOutRotationContext);
+}
+
+
diff --git a/Engine/Source/Editor/Experimental/EditorInteractiveToolsFramework/Public/EditorGizmos/TransformGizmo.h b/Engine/Source/Editor/Experimental/EditorInteractiveToolsFramework/Public/EditorGizmos/TransformGizmo.h
--- a/Engine/Source/Editor/Experimental/EditorInteractiveToolsFramework/Public/EditorGizmos/TransformGizmo.h
+++ b/Engine/Source/Editor/Experimental/EditorInteractiveToolsFramework/Public/EditorGizmos/TransformGizmo.h
@@ -52,6 +52,7 @@ class UMaterialInstanceDynamic;
 class UMaterialInterface;
 class UObject;
 class UTransformProxy;
+class UGizmoElementGimbal;
 
 //
 // Interface for the Transform gizmo.
@@ -82,6 +83,9 @@ enum class ETransformGizmoPartIdentifier
 	RotateZAxis,
 	RotateScreenSpace,
 	RotateArcball,
+	RotateXGimbal,
+	RotateYGimbal,
+	RotateZGimbal,
 	ScaleAll,
 	ScaleXAxis,
 	ScaleYAxis,
@@ -115,6 +119,10 @@ struct FGizmosParameters
 	/** Property used to define whether the Y axis is indirectly controlled via CTRL+MMB or CTRL+RMB (and Z via CTRL+RMB or CTRL+LMB+RMB). */
 	UPROPERTY(EditAnywhere, Category = NewTRSGizmo)
 	bool bCtrlMiddleDoesY = true;
+
+	/** Enable/Disable the explicit coordinate system. */
+	UPROPERTY(EditAnywhere, Category = NewTRSGizmo)
+	bool bEnableExplicit = false;
 };
 
 /**
@@ -354,9 +362,21 @@ protected:
 	UPROPERTY()
 	TObjectPtr<UGizmoElementTorus> RotateZAxisElement;
 
-	/** Rotate outer circle */
+	/** Gimbal Rotate X Axis */
+	UPROPERTY()
+	TObjectPtr<UGizmoElementTorus> RotateXGimbalElement;
+
+	/** Gimbal Rotate Y Axis */
+	UPROPERTY()
+	TObjectPtr<UGizmoElementTorus> RotateYGimbalElement;
+	
+	/** Gimbal Rotate Z Axis */
 	UPROPERTY()
-	TObjectPtr<UGizmoElementCircle> RotateOuterCircleElement;
+	TObjectPtr<UGizmoElementTorus> RotateZGimbalElement;
+
+	/** Gimbal Rotate root element */
+	UPROPERTY()
+	TObjectPtr<UGizmoElementGimbal> RotateGimbalElement;
 
 	/** Rotate arcball inner circle */
 	UPROPERTY()
@@ -431,6 +451,9 @@ protected:
 	/** Update current gizmo mode based on transform source */
 	virtual void UpdateMode();
 
+	/** Update current gizmo rotation mode (default or gimbal)*/
+	void UpdateRotationMode();
+
 	/** Enable the given mode with the specified axes, EAxisList::Type::None will hide objects associated with mode */
 	virtual void EnableMode(EGizmoTransformMode InGizmoMode, EAxisList::Type InAxisListToDraw);
 
@@ -452,6 +475,9 @@ protected:
 	/** Construct scale axis handle */
 	virtual UGizmoElementArrow* MakeScaleAxis(ETransformGizmoPartIdentifier InPartId, const FVector& InAxisDir, const FVector& InSideDir, UMaterialInterface* InMaterial);
 
+	/** Construct the default rotate axis handle based on the ETransformGizmoPartIdentifier. */
+	UGizmoElementTorus* MakeDefaultRotateAxis(const ETransformGizmoPartIdentifier InPartId);
+
 	/** Construct rotate axis handle */
 	virtual UGizmoElementTorus* MakeRotateAxis(ETransformGizmoPartIdentifier InPartId, const FVector& TorusAxis0, const FVector& TorusAxis1,
 		UMaterialInterface* InMaterial, UMaterialInterface* InCurrentMaterial);
@@ -476,7 +502,10 @@ protected:
 	virtual FInputRayHit UpdateHoveredPart(const FInputDeviceRay& DevicePos);
 
 	/** Get current interaction axis */
-	virtual FVector GetWorldAxis(const FVector& InAxis);
+	virtual FVector GetWorldAxis(const FVector& InAxis) const;
+
+	/** Get gimbal rotation axis (InAxis is expected to be in [0,2] range) */
+	FVector GetGimbalRotationAxis(const int32 InAxis) const;
 
 	/** Get nearest param along input ray to the current interaction axis */
 	virtual float GetNearestRayParamToInteractionAxis(const FInputDeviceRay& InRay);
@@ -589,14 +618,34 @@ protected:
 
 	/** Get screen-space axis for rotation drag */
 	FVector2D GetScreenRotateAxisDir(const FInputDeviceRay& InPressPos);
+	
+	/** Get screen-space axis for gimbal rotation drag */
+	FVector2D GetScreenGimbalRotateAxisDir(const FInputDeviceRay& InPressPos);
+
+	/** Get screen-space axis based on a world space axis */
+	FVector2D GetWorldToScreenRotateAxisDir(const FInputDeviceRay& InPressPos, const FVector& InWorldAxis);
 
 	/** Compute rotate delta based on screen-space start/end positions */
 	virtual FQuat ComputeAxisRotateDelta(const FVector2D& InStartPos, const FInputDeviceRay& InDragPos);
 
+	/** Compute gimbal rotate delta based on screen-space start/end positions */
+	FQuat ComputeGimbalRotateDelta(const FVector2D& InStartPos, const FVector2D& InEndPos);
+
 	/** Prepares data for arc rotation. This will return false if this is not possible (the rotate handle is perpendicular to the view) */
 	bool OnClickPressRotateArc( const FInputDeviceRay& InPressPos,
 		const FVector& InPlaneNormal, const FVector& InPlaneAxis1, const FVector& InPlaneAxis2);
 
+	/**
+	 * Gimbal rotate interaction methods
+	 * note that gimbal rotation release functions use OnClickReleaseRotateAxis
+	 */
+
+	/** Handle click press for any gimbal rotate axis */
+	void OnClickPressGimbalRotateAxis(const FInputDeviceRay& InPressPos);
+
+	/** Handle click drag for gimbal rotate axis */
+	virtual void OnClickDragGimbalRotateAxis(const FInputDeviceRay& InDragPos);
+	
 	/**
 	 * Screen-space rotate interaction methods
 	 */
@@ -719,6 +768,9 @@ protected:
 
 	// Returns whether the gizmo can interact. Note that this can be true even if the gizmo is not visible to support indirect manipulation.
 	bool CanInteract(const EViewportContext InViewportContext = EViewportContext::Focused) const;
+
+	// Returns the current rotation context.
+	const FRotationContext& GetRotationContext() const;
 	
 protected:
 
@@ -901,6 +953,9 @@ protected:
 	/** Switch from tangential to normal projection based on the first mouse drag. */
 	bool bTrySwitchingToNormalPull = false;
 
+	/** Used to check if the gimbal mode is currently active (this is updated when ticking the gizmo) */
+	bool bGimbalRotationMode = false;
+	
 private:
 	/** Debug attributes to display the pull direction */
 	bool bDebugRotate = false;
diff --git a/Engine/Source/Editor/Experimental/EditorInteractiveToolsFramework/Public/EditorGizmos/TransformGizmoInterfaces.h b/Engine/Source/Editor/Experimental/EditorInteractiveToolsFramework/Public/EditorGizmos/TransformGizmoInterfaces.h
--- a/Engine/Source/Editor/Experimental/EditorInteractiveToolsFramework/Public/EditorGizmos/TransformGizmoInterfaces.h
+++ b/Engine/Source/Editor/Experimental/EditorInteractiveToolsFramework/Public/EditorGizmos/TransformGizmoInterfaces.h
@@ -5,6 +5,7 @@
 #include "CoreMinimal.h"
 #include "ToolContextInterfaces.h"
 #include "UnrealWidgetFwd.h"
+#include "EulerTransform.h"
 #include "TransformGizmoInterfaces.generated.h"
 
 class UMaterial;
@@ -35,6 +36,32 @@ enum class EViewportContext : uint8
 	Focused
 };
 
+/**
+ * FRotationContext is a datastore for rotation data sent to the gizmo
+ */
+
+USTRUCT()
+struct FRotationContext
+{
+	GENERATED_BODY()
+
+	/** Enable explicit rotation axis */
+	UPROPERTY()
+	bool bUseExplicitRotator = false;
+
+	/** Order in which this rotation is expressed */
+	UPROPERTY()
+	EEulerRotationOrder RotationOrder = EEulerRotationOrder::XYZ;
+
+	/** Rotation values */
+	UPROPERTY()
+	FRotator Rotation = FRotator::ZeroRotator;
+
+	/** Parent offset transform */
+	UPROPERTY()
+	FTransform Offset = FTransform::Identity;
+};
+
 UINTERFACE()
 class EDITORINTERACTIVETOOLSFRAMEWORK_API UTransformGizmoSource : public UInterface
 {
@@ -82,6 +109,11 @@ public:
 	 * Get current scale type
 	 */
 	virtual EGizmoTransformScaleType GetScaleType() const PURE_VIRTUAL(ITransformGizmoSource::GetScaleType, return EGizmoTransformScaleType::Default;);
+
+	/*
+	 * Get rotation context
+	 */
+	virtual const FRotationContext& GetRotationContext() const PURE_VIRTUAL(ITransformGizmoSource::GetRotationContext, static const FRotationContext Context; return Context;);
 };
 
 /**
diff --git a/Engine/Source/Editor/Experimental/EditorInteractiveToolsFramework/Public/EditorInteractiveGizmoManager.h b/Engine/Source/Editor/Experimental/EditorInteractiveToolsFramework/Public/EditorInteractiveGizmoManager.h
--- a/Engine/Source/Editor/Experimental/EditorInteractiveToolsFramework/Public/EditorInteractiveGizmoManager.h
+++ b/Engine/Source/Editor/Experimental/EditorInteractiveToolsFramework/Public/EditorInteractiveGizmoManager.h
@@ -166,6 +166,11 @@ public:
 	 */
     static const TOptional<FGizmosParameters>& GetDefaultGizmosParameters();
 
+	/**
+	 * Returns true if the Explicit mode is enabled;
+	 */
+	static bool IsExplicitModeEnabled();
+
 protected:
 
 	/**
diff --git a/Engine/Source/Editor/UnrealEd/Private/EditorModeManager.cpp b/Engine/Source/Editor/UnrealEd/Private/EditorModeManager.cpp
--- a/Engine/Source/Editor/UnrealEd/Private/EditorModeManager.cpp
+++ b/Engine/Source/Editor/UnrealEd/Private/EditorModeManager.cpp
@@ -46,6 +46,7 @@
 #include "TextureResource.h"
 #include "EditorGizmos/EditorGizmoStateTarget.h"
 #include "EditorGizmos/EditorTransformGizmoUtil.h"
+#include "EditorGizmos/GizmoRotationUtil.h"
 #include "Toolkits/ToolkitManager.h"
 
 /*------------------------------------------------------------------------------
@@ -125,7 +126,8 @@ FEditorModeTools::~FEditorModeTools()
 		InteractiveToolsContext->Deactivate();
 		InteractiveToolsContext->ShutdownContext();
 		InteractiveToolsContext = nullptr;
-		GizmoStateTarget = nullptr;
+		WeakGizmoStateTarget = nullptr;
+		WeakGizmoContext = nullptr;
 	}
 }
 
@@ -139,7 +141,8 @@ void FEditorModeTools::LoadConfig(void)
 	GConfig->GetInt(TEXT("FEditorModeTools"),TEXT("CoordSystem"), CoordSystemAsInt,
 		GEditorPerProjectIni);
 
-	if (static_cast<ECoordSystem>(CoordSystemAsInt) == COORD_Parent)
+	const ECoordSystem CoordinateSystem = static_cast<ECoordSystem>(CoordSystemAsInt);
+	if (CoordinateSystem == COORD_Parent || CoordinateSystem == COORD_Explicit)
 	{
 		// parent mode is only supported with new trs gizmos for now
 		const bool bUseNewGizmo = UEditorInteractiveGizmoManager::UsesNewTRSGizmos();
@@ -1022,6 +1025,15 @@ FMatrix FEditorModeTools::GetCustomDrawingCoordinateSystem() const
 
 	switch (GetCoordSystem())
 	{
+		case COORD_World:
+		{
+			if (UEditorTransformGizmoContextObject* GizmoContextObject = WeakGizmoContext.Get())
+			{
+				GizmoContextObject->RotationContext = FRotationContext();
+			}
+		}
+		break;
+		
 		case COORD_Local:
 			Matrix = GetLocalCoordinateSystem();
 		break;
@@ -1030,8 +1042,9 @@ FMatrix FEditorModeTools::GetCustomDrawingCoordinateSystem() const
 			Matrix = GetParentSpaceCoordinateSystem();
 		break;
 
-		case COORD_World:
-			break;
+		case COORD_Explicit:
+			Matrix = GetExplicitCoordinateSystem();
+		break;
 
 		default:
 			break;
@@ -1071,6 +1084,26 @@ FMatrix FEditorModeTools::GetParentSpaceCoordinateSystem() const
 	});
 }
 
+FMatrix FEditorModeTools::GetExplicitCoordinateSystem() const
+{
+	return GetCustomCoordinateSystem([this](const TTypedElement<ITypedElementWorldInterface>& InElement, FTransform& OutTransform)
+	{
+		if (UEditorTransformGizmoContextObject* GizmoContext = GetGizmoContext())
+		{
+			FRotationContext& RotationContext = GizmoContext->RotationContext;
+			
+			RotationContext = FRotationContext();
+			RotationContext.bUseExplicitRotator = GetWidgetMode() == UE::Widget::EWidgetMode::WM_Rotate;
+
+			if (UE::GizmoRotationUtil::GetRelativeTransform(InElement, OutTransform, RotationContext))
+			{
+				return true;
+			}
+		}
+		return false;
+	});
+}
+
 FMatrix FEditorModeTools::GetCustomCoordinateSystem(TUniqueFunction<void(const TTypedElement<ITypedElementWorldInterface>&, FTransform&)>&& InGetTransformFunc) const
 {
 	FMatrix Matrix = FMatrix::Identity;
@@ -1086,6 +1119,11 @@ FMatrix FEditorModeTools::GetCustomCoordinateSystem(TUniqueFunction<void(const T
 	// If there isn't an active mode overriding the local coordinate system, create it by looking at the current selection.
 	if (!CustomCoordinateSystemProvided)
 	{
+		if (UEditorTransformGizmoContextObject* GizmoContextObject = WeakGizmoContext.Get())
+		{
+			GizmoContextObject->RotationContext = FRotationContext();
+		}
+		
 		TTypedElement<ITypedElementWorldInterface> LastSelected;
 		if ((this == &GLevelEditorModeTools()) && GCurrentLevelEditingViewportClient)
 		{
@@ -2018,14 +2056,31 @@ UModeManagerInteractiveToolsContext* FEditorModeTools::GetInteractiveToolsContex
 
 IGizmoStateTarget* FEditorModeTools::GetGizmoStateTarget()
 {
-	if (!GizmoStateTarget.IsValid())
+	if (!WeakGizmoStateTarget.IsValid())
 	{
-		GizmoStateTarget = UEditorGizmoStateTarget::Construct(
+		WeakGizmoStateTarget = UEditorGizmoStateTarget::Construct(
 			this,
 			NSLOCTEXT("UTransformGizmo", "UTransformGizmoTransaction", "Transform"),
 			InteractiveToolsContext->GizmoManager);
 	}
-	return GizmoStateTarget.Get();
+	return WeakGizmoStateTarget.Get();
+}
+
+UEditorTransformGizmoContextObject* FEditorModeTools::GetGizmoContext() const
+{
+	if (!WeakGizmoContext.IsValid())
+	{
+		const UModeManagerInteractiveToolsContext* ToolsContext = GetInteractiveToolsContext();
+		const UContextObjectStore* ContextStore = ToolsContext ? ToolsContext->ToolManager->GetContextObjectStore() : nullptr;
+		WeakGizmoContext = ContextStore ? ContextStore->FindContext<UEditorTransformGizmoContextObject>() : nullptr;
+	}
+	return WeakGizmoContext.Get();
+}
+
+const FRotationContext& FEditorModeTools::GetRotationContext() const
+{
+	static const FRotationContext DefaultContext;
+	return WeakGizmoContext.IsValid() ? WeakGizmoContext->RotationContext : DefaultContext;
 }
 
 bool FEditorModeTools::BeginTransform(const FGizmoState& InState)
diff --git a/Engine/Source/Editor/UnrealEd/Private/EditorViewportCommands.cpp b/Engine/Source/Editor/UnrealEd/Private/EditorViewportCommands.cpp
--- a/Engine/Source/Editor/UnrealEd/Private/EditorViewportCommands.cpp
+++ b/Engine/Source/Editor/UnrealEd/Private/EditorViewportCommands.cpp
@@ -174,6 +174,13 @@ void FEditorViewportCommands::RegisterCommands()
 		EUserInterfaceActionType::RadioButton,
 		FInputChord()
 	);
+	UI_COMMAND(
+		RelativeCoordinateSystem_Explicit,
+		"Explicit Space",
+		"Move objects relative to the parent's axes and rotate them using gimbal axes",
+		EUserInterfaceActionType::RadioButton,
+		FInputChord()
+	);
 
 #if PLATFORM_MAC
 	UI_COMMAND( CycleTransformGizmoCoordSystem, "Cycle Transform Coordinate System", "Cycles the transform gizmo coordinate systems between world and local (object) space", EUserInterfaceActionType::Button, FInputChord(EModifierKey::Command, EKeys::Tilde));
diff --git a/Engine/Source/Editor/UnrealEd/Private/Elements/Actor/ActorElementEditorViewportInteractionCustomization.cpp b/Engine/Source/Editor/UnrealEd/Private/Elements/Actor/ActorElementEditorViewportInteractionCustomization.cpp
--- a/Engine/Source/Editor/UnrealEd/Private/Elements/Actor/ActorElementEditorViewportInteractionCustomization.cpp
+++ b/Engine/Source/Editor/UnrealEd/Private/Elements/Actor/ActorElementEditorViewportInteractionCustomization.cpp
@@ -8,10 +8,11 @@
 #include "EditorModeManager.h"
 #include "AI/NavigationSystemBase.h"
 #include "Components/BrushComponent.h"
+#include "EditorGizmos/EditorTransformGizmoUtil.h"
 #include "Engine/DocumentationActor.h"
 #include "Particles/Emitter.h"
 #include "Elements/Interfaces/TypedElementHierarchyInterface.h"
-#include "Elements/Framework/TypedElementRegistry.h"
+#include "Toolkits/IToolkitHost.h"
 
 void FActorElementEditorViewportInteractionCustomization::GizmoManipulationStarted(
 	const TTypedElement<ITypedElementWorldInterface>& InElementWorldHandle,
@@ -70,14 +71,32 @@ void FActorElementEditorViewportInteractionCustomization::GizmoManipulationDelta
 	const bool bRotatingActor = !DeltaRotation.IsZero();
 	if (bRotatingActor)
 	{			
-		if (RootActor->GetRootComponent())
+		if (USceneComponent* RootComponent = RootActor->GetRootComponent())
 		{
 			const FRotator OriginalRotation = RootActor->GetRootComponent()->GetComponentRotation();
 
-			RootActor->EditorApplyRotation(DeltaRotation, InInputState.bAltKeyDown, InInputState.bShiftKeyDown, InInputState.bCtrlKeyDown);
+			auto ShouldApplyDeltaDirectly = [ToolkitHostPtr = GetToolkitHost()]()
+			{
+				if (ToolkitHostPtr)
+				{
+					const FEditorModeTools& ModeManager = ToolkitHostPtr->GetEditorModeManager();
+					return ModeManager.GetRotationContext().bUseExplicitRotator;
+				}
+				return false;
+			};
+
+			if (ShouldApplyDeltaDirectly())
+			{
+				const FRotator RelativeRot = RootComponent->GetRelativeRotation();
+				RootComponent->SetRelativeRotationExact(RelativeRot + DeltaRotation);
+			}
+			else
+			{
+				RootActor->EditorApplyRotation(DeltaRotation, InInputState.bAltKeyDown, InInputState.bShiftKeyDown, InInputState.bCtrlKeyDown);
+			}
 
 			// Check to see if we should transform the rigid body
-			UPrimitiveComponent* RootPrimitiveComponent = Cast<UPrimitiveComponent>(RootActor->GetRootComponent());
+			UPrimitiveComponent* RootPrimitiveComponent = Cast<UPrimitiveComponent>(RootComponent);
 			if (bIsSimulatingInEditor && GIsPlayInEditorWorld && RootPrimitiveComponent)
 			{
 				FRotator ActorRotWind, ActorRotRem;
diff --git a/Engine/Source/Editor/UnrealEd/Private/Elements/Component/ComponentElementEditorViewportInteractionCustomization.cpp b/Engine/Source/Editor/UnrealEd/Private/Elements/Component/ComponentElementEditorViewportInteractionCustomization.cpp
--- a/Engine/Source/Editor/UnrealEd/Private/Elements/Component/ComponentElementEditorViewportInteractionCustomization.cpp
+++ b/Engine/Source/Editor/UnrealEd/Private/Elements/Component/ComponentElementEditorViewportInteractionCustomization.cpp
@@ -8,6 +8,9 @@
 #include "Editor.h"
 #include "EditorModeManager.h"
 #include "EditorSupportDelegates.h"
+#include "EditorGizmos/TransformGizmoInterfaces.h"
+
+#include "Toolkits/IToolkitHost.h"
 
 bool FComponentElementEditorViewportInteractionCustomization::GetGizmoPivotLocation(const TTypedElement<ITypedElementWorldInterface>& InElementWorldHandle, const UE::Widget::EWidgetMode InWidgetMode, FVector& OutPivotLocation)
 {
@@ -38,7 +41,35 @@ void FComponentElementEditorViewportInteractionCustomization::GizmoManipulationD
 		const FRotator DeltaRotation = InDeltaTransform.Rotator();
 		const FVector DeltaScale3D = InDeltaTransform.GetScale3D();
 
-		FComponentElementEditorViewportInteractionCustomization::ApplyDeltaToComponent(SceneComponent, /*bDelta*/true, &DeltaTranslation, &DeltaRotation, &DeltaScale3D, InPivotLocation, InInputState);
+		auto ShouldApplyDeltaDirectly = [ToolkitHostPtr = GetToolkitHost()]()
+		{
+			if (ToolkitHostPtr)
+			{
+				const FEditorModeTools& ModeManager = ToolkitHostPtr->GetEditorModeManager();
+				return ModeManager.GetRotationContext().bUseExplicitRotator;
+			}
+			return false;
+		};
+
+		const FRotator* DeltaRotationPtr = &DeltaRotation;
+		if (ShouldApplyDeltaDirectly())
+		{
+			//
+			const FRotator RelativeRot = SceneComponent->GetRelativeRotation();
+			SceneComponent->SetRelativeRotationExact(RelativeRot + DeltaRotation);
+
+			//
+			FVector NewCompLocation = SceneComponent->GetRelativeLocation();
+			NewCompLocation -= InPivotLocation;
+			NewCompLocation = FRotationMatrix(DeltaRotation).TransformPosition(NewCompLocation);
+			NewCompLocation += InPivotLocation;
+			SceneComponent->SetRelativeLocation(NewCompLocation);
+
+			//
+			DeltaRotationPtr = nullptr;
+		}
+		
+		FComponentElementEditorViewportInteractionCustomization::ApplyDeltaToComponent(SceneComponent, /*bDelta*/true, &DeltaTranslation, DeltaRotationPtr, &DeltaScale3D, InPivotLocation, InInputState);
 	}
 }
 
diff --git a/Engine/Source/Editor/UnrealEd/Private/SEditorViewport.cpp b/Engine/Source/Editor/UnrealEd/Private/SEditorViewport.cpp
--- a/Engine/Source/Editor/UnrealEd/Private/SEditorViewport.cpp
+++ b/Engine/Source/Editor/UnrealEd/Private/SEditorViewport.cpp
@@ -432,6 +432,13 @@ void SEditorViewport::BindCommands()
 		FIsActionChecked::CreateSP(this, &SEditorViewport::IsCoordSystemActive, COORD_Parent)
 	);
 
+	CommandListRef.MapAction(
+		Commands.RelativeCoordinateSystem_Explicit,
+		FExecuteAction::CreateSP(ClientRef, &FEditorViewportClient::SetWidgetCoordSystemSpace, COORD_Explicit),
+		FCanExecuteAction(),
+		FIsActionChecked::CreateSP(this, &SEditorViewport::IsCoordSystemActive, COORD_Explicit)
+	);
+	
 	CommandListRef.MapAction(
 		Commands.CycleTransformGizmos,
 		FExecuteAction::CreateSP( this, &SEditorViewport::OnCycleWidgetMode ),
@@ -921,7 +928,12 @@ void SEditorViewport::OnCycleCoordinateSystem()
 	++CoordSystemAsInt;
 
 	// parent mode is only supported with new trs gizmos for now
-	const int CoordMax = UEditorInteractiveGizmoManager::UsesNewTRSGizmos() ? COORD_Max : COORD_Parent;
+	int CoordMax = COORD_Parent;
+	if (UEditorInteractiveGizmoManager::UsesNewTRSGizmos())
+	{
+		CoordMax = UEditorInteractiveGizmoManager::IsExplicitModeEnabled() ? COORD_Max : COORD_Explicit;
+	}
+	
 	if( CoordSystemAsInt >= CoordMax )
 	{
 		CoordSystemAsInt = COORD_World;
diff --git a/Engine/Source/Editor/UnrealEd/Private/STransformViewportToolbar.cpp b/Engine/Source/Editor/UnrealEd/Private/STransformViewportToolbar.cpp
--- a/Engine/Source/Editor/UnrealEd/Private/STransformViewportToolbar.cpp
+++ b/Engine/Source/Editor/UnrealEd/Private/STransformViewportToolbar.cpp
@@ -201,7 +201,7 @@ TSharedRef< SWidget > STransformViewportToolBar::MakeTransformToolBar( const TSh
 			if (UEditorInteractiveGizmoManager::UsesNewTRSGizmos())
 			{
 				return LOCTEXT(	"CycleTransformGizmoCoordSystemWithParent_ToolTip",
-								"Cycles the transform gizmo coordinate systems between world, local and parent space");
+								"Cycles the transform gizmo coordinate systems between world, local, parent and explicit space");
 			}
 			return FEditorViewportCommands::Get().CycleTransformGizmoCoordSystem->GetDescription();
 		});
diff --git a/Engine/Source/Editor/UnrealEd/Private/ViewportToolbar/UnrealEdViewportToolbar.cpp b/Engine/Source/Editor/UnrealEd/Private/ViewportToolbar/UnrealEdViewportToolbar.cpp
--- a/Engine/Source/Editor/UnrealEd/Private/ViewportToolbar/UnrealEdViewportToolbar.cpp
+++ b/Engine/Source/Editor/UnrealEd/Private/ViewportToolbar/UnrealEdViewportToolbar.cpp
@@ -673,6 +673,12 @@ FSlateIcon GetIconFromCoordSystem(ECoordSystem InCoordSystem)
 		static const FName ParentIcon("Icons.ConstraintManager.ParentHierarchy");
 		return FSlateIcon(FAppStyle::GetAppStyleSetName(), ParentIcon);
 	}
+	else if (InCoordSystem == COORD_Explicit)
+	{
+		// TODO: find a better icon for this
+		static const FName GimbalIcon("EditorViewport.RotateMode");
+		return FSlateIcon(FAppStyle::GetAppStyleSetName(), GimbalIcon);
+	}
 	else
 	{
 		static FName LocalIcon("Icons.Transform");
@@ -753,18 +759,26 @@ FToolMenuEntry CreateTransformsSubmenu()
 													);
 													UnnamedSection.AddEntry(WorldCoords);
 
+													FToolMenuEntry LocalCoords = FToolMenuEntry::InitMenuEntry(
+														FEditorViewportCommands::Get().RelativeCoordinateSystem_Local
+													);
+													UnnamedSection.AddEntry(LocalCoords);
+
 													if (UEditorInteractiveGizmoManager::UsesNewTRSGizmos())
 													{
 														FToolMenuEntry ParentCoords = FToolMenuEntry::InitMenuEntry(
 															FEditorViewportCommands::Get().RelativeCoordinateSystem_Parent
 														);
 														UnnamedSection.AddEntry(ParentCoords);
-													}
 
-													FToolMenuEntry LocalCoords = FToolMenuEntry::InitMenuEntry(
-														FEditorViewportCommands::Get().RelativeCoordinateSystem_Local
-													);
-													UnnamedSection.AddEntry(LocalCoords);
+														if (UEditorInteractiveGizmoManager::IsExplicitModeEnabled())
+														{
+															FToolMenuEntry ExplicitCoords = FToolMenuEntry::InitMenuEntry(
+															   FEditorViewportCommands::Get().RelativeCoordinateSystem_Explicit
+														   );
+														   UnnamedSection.AddEntry(ExplicitCoords);
+														}
+													}
 												}
 											)
 										);
diff --git a/Engine/Source/Editor/UnrealEd/Public/EditorModeManager.h b/Engine/Source/Editor/UnrealEd/Public/EditorModeManager.h
--- a/Engine/Source/Editor/UnrealEd/Public/EditorModeManager.h
+++ b/Engine/Source/Editor/UnrealEd/Public/EditorModeManager.h
@@ -13,6 +13,8 @@
 
 #include "Elements/Framework/TypedElementSelectionSet.h"
 
+struct FRotationContext;
+class UEditorTransformGizmoContextObject;
 class ITypedElementWorldInterface;
 class FCanvas;
 class FEditorViewportClient;
@@ -139,6 +141,7 @@ public:
 	UNREALED_API FMatrix GetCustomInputCoordinateSystem() const;
 	UNREALED_API FMatrix GetLocalCoordinateSystem() const;
 	UNREALED_API FMatrix GetParentSpaceCoordinateSystem() const;
+	UNREALED_API FMatrix GetExplicitCoordinateSystem() const;
 	
 	/** 
 	 * Returns true if the passed in editor mode is active 
@@ -545,11 +548,14 @@ public:
 	/** @return ToolsContext for this Mode Manager */
 	UNREALED_API UModeManagerInteractiveToolsContext* GetInteractiveToolsContext() const;
 
-	/** New TRS Gizmo interface */
+	/** Begin New TRS Gizmo interface */
 	UNREALED_API IGizmoStateTarget* GetGizmoStateTarget();
 	UNREALED_API bool BeginTransform(const FGizmoState& InState);
 	UNREALED_API bool EndTransform(const FGizmoState& InState) const;
 	UNREALED_API bool HasOngoingTransform() const;
+	UNREALED_API UEditorTransformGizmoContextObject* GetGizmoContext() const;
+	UNREALED_API const FRotationContext& GetRotationContext() const;
+	/** End New TRS Gizmo interface */
 	
 protected:
 	/** 
@@ -645,7 +651,10 @@ private:
 	bool bIsExitingModesDuringTick = false;
 
 	/** GizmoStateTarget used to handle a new TRS gizmo transform begin/end sequence. */
-	TWeakObjectPtr<UEditorGizmoStateTarget> GizmoStateTarget;
+	TWeakObjectPtr<UEditorGizmoStateTarget> WeakGizmoStateTarget;
+
+	/** GizmoContext used to get context specific gizmo properties. */
+	mutable TWeakObjectPtr<UEditorTransformGizmoContextObject> WeakGizmoContext;
 
 	/** Flag to track if we started a new TRS gizmo transform.
 	 * NOTE: we have to use it for now as StartTracking / EndTracking iterates thru the modes even if the ITF context captured something.
diff --git a/Engine/Source/Editor/UnrealEd/Public/EditorViewportCommands.h b/Engine/Source/Editor/UnrealEd/Public/EditorViewportCommands.h
--- a/Engine/Source/Editor/UnrealEd/Public/EditorViewportCommands.h
+++ b/Engine/Source/Editor/UnrealEd/Public/EditorViewportCommands.h
@@ -215,6 +215,9 @@ public:
 	/** Parent-relative coordinate system */
 	TSharedPtr<FUICommandInfo> RelativeCoordinateSystem_Parent;
 
+	/** Explicit coordinate system (Rotation Order/Gimbal or Parent-relative coordinate system */
+	TSharedPtr<FUICommandInfo> RelativeCoordinateSystem_Explicit;
+
 	TSharedPtr< FUICommandInfo > CycleTransformGizmos;
 	TSharedPtr< FUICommandInfo > CycleTransformGizmoCoordSystem;
 
