fixed missing copyright
39183463 follow-up

#rnx

diff --git a/Engine/Source/Editor/Experimental/EditorInteractiveToolsFramework/Private/EditorGizmos/GizmoRotationUtil.cpp b/Engine/Source/Editor/Experimental/EditorInteractiveToolsFramework/Private/EditorGizmos/GizmoRotationUtil.cpp
--- a/Engine/Source/Editor/Experimental/EditorInteractiveToolsFramework/Private/EditorGizmos/GizmoRotationUtil.cpp
+++ b/Engine/Source/Editor/Experimental/EditorInteractiveToolsFramework/Private/EditorGizmos/GizmoRotationUtil.cpp
@@ -1,5 +1,4 @@
-﻿// Fill out your copyright notice in the Description page of Project Settings.
-
+﻿// Copyright Epic Games, Inc. All Rights Reserved.
 
 #include "EditorGizmos/GizmoRotationUtil.h"
 #include "AnimationCoreLibrary.h"
