[Truncation Warnings] Fix warnings in EditorInteractiveToolsFramework module

#rnx

diff --git a/Engine/Source/Editor/Experimental/EditorInteractiveToolsFramework/Private/EdModeInteractiveToolsContext.cpp b/Engine/Source/Editor/Experimental/EditorInteractiveToolsFramework/Private/EdModeInteractiveToolsContext.cpp
--- a/Engine/Source/Editor/Experimental/EditorInteractiveToolsFramework/Private/EdModeInteractiveToolsContext.cpp
+++ b/Engine/Source/Editor/Experimental/EditorInteractiveToolsFramework/Private/EdModeInteractiveToolsContext.cpp
@@ -1016,9 +1016,9 @@ bool UModeManagerInteractiveToolsContext::InputKey(FEditorViewportClient* Viewpo
 					ModifierKeys.IsShiftDown(), ModifierKeys.IsAltDown(),
 					ModifierKeys.IsControlDown(), ModifierKeys.IsCommandDown());
 
-				InputState.Mouse.WheelDelta = (Event != IE_Pressed) ? 0
-					: (Key == EKeys::MouseScrollUp) ? 1 
-					: -1;
+				InputState.Mouse.WheelDelta = (Event != IE_Pressed) ? 0.f
+					: (Key == EKeys::MouseScrollUp) ? 1.f 
+					: -1.f;
 
 				return InputRouter->PostInputEvent(InputState);
 			}
diff --git a/Engine/Source/Editor/Experimental/EditorInteractiveToolsFramework/Private/EditorGizmos/TransformGizmo.cpp b/Engine/Source/Editor/Experimental/EditorInteractiveToolsFramework/Private/EditorGizmos/TransformGizmo.cpp
--- a/Engine/Source/Editor/Experimental/EditorInteractiveToolsFramework/Private/EditorGizmos/TransformGizmo.cpp
+++ b/Engine/Source/Editor/Experimental/EditorInteractiveToolsFramework/Private/EditorGizmos/TransformGizmo.cpp
@@ -38,7 +38,7 @@ static FAutoConsoleVariableRef CVarDotThreshold(
 	TEXT("Dot threshold for determining whether the rotation plane is perpendicular to the camera view [0.2, 1.0]"),
 	FConsoleVariableDelegate::CreateLambda([](IConsoleVariable*)
 	{
-		DotThreshold = FMath::Clamp(DotThreshold, 0.2, 1.0);
+		DotThreshold = FMath::Clamp(DotThreshold, 0.2f, 1.0f);
 	})
 );
 
@@ -1324,7 +1324,7 @@ UGizmoElementCircle* UTransformGizmo::MakeRotateCircleHandle(ETransformGizmoPart
 
 	if (bFill)
 	{
-		const FLinearColor LightColor(InColor.R, InColor.G, InColor.B, InColor.A * 0.5);
+		const FLinearColor LightColor(InColor.R, InColor.G, InColor.B, InColor.A * 0.5f);
 		CircleElement->SetVertexColor(LightColor);
 		CircleElement->SetMaterial(TransparentVertexColorMaterial);
 
@@ -1533,7 +1533,7 @@ float UTransformGizmo::GetNearestRayParamToInteractionAxis(const FInputDeviceRay
 	return AxisNearestParam;
 }
 
-bool UTransformGizmo::GetRayParamIntersectionWithInteractionPlane(const FInputDeviceRay& InRay, float& OutHitParam)
+bool UTransformGizmo::GetRayParamIntersectionWithInteractionPlane(const FInputDeviceRay& InRay, FVector::FReal& OutHitParam)
 {
 	// if ray is parallel to plane, nothing has been hit
 	if (FMath::IsNearlyZero(FVector::DotProduct(InteractionPlanarNormal, InRay.WorldRay.Direction)))
@@ -1865,9 +1865,9 @@ void UTransformGizmo::OnClickPressAxis(const FInputDeviceRay& PressPos)
 	const FVector ZAxis = GetWorldAxis(FVector::ZAxisVector);
 
 	const FVector ViewDirection = GizmoViewContext->GetViewDirection();
-	const double XDot = FMath::Abs(FVector::DotProduct(ViewDirection, XAxis));
-	const double YDot = FMath::Abs(FVector::DotProduct(ViewDirection, YAxis));
-	const double ZDot = FMath::Abs(FVector::DotProduct(ViewDirection, ZAxis));
+	const FVector::FReal XDot = FMath::Abs(FVector::DotProduct(ViewDirection, XAxis));
+	const FVector::FReal YDot = FMath::Abs(FVector::DotProduct(ViewDirection, YAxis));
+	const FVector::FReal ZDot = FMath::Abs(FVector::DotProduct(ViewDirection, ZAxis));
 
 	if (FVector::DotProduct(InteractionAxisDirection, XAxis) > 0.1)
 	{
@@ -1885,7 +1885,7 @@ void UTransformGizmo::OnClickPressAxis(const FInputDeviceRay& PressPos)
 		NormalToRemove = (XDot > YDot) ? YAxis : XAxis;
 	}
 	
-	float HitDepth;
+	FVector::FReal HitDepth;
 	if (GetRayParamIntersectionWithInteractionPlane(PressPos, HitDepth))
 	{
 		InteractionPlanarStartPoint = PressPos.WorldRay.Origin + PressPos.WorldRay.Direction * HitDepth;
@@ -1920,7 +1920,7 @@ void UTransformGizmo::OnClickDragTranslateAxis(const FInputDeviceRay& DragPos)
 		return;
 	}
 
-	float HitDepth;
+	FVector::FReal HitDepth;
 	if (GetRayParamIntersectionWithInteractionPlane(DragPos, HitDepth))
 	{
 		FVector HitPoint = DragPos.WorldRay.Origin + DragPos.WorldRay.Direction * HitDepth;
@@ -1974,7 +1974,7 @@ void UTransformGizmo::OnClickPressTranslateXZPlanar(const FInputDeviceRay& Press
 
 void UTransformGizmo::OnClickPressPlanar(const FInputDeviceRay& PressPos)
 {
-	float HitDepth;
+	FVector::FReal HitDepth;
 	if (GetRayParamIntersectionWithInteractionPlane(PressPos, HitDepth))
 	{
 		InteractionPlanarStartPoint = PressPos.WorldRay.Origin + PressPos.WorldRay.Direction * HitDepth;
@@ -1987,7 +1987,7 @@ void UTransformGizmo::OnClickPressPlanar(const FInputDeviceRay& PressPos)
 
 void UTransformGizmo::OnClickDragTranslatePlanar(const FInputDeviceRay& DragPos)
 {
-	float HitDepth;
+	FVector::FReal HitDepth;
 	if (GetRayParamIntersectionWithInteractionPlane(DragPos, HitDepth))
 	{
 		const FVector HitPoint = DragPos.WorldRay.Origin + DragPos.WorldRay.Direction * HitDepth;
@@ -2171,9 +2171,9 @@ FVector UTransformGizmo::ComputeScaleDelta(const FVector2D& InStartPos, const FV
 		(InteractionAxisList & EAxisList::Y) ? ScaleDelta : 0.0f,
 		(InteractionAxisList & EAxisList::Z) ? ScaleDelta : 0.0f);
 
-	const float ScaleMax = Scale.GetMax();
-	const float ScaleMin = Scale.GetMin();
-	const float ScaleApplied = (ScaleMax > -ScaleMin) ? ScaleMax : ScaleMin;
+	const FVector::FReal ScaleMax = Scale.GetMax();
+	const FVector::FReal ScaleMin = Scale.GetMin();
+	const FVector::FReal ScaleApplied = (ScaleMax > -ScaleMin) ? ScaleMax : ScaleMin;
 
 	OutScreenDelta = InteractionScreenAxisDirection * ScaleApplied;
 
@@ -2203,7 +2203,7 @@ bool UTransformGizmo::OnClickPressRotateArc(
 		InteractionPlanarAxisX = InPlaneAxis1;
 		InteractionPlanarAxisY = InPlaneAxis2;
 	
-		float HitDepth;
+		FVector::FReal HitDepth;
 		if (GetRayParamIntersectionWithInteractionPlane(InPressPos, HitDepth))
 		{
 			const FVector HitPoint = InPressPos.WorldRay.Origin + InPressPos.WorldRay.Direction * HitDepth;
@@ -2406,7 +2406,7 @@ void UTransformGizmo::OnClickDragRotateAxis(const FInputDeviceRay& DragPos)
 	}
 	case EAxisRotateMode::Arc:
 	{
-		float HitDepth;
+		FVector::FReal HitDepth;
 		if (GetRayParamIntersectionWithInteractionPlane(DragPos, HitDepth))
 		{
 			const FVector HitPoint = DragPos.WorldRay.Origin + DragPos.WorldRay.Direction * HitDepth;
@@ -2578,7 +2578,7 @@ void UTransformGizmo::OnClickDragGimbalRotateAxis(const FInputDeviceRay& InDragP
 		}
 	case EAxisRotateMode::Arc:
 		{
-			float HitDepth;
+		FVector::FReal HitDepth;
 			if (GetRayParamIntersectionWithInteractionPlane(InDragPos, HitDepth))
 			{
 				const FVector HitPoint = InDragPos.WorldRay.Origin + InDragPos.WorldRay.Direction * HitDepth;
@@ -2628,7 +2628,7 @@ void UTransformGizmo::OnClickPressScreenSpaceRotate(const FInputDeviceRay& InPre
 	}
 	else
 	{
-		float HitDepth;
+		FVector::FReal HitDepth;
 		if (!GetRayParamIntersectionWithInteractionPlane(InPressPos, HitDepth))
 		{
 			return;
@@ -2675,7 +2675,7 @@ void UTransformGizmo::OnClickDragScreenSpaceRotate(const FInputDeviceRay& InDrag
 		return;
 	}
 	
-	float HitDepth;
+	FVector::FReal HitDepth;
 	if (GetRayParamIntersectionWithInteractionPlane(InDragPos, HitDepth))
 	{
 		const FVector HitPoint = InDragPos.WorldRay.Origin + InDragPos.WorldRay.Direction * HitDepth;
@@ -2689,9 +2689,9 @@ void UTransformGizmo::OnClickDragScreenSpaceRotate(const FInputDeviceRay& InDrag
 	}
 }
 
-FQuat UTransformGizmo::ComputeAngularRotateDelta(double InStartAngle, double InEndAngle)
+FQuat UTransformGizmo::ComputeAngularRotateDelta(FQuat::FReal InStartAngle, FQuat::FReal InEndAngle)
 {
-	const float DeltaAngle = InEndAngle - InStartAngle;
+	const FQuat::FReal DeltaAngle = InEndAngle - InStartAngle;
 	return FQuat(InteractionPlanarNormal, DeltaAngle);
 }
 
@@ -2787,7 +2787,7 @@ void UTransformGizmo::OnClickPressArcBallRotate(const FInputDeviceRay& PressPos)
 		if (bIndirectManipulation)
 		{
 			// change the arc ball center in indirect manipulation if we didn't hit the sphere
-			float HitDepth;
+			FVector::FReal HitDepth;
 			if (GetRayParamIntersectionWithInteractionPlane(PressPos, HitDepth))
 			{
 				InteractionPlanarOrigin = PressPos.WorldRay.Origin + PressPos.WorldRay.Direction * HitDepth;
diff --git a/Engine/Source/Editor/Experimental/EditorInteractiveToolsFramework/Public/EditorGizmos/TransformGizmo.h b/Engine/Source/Editor/Experimental/EditorInteractiveToolsFramework/Public/EditorGizmos/TransformGizmo.h
--- a/Engine/Source/Editor/Experimental/EditorInteractiveToolsFramework/Public/EditorGizmos/TransformGizmo.h
+++ b/Engine/Source/Editor/Experimental/EditorInteractiveToolsFramework/Public/EditorGizmos/TransformGizmo.h
@@ -511,7 +511,7 @@ protected:
 	virtual float GetNearestRayParamToInteractionAxis(const FInputDeviceRay& InRay);
 
 	/** Return true if input ray intersects current interaction plane and return param along ray in OutHitParam */
-	virtual bool GetRayParamIntersectionWithInteractionPlane(const FInputDeviceRay& InRay, float& OutHitParam);
+	virtual bool GetRayParamIntersectionWithInteractionPlane(const FInputDeviceRay& InRay, FVector::FReal& OutHitParam);
 
 	/** Update hover state for given part id */
 	virtual void UpdateHoverState(const bool bInHover, const ETransformGizmoPartIdentifier InPartId);
@@ -660,7 +660,7 @@ protected:
 	virtual void OnClickReleaseScreenSpaceRotate(const FInputDeviceRay& ReleasePos);
 
 	/** Compute rotate delta based on start/end angles */
-	virtual FQuat ComputeAngularRotateDelta(double InStartAngle, double InEndAngle);
+	virtual FQuat ComputeAngularRotateDelta(FQuat::FReal InStartAngle, FQuat::FReal InEndAngle);
 	
 	/**
 	 * Arc ball rotate interaction methods
