[Truncation Warnings] Enable truncation warnings in EditorInteractiveToolsFramework module

#rnx

diff --git a/Engine/Source/Editor/Experimental/EditorInteractiveToolsFramework/EditorInteractiveToolsFramework.Build.cs b/Engine/Source/Editor/Experimental/EditorInteractiveToolsFramework/EditorInteractiveToolsFramework.Build.cs
--- a/Engine/Source/Editor/Experimental/EditorInteractiveToolsFramework/EditorInteractiveToolsFramework.Build.cs
+++ b/Engine/Source/Editor/Experimental/EditorInteractiveToolsFramework/EditorInteractiveToolsFramework.Build.cs
@@ -6,6 +6,9 @@ public class EditorInteractiveToolsFramework : ModuleRules
 {
 	public EditorInteractiveToolsFramework(ReadOnlyTargetRules Target) : base(Target)
 	{
+		// Enable truncation warnings in this module
+		CppCompileWarningSettings.UnsafeTypeCastWarningLevel = WarningLevel.Warning;
+
 		PCHUsage = ModuleRules.PCHUsageMode.UseExplicitOrSharedPCHs;			
 		
 		PublicDependencyModuleNames.AddRange(
diff --git a/Engine/Source/Editor/Experimental/EditorInteractiveToolsFramework/Private/EdModeInteractiveToolsContext.cpp b/Engine/Source/Editor/Experimental/EditorInteractiveToolsFramework/Private/EdModeInteractiveToolsContext.cpp
--- a/Engine/Source/Editor/Experimental/EditorInteractiveToolsFramework/Private/EdModeInteractiveToolsContext.cpp
+++ b/Engine/Source/Editor/Experimental/EditorInteractiveToolsFramework/Private/EdModeInteractiveToolsContext.cpp
@@ -103,7 +103,7 @@ public:
 				CachedViewState.Orientation = FQuat::Identity;
 			}
 
-			CachedViewState.OrthoWorldCoordinateWidth = ViewportClient->GetOrthoUnitsPerPixel(ViewportClient->Viewport) * ViewportClient->Viewport->GetSizeXY().X;
+			CachedViewState.OrthoWorldCoordinateWidth = ViewportClient->GetOrthoUnitsPerPixel(ViewportClient->Viewport) * static_cast<float>(ViewportClient->Viewport->GetSizeXY().X);
 		}
 
 		CachedViewState.bIsVR = false;
@@ -569,7 +569,7 @@ public:
 				ViewCameraState.Orientation = FQuat::Identity;
 			}
 
-			ViewCameraState.OrthoWorldCoordinateWidth = ViewportClient->GetOrthoUnitsPerPixel(ViewportClient->Viewport) * ViewportClient->Viewport->GetSizeXY().X;
+			ViewCameraState.OrthoWorldCoordinateWidth = ViewportClient->GetOrthoUnitsPerPixel(ViewportClient->Viewport) * static_cast<float>(ViewportClient->Viewport->GetSizeXY().X);
 		}
 
 		ViewCameraState.bIsVR = false;
diff --git a/Engine/Source/Editor/Experimental/EditorInteractiveToolsFramework/Private/EditorGizmos/TransformGizmo.cpp b/Engine/Source/Editor/Experimental/EditorInteractiveToolsFramework/Private/EditorGizmos/TransformGizmo.cpp
--- a/Engine/Source/Editor/Experimental/EditorInteractiveToolsFramework/Private/EditorGizmos/TransformGizmo.cpp
+++ b/Engine/Source/Editor/Experimental/EditorInteractiveToolsFramework/Private/EditorGizmos/TransformGizmo.cpp
@@ -1300,7 +1300,7 @@ UGizmoElementTorus* UTransformGizmo::MakeRotateAxis(ETransformGizmoPartIdentifie
 	RotateAxisElement->SetPartialEndAngle(UE_PI);
 	RotateAxisElement->SetViewDependentAxis(TorusNormal);
 	RotateAxisElement->SetViewAlignType(EGizmoElementViewAlignType::Axial);
-	RotateAxisElement->SetViewAlignAxialAngleTol(UE_DOUBLE_SMALL_NUMBER);
+	RotateAxisElement->SetViewAlignAxialAngleTol(static_cast<float>(UE_DOUBLE_SMALL_NUMBER));
 	RotateAxisElement->SetPartialViewDependentMaxCosTol(1.0-UE_DOUBLE_SMALL_NUMBER);	
 	RotateAxisElement->SetViewAlignAxis(TorusNormal);
 	RotateAxisElement->SetViewAlignNormal(TorusAxis1);
@@ -1322,7 +1322,7 @@ UGizmoElementCircle* UTransformGizmo::MakeRotateCircleHandle(ETransformGizmoPart
 	CircleElement->SetViewAlignType(EGizmoElementViewAlignType::PointOnly);
 	CircleElement->SetViewAlignNormal(-FVector::ForwardVector);
 
-	if (bFill)
+	if (FMath::IsNearlyEqual(bFill, 0.0f))
 	{
 		const FLinearColor LightColor(InColor.R, InColor.G, InColor.B, InColor.A * 0.5f);
 		CircleElement->SetVertexColor(LightColor);
