follow up to 39669428

#rnx

diff --git a/Engine/Source/Editor/Experimental/EditorInteractiveToolsFramework/Private/EditorGizmos/TransformGizmo.cpp b/Engine/Source/Editor/Experimental/EditorInteractiveToolsFramework/Private/EditorGizmos/TransformGizmo.cpp
--- a/Engine/Source/Editor/Experimental/EditorInteractiveToolsFramework/Private/EditorGizmos/TransformGizmo.cpp
+++ b/Engine/Source/Editor/Experimental/EditorInteractiveToolsFramework/Private/EditorGizmos/TransformGizmo.cpp
@@ -1308,7 +1308,7 @@ UGizmoElementTorus* UTransformGizmo::MakeRotateAxis(ETransformGizmoPartIdentifie
 	return RotateAxisElement;
 }
 
-UGizmoElementCircle* UTransformGizmo::MakeRotateCircleHandle(ETransformGizmoPartIdentifier InPartId, float InRadius, const FLinearColor& InColor, float bFill)
+UGizmoElementCircle* UTransformGizmo::MakeRotateCircleHandle(ETransformGizmoPartIdentifier InPartId, float InRadius, const FLinearColor& InColor, const bool bFill)
 {
 	const float SizeCoeff = GetSizeCoefficient();
 	
@@ -1322,7 +1322,7 @@ UGizmoElementCircle* UTransformGizmo::MakeRotateCircleHandle(ETransformGizmoPart
 	CircleElement->SetViewAlignType(EGizmoElementViewAlignType::PointOnly);
 	CircleElement->SetViewAlignNormal(-FVector::ForwardVector);
 
-	if (FMath::IsNearlyEqual(bFill, 0.0f))
+	if (bFill)
 	{
 		const FLinearColor LightColor(InColor.R, InColor.G, InColor.B, InColor.A * 0.5f);
 		CircleElement->SetVertexColor(LightColor);
diff --git a/Engine/Source/Editor/Experimental/EditorInteractiveToolsFramework/Public/EditorGizmos/TransformGizmo.h b/Engine/Source/Editor/Experimental/EditorInteractiveToolsFramework/Public/EditorGizmos/TransformGizmo.h
--- a/Engine/Source/Editor/Experimental/EditorInteractiveToolsFramework/Public/EditorGizmos/TransformGizmo.h
+++ b/Engine/Source/Editor/Experimental/EditorInteractiveToolsFramework/Public/EditorGizmos/TransformGizmo.h
@@ -493,7 +493,7 @@ protected:
 	virtual UGizmoElementRectangle* MakeTranslateScreenSpaceHandle();
 
 	/** Construct rotate screen space handle */
-	virtual UGizmoElementCircle* MakeRotateCircleHandle(ETransformGizmoPartIdentifier InPartId, float InRadius, const FLinearColor& InColor, float bFill);
+	virtual UGizmoElementCircle* MakeRotateCircleHandle(ETransformGizmoPartIdentifier InPartId, float InRadius, const FLinearColor& InColor, const bool bFill);
 
 	/** Get gizmo transform based on cached current transform. */
 	virtual FTransform GetGizmoTransform() const;
