[EditorGizmos] New Transform gizmo grid snapping integration (to match legacy)

diff --git a/Engine/Source/Editor/Experimental/EditorInteractiveToolsFramework/EditorInteractiveToolsFramework.Build.cs b/Engine/Source/Editor/Experimental/EditorInteractiveToolsFramework/EditorInteractiveToolsFramework.Build.cs
--- a/Engine/Source/Editor/Experimental/EditorInteractiveToolsFramework/EditorInteractiveToolsFramework.Build.cs
+++ b/Engine/Source/Editor/Experimental/EditorInteractiveToolsFramework/EditorInteractiveToolsFramework.Build.cs
@@ -39,8 +39,9 @@ public class EditorInteractiveToolsFramework : ModuleRules
 				"StaticMeshDescription",
                 "EditorSubsystem",
                 "TypedElementRuntime",
-                "AnimationCore"
-				// ... add private dependencies that you statically link with here ...	
+                "AnimationCore",
+                "ViewportSnapping", // For ISnappingPolicy support
+                // ... add private dependencies that you statically link with here ...
 			}
             );
 	}
diff --git a/Engine/Source/Editor/Experimental/EditorInteractiveToolsFramework/Private/EditorGizmos/TransformGizmo.cpp b/Engine/Source/Editor/Experimental/EditorInteractiveToolsFramework/Private/EditorGizmos/TransformGizmo.cpp
--- a/Engine/Source/Editor/Experimental/EditorInteractiveToolsFramework/Private/EditorGizmos/TransformGizmo.cpp
+++ b/Engine/Source/Editor/Experimental/EditorInteractiveToolsFramework/Private/EditorGizmos/TransformGizmo.cpp
@@ -21,6 +21,7 @@
 #include "Behaviors/MultiButtonClickDragBehavior.h"
 #include "Intersection/IntersectionUtil.h"
 #include "PrimitiveDrawingUtils.h"
+#include "SceneQueries/SceneSnappingManager.h"
 
 #define LOCTEXT_NAMESPACE "UTransformGizmo"
 
@@ -1897,6 +1898,155 @@ void UTransformGizmo::OnClickPressAxis(const FInputDeviceRay& PressPos)
 	SetModeLastHitPart(EGizmoTransformMode::Translate, LastHitPart);
 }
 
+void UTransformGizmo::SnapTranslateDelta(FVector& InOutWorldDelta) const
+{
+	UInteractiveGizmoManager* GizmoManager = GetGizmoManager();
+	if (!GizmoManager)
+	{
+		return;
+	}
+	
+	IToolsContextQueriesAPI* QueriesAPI = GizmoManager->GetContextQueriesAPI();
+	if (!QueriesAPI)
+	{
+		return;
+	}
+
+	if (const USceneSnappingManager* SnapManager = USceneSnappingManager::Find(GizmoManager))
+	{
+		FSceneSnapQueryRequest SnapRequest;
+		SnapRequest.RequestCoordinateSpace = QueriesAPI->GetCurrentCoordinateSystem();
+		SnapRequest.RequestType = ESceneSnapQueryType::Position;
+		SnapRequest.AxisList = InteractionAxisList;
+
+		FVector PositionToSnap = InOutWorldDelta;
+		if (SnapRequest.RequestCoordinateSpace == EToolContextCoordinateSystem::Local)
+		{
+			// Convert the delta to local space
+			const FQuat Rotation = ActiveTarget->GetTransform().GetRotation();
+			PositionToSnap = Rotation.Inverse() * InOutWorldDelta;
+		}
+
+		SnapRequest.Position = PositionToSnap;
+
+		TArray<FSceneSnapQueryResult> SnapResults;
+		if (SnapManager->ExecuteSceneSnapQuery(SnapRequest, SnapResults))
+		{
+			check(!SnapResults.IsEmpty());
+			const FSceneSnapQueryResult& SnapResult = SnapResults[0];
+
+			// Convert the delta back to the desired space
+			if (SnapResult.GetCoordinateSpace(SnapRequest) == EToolContextCoordinateSystem::Local)
+			{
+				const FQuat Rotation = ActiveTarget->GetTransform().GetRotation();
+				InOutWorldDelta = Rotation * SnapResult.Position;
+			}
+			else
+			{
+				InOutWorldDelta = SnapResult.Position;
+			}
+		}
+	}
+}
+
+void UTransformGizmo::SnapRotateDelta(FQuat& InOutWorldDelta) const
+{
+	UInteractiveGizmoManager* GizmoManager = GetGizmoManager();
+	if (!GizmoManager)
+	{
+		return;
+	}
+
+	IToolsContextQueriesAPI* QueriesAPI = GizmoManager->GetContextQueriesAPI();
+	if (!QueriesAPI)
+	{
+		return;
+	}
+
+	if (const USceneSnappingManager* SnapManager = USceneSnappingManager::Find(GizmoManager))
+	{
+		FSceneSnapQueryRequest SnapRequest;
+		SnapRequest.RequestCoordinateSpace = QueriesAPI->GetCurrentCoordinateSystem();
+		SnapRequest.RequestType = ESceneSnapQueryType::Rotation;
+		SnapRequest.AxisList = InteractionAxisList;
+		SnapRequest.DeltaRotation = InOutWorldDelta;
+
+		TArray<FSceneSnapQueryResult> SnapResults;
+		if (SnapManager->ExecuteSceneSnapQuery(SnapRequest, SnapResults))
+		{
+			check(!SnapResults.IsEmpty());
+
+			InOutWorldDelta = SnapResults[0].DeltaRotation;
+		}
+	}
+}
+
+void UTransformGizmo::SnapRotateAngleDelta(double& InOutAngleDelta) const
+{
+	UInteractiveGizmoManager* GizmoManager = GetGizmoManager();
+	if (!GizmoManager)
+	{
+		return;
+	}
+
+	IToolsContextQueriesAPI* QueriesAPI = GizmoManager->GetContextQueriesAPI();
+	if (!QueriesAPI)
+	{
+		return;
+	}
+
+	if (const USceneSnappingManager* SnapManager = USceneSnappingManager::Find(GizmoManager))
+	{
+		FSceneSnapQueryRequest SnapRequest;
+		SnapRequest.RequestCoordinateSpace = QueriesAPI->GetCurrentCoordinateSystem();
+		SnapRequest.RequestType = ESceneSnapQueryType::RotationAngle;
+		SnapRequest.AxisList = InteractionAxisList;
+		SnapRequest.RotationAngle = InOutAngleDelta;
+
+		TArray<FSceneSnapQueryResult> SnapResults;
+		if (SnapManager->ExecuteSceneSnapQuery(SnapRequest, SnapResults))
+		{
+			check(!SnapResults.IsEmpty());
+
+			InOutAngleDelta = SnapResults[0].RotationAngle;
+		}
+	}
+}
+
+void UTransformGizmo::SnapScaleDelta(FVector& InOutLocalScaleDelta) const
+{
+	UInteractiveGizmoManager* GizmoManager = GetGizmoManager();
+	if (!GizmoManager)
+	{
+		return;
+	}
+
+	IToolsContextQueriesAPI* QueriesAPI = GizmoManager->GetContextQueriesAPI();
+	if (!QueriesAPI)
+	{
+		return;
+	}
+
+	if (const USceneSnappingManager* SnapManager = USceneSnappingManager::Find(GizmoManager))
+	{
+		FSceneSnapQueryRequest SnapRequest;
+
+		// Scale should always be in local space
+		SnapRequest.RequestCoordinateSpace = EToolContextCoordinateSystem::Local;
+		SnapRequest.RequestType = ESceneSnapQueryType::Scale;
+		SnapRequest.AxisList = InteractionAxisList;
+		SnapRequest.Scale = InOutLocalScaleDelta;
+
+		TArray<FSceneSnapQueryResult> SnapResults;
+		if (SnapManager->ExecuteSceneSnapQuery(SnapRequest, SnapResults))
+		{
+			check(!SnapResults.IsEmpty());
+
+			InOutLocalScaleDelta = SnapResults[0].Scale;
+		}
+	}
+}
+
 void UTransformGizmo::OnClickDragTranslateAxis(const FInputDeviceRay& DragPos)
 {
 	// indirect manipulation uses a 2D projection approach instead of plane intersection
@@ -1927,14 +2077,16 @@ void UTransformGizmo::OnClickDragTranslateAxis(const FInputDeviceRay& DragPos)
 		FVector HitPoint = DragPos.WorldRay.Origin + DragPos.WorldRay.Direction * HitDepth;
 
 		const FVector DeltaToStart = HitPoint-InteractionPlanarStartPoint;
-		const FVector AxisToRemove = NormalToRemove * FVector::DotProduct(DeltaToStart, NormalToRemove); 
+		const FVector AxisToRemove = NormalToRemove * FVector::DotProduct(DeltaToStart, NormalToRemove);
 		
 		HitPoint -= AxisToRemove;
 	
-		const FVector Delta = ComputePlanarTranslateDelta(InteractionPlanarCurrPoint, HitPoint);
+		FVector Delta = ComputePlanarTranslateDelta(InteractionPlanarCurrPoint, HitPoint);
 
+		SnapTranslateDelta(Delta);
 		ApplyTranslateDelta(Delta);
-		InteractionPlanarCurrPoint = HitPoint;
+
+		InteractionPlanarCurrPoint += Delta;
 	}
 }
 
@@ -1992,9 +2144,10 @@ void UTransformGizmo::OnClickDragTranslatePlanar(const FInputDeviceRay& DragPos)
 	if (GetRayParamIntersectionWithInteractionPlane(DragPos, HitDepth))
 	{
 		const FVector HitPoint = DragPos.WorldRay.Origin + DragPos.WorldRay.Direction * HitDepth;
-		const FVector Delta = ComputePlanarTranslateDelta(InteractionPlanarCurrPoint, HitPoint);
+		FVector Delta = ComputePlanarTranslateDelta(InteractionPlanarCurrPoint, HitPoint);
+		SnapTranslateDelta(Delta);
 		ApplyTranslateDelta(Delta);
-		InteractionPlanarCurrPoint = HitPoint;
+		InteractionPlanarCurrPoint += Delta;
 	}
 }
 
@@ -2138,6 +2291,7 @@ void UTransformGizmo::OnClickDragScale(const FInputDeviceRay& DragPos)
 	InteractionScreenEndPos += ScreenDelta;
 
 	FVector ScaleDelta = ComputeScaleDelta(InteractionScreenStartPos, InteractionScreenEndPos, ScreenDelta);
+	SnapScaleDelta(ScaleDelta);
 
 	if (ScaleDelta.X != 0.0 || ScaleDelta.Y != 0.0 || ScaleDelta.Z != 0.0)
 	{
@@ -2400,9 +2554,15 @@ void UTransformGizmo::OnClickDragRotateAxis(const FInputDeviceRay& DragPos)
 	{
 	case EAxisRotateMode::Pull:
 	{
-		const FQuat DeltaRot = ComputeAxisRotateDelta(InteractionScreenCurrPos, DragPos);
+		double DeltaAngle = FMath::DegreesToRadians(ComputeAxisRotateDeltaAngle(InteractionScreenCurrPos, DragPos));
+		SnapRotateAngleDelta(DeltaAngle);
+
+		const FQuat DeltaRot = ComputeAxisRotateDelta(FMath::RadiansToDegrees(DeltaAngle));
+
 		ApplyRotateDelta(DeltaRot);
+
 		InteractionScreenCurrPos = DragPos.ScreenPosition;
+
 		break;
 	}
 	case EAxisRotateMode::Arc:
@@ -2411,12 +2571,19 @@ void UTransformGizmo::OnClickDragRotateAxis(const FInputDeviceRay& DragPos)
 		if (GetRayParamIntersectionWithInteractionPlane(DragPos, HitDepth))
 		{
 			const FVector HitPoint = DragPos.WorldRay.Origin + DragPos.WorldRay.Direction * HitDepth;
-			const float HitAngle = GizmoMath::ComputeAngleInPlane(HitPoint,
+			double HitAngle = GizmoMath::ComputeAngleInPlane(HitPoint,
 				InteractionPlanarOrigin, InteractionPlanarNormal, InteractionPlanarAxisX, InteractionPlanarAxisY);
 
-			const FQuat Delta = ComputeAngularRotateDelta(InteractionCurrAngle, HitAngle);
-	
+			double HitAngleDelta = HitAngle - InteractionCurrAngle;
+			SnapRotateAngleDelta(HitAngleDelta);
+
+			// Reapply the (potentially) snapped delta angle
+			HitAngle = InteractionCurrAngle + HitAngleDelta;
+
+			FQuat Delta = ComputeAngularRotateDelta(InteractionCurrAngle, HitAngle);
+
 			ApplyRotateDelta(Delta);
+
 			InteractionCurrAngle = HitAngle;
 		}
 		break;
@@ -2427,7 +2594,7 @@ void UTransformGizmo::OnClickDragRotateAxis(const FInputDeviceRay& DragPos)
 	}
 }
 
-FQuat UTransformGizmo::ComputeAxisRotateDelta(const FVector2D& InStartPos, const FInputDeviceRay& InDragPos)
+double UTransformGizmo::ComputeAxisRotateDeltaAngle(const FVector2D& InStartPos, const FInputDeviceRay& InDragPos)
 {
 	FVector2D DragDir = InDragPos.ScreenPosition - InStartPos;
 
@@ -2443,15 +2610,19 @@ FQuat UTransformGizmo::ComputeAxisRotateDelta(const FVector2D& InStartPos, const
 		}
 		bTrySwitchingToNormalPull = false;
 	}
-	
+
 	const FVector2D DragDirToRemove = NormalProjectionToRemove * FVector2D::DotProduct(DragDir, NormalProjectionToRemove);
 	DragDir -= DragDirToRemove;
 
-	const double Delta = FVector2D::DotProduct(InteractionScreenAxisDirection, DragDir) * 0.25;
+	return FVector2D::DotProduct(InteractionScreenAxisDirection, DragDir) * 0.25;
+}
+
+FQuat UTransformGizmo::ComputeAxisRotateDelta(const double& InDeltaAngle)
+{
 	FRotator DeltaRot(
-		InteractionAxisList == EAxisList::Y ? Delta : 0.0,
-		InteractionAxisList == EAxisList::Z ? Delta : 0.0,
-		InteractionAxisList == EAxisList::X ? Delta : 0.0);
+		InteractionAxisList == EAxisList::Y ? InDeltaAngle : 0.0,
+		InteractionAxisList == EAxisList::Z ? InDeltaAngle : 0.0,
+		InteractionAxisList == EAxisList::X ? InDeltaAngle : 0.0);
 
 	auto GetCoordinateSystem = [&]()
 	{
@@ -2572,8 +2743,9 @@ void UTransformGizmo::OnClickDragGimbalRotateAxis(const FInputDeviceRay& InDragP
 	{
 	case EAxisRotateMode::Pull:
 		{
-			const FQuat DeltaRot = ComputeGimbalRotateDelta(InteractionScreenCurrPos, InDragPos.ScreenPosition);
+			FQuat DeltaRot = ComputeGimbalRotateDelta(InteractionScreenCurrPos, InDragPos.ScreenPosition);
 			ApplyRotateDelta(DeltaRot);
+			SnapRotateDelta(DeltaRot);
 			InteractionScreenCurrPos = InDragPos.ScreenPosition;
 			break;
 		}
@@ -2583,10 +2755,16 @@ void UTransformGizmo::OnClickDragGimbalRotateAxis(const FInputDeviceRay& InDragP
 			if (GetRayParamIntersectionWithInteractionPlane(InDragPos, HitDepth))
 			{
 				const FVector HitPoint = InDragPos.WorldRay.Origin + InDragPos.WorldRay.Direction * HitDepth;
-				const float HitAngle = GizmoMath::ComputeAngleInPlane(HitPoint,
+				double HitAngle = GizmoMath::ComputeAngleInPlane(HitPoint,
 					InteractionPlanarOrigin, InteractionPlanarNormal, InteractionPlanarAxisX, InteractionPlanarAxisY);
 
-				const float DeltaAngle = FMath::RadiansToDegrees(InteractionCurrAngle - HitAngle);
+				double HitAngleDelta = HitAngle - InteractionCurrAngle;
+				SnapRotateAngleDelta(HitAngleDelta);
+
+				// Reapply the (potentially) snapped delta angle
+				HitAngle = InteractionCurrAngle + HitAngleDelta;
+
+				const double DeltaAngle = FMath::RadiansToDegrees(InteractionCurrAngle - HitAngle);
 				const FVector DeltaRot(
 					InteractionAxisList == EAxisList::X ? DeltaAngle : 0.0,
 					InteractionAxisList == EAxisList::Y ? DeltaAngle : 0.0,
@@ -2594,7 +2772,7 @@ void UTransformGizmo::OnClickDragGimbalRotateAxis(const FInputDeviceRay& InDragP
 	
 				static constexpr bool bUseUEHandyness = true;
 				const FQuat Delta = AnimationCore::QuatFromEuler(DeltaRot, GetRotationContext().RotationOrder, bUseUEHandyness);
-	
+
 				ApplyRotateDelta(Delta);
 				InteractionCurrAngle = HitAngle;
 			}
@@ -2668,24 +2846,35 @@ void UTransformGizmo::OnClickDragScreenSpaceRotate(const FInputDeviceRay& InDrag
 		const FVector2D DragDirToRemove = NormalProjectionToRemove * FVector2D::DotProduct(DragDir, NormalProjectionToRemove);
 		DragDir -= DragDirToRemove;
 
-		const double Delta = FVector2D::DotProduct(InteractionScreenAxisDirection, DragDir);
-		const FQuat DeltaRotAxis(InteractionPlanarNormal, FMath::DegreesToRadians(Delta));
+		double DeltaAngle = FMath::DegreesToRadians(FVector2D::DotProduct(InteractionScreenAxisDirection, DragDir));
+		SnapRotateAngleDelta(DeltaAngle);
+
+		const FQuat DeltaRotAxis(InteractionPlanarNormal, DeltaAngle);
 
 		ApplyRotateDelta(DeltaRotAxis);
+
 		InteractionScreenCurrPos = InDragPos.ScreenPosition;
+
 		return;
 	}
-	
+
 	FVector::FReal HitDepth;
 	if (GetRayParamIntersectionWithInteractionPlane(InDragPos, HitDepth))
 	{
 		const FVector HitPoint = InDragPos.WorldRay.Origin + InDragPos.WorldRay.Direction * HitDepth;
-		const float HitAngle = GizmoMath::ComputeAngleInPlane(HitPoint,
+		double HitAngle = GizmoMath::ComputeAngleInPlane(HitPoint,
 			InteractionPlanarOrigin, InteractionPlanarNormal, InteractionPlanarAxisX, InteractionPlanarAxisY);
 
+		double HitAngleDelta = HitAngle - InteractionCurrAngle;
+		SnapRotateAngleDelta(HitAngleDelta);
+
+		// Reapply the (potentially) snapped delta angle
+		HitAngle = InteractionCurrAngle + HitAngleDelta;
+
 		const FQuat Delta = ComputeAngularRotateDelta(InteractionCurrAngle, HitAngle);
-		
+
 		ApplyRotateDelta(Delta);
+
 		InteractionCurrAngle = HitAngle;
 	}
 }
@@ -2838,9 +3027,10 @@ void UTransformGizmo::OnClickDragArcBallRotate(const FInputDeviceRay& DragPos)
 	const FQuat FinalRotation = DeltaQ * StartRotation;
 	const FQuat InvCurrentRot = CurrentRotation.Inverse();
 	
-	const FQuat DeltaRot = (FinalRotation * InvCurrentRot).GetNormalized();
+	FQuat DeltaRot = (FinalRotation * InvCurrentRot).GetNormalized();
 	if (!DeltaRot.IsIdentity())
 	{
+		SnapRotateDelta(DeltaRot);
 		ApplyRotateDelta(DeltaRot);
 		CurrentRotation = DeltaRot * CurrentRotation; 
 	}
diff --git a/Engine/Source/Editor/Experimental/EditorInteractiveToolsFramework/Private/EditorInteractiveGizmoManager.cpp b/Engine/Source/Editor/Experimental/EditorInteractiveToolsFramework/Private/EditorInteractiveGizmoManager.cpp
--- a/Engine/Source/Editor/Experimental/EditorInteractiveToolsFramework/Private/EditorInteractiveGizmoManager.cpp
+++ b/Engine/Source/Editor/Experimental/EditorInteractiveToolsFramework/Private/EditorInteractiveGizmoManager.cpp
@@ -9,6 +9,7 @@
 #include "EditorInteractiveGizmoSelectionBuilder.h"
 #include "EditorInteractiveGizmoSubsystem.h"
 #include "EditorModeManager.h"
+#include "EdModeInteractiveToolsContext.h"
 #include "HAL/IConsoleManager.h"
 #include "HAL/Platform.h"
 #include "HAL/PlatformCrt.h"
@@ -18,6 +19,7 @@
 #include "Misc/AssertionMacros.h"
 #include "SceneView.h"
 #include "ShowFlags.h"
+#include "Snapping/EditorSnappingManager.h"
 #include "Templates/Casts.h"
 #include "ToolContextInterfaces.h"
 #include "EditorGizmos/TransformGizmo.h"
@@ -87,11 +89,26 @@ void UEditorInteractiveGizmoManager::InitializeWithEditorModeManager(IToolsConte
 {
 	Super::Initialize(QueriesAPIIn, TransactionsAPIIn, InputRouterIn);
 	EditorModeManager = InEditorModeManager;
-}
 
+	UModeManagerInteractiveToolsContext* InteractiveToolsContext = EditorModeManager->GetInteractiveToolsContext();
+	check(InteractiveToolsContext);
+
+	if (InteractiveToolsContext)
+	{
+		UE::Editor::Gizmos::RegisterSceneSnappingManager(InteractiveToolsContext);
+	}
+}
 
 void UEditorInteractiveGizmoManager::Shutdown()
 {
+	if (EditorModeManager)
+	{
+		if (const UModeManagerInteractiveToolsContext* InteractiveToolsContext = EditorModeManager->GetInteractiveToolsContext())
+		{
+			UE::Editor::Gizmos::DeregisterSceneSnappingManager(InteractiveToolsContext);
+		}
+	}
+
 	Registry->Shutdown();
 	Super::Shutdown();
 }
diff --git a/Engine/Source/Editor/Experimental/EditorInteractiveToolsFramework/Private/Snapping/EditorSnappingManager.cpp b/Engine/Source/Editor/Experimental/EditorInteractiveToolsFramework/Private/Snapping/EditorSnappingManager.cpp
new file mode 100644
--- /dev/null
+++ b/Engine/Source/Editor/Experimental/EditorInteractiveToolsFramework/Private/Snapping/EditorSnappingManager.cpp
@@ -0,0 +1,289 @@
+// Copyright Epic Games, Inc. All Rights Reserved.
+
+#include "Snapping/EditorSnappingManager.h"
+
+#include "ContextObjectStore.h"
+#include "GridSnapping.h"
+#include "InteractiveToolsContext.h"
+#include "PolicySnapping.h"
+#include "SceneQueries/SceneSnappingManager.h"
+#include "ToolContextInterfaces.h"
+#include "VectorUtil.h"
+#include "ViewportSnappingModule.h"
+
+namespace UE::Editor::Gizmos
+{
+	namespace EditorSnappingManagerLocals
+	{
+		// This assumes the provided AxisList only specifies a single axis
+		EAxis::Type GetAxisFromAxisList(EAxisList::Type InAxisList)
+		{
+			if (InAxisList == EAxisList::X)
+            {
+                return EAxis::X;
+            }
+
+			if (InAxisList == EAxisList::Y)
+			{
+				return EAxis::Y;
+			}
+
+			if (InAxisList == EAxisList::Z)
+			{
+				return EAxis::Z;
+			}
+
+			ensureMsgf(true, TEXT("Invalid AxisList, it must specify either X, Y or Z"));
+
+			return EAxis::None;
+		}
+	}
+
+	bool RegisterSceneSnappingManager(UInteractiveToolsContext* InToolsContext)
+	{
+		if (!ensure(InToolsContext))
+		{
+			return false;
+		}
+
+		// Check for existing registration, and return true if found
+		const UEditorSceneSnappingManager* FoundRegisteredSnappingManager = InToolsContext->ContextObjectStore->FindContext<UEditorSceneSnappingManager>();
+		if (FoundRegisteredSnappingManager)
+		{
+			return true;
+		}
+
+		UEditorSceneSnappingManager* SelectionManager = NewObject<UEditorSceneSnappingManager>(InToolsContext->ToolManager);
+		if (!ensure(SelectionManager))
+		{
+			return false;
+		}
+
+		SelectionManager->Initialize(InToolsContext);
+		InToolsContext->ContextObjectStore->AddContextObject(SelectionManager);
+
+		return true;
+	}
+
+	bool DeregisterSceneSnappingManager(const UInteractiveToolsContext* InToolsContext)
+	{
+		if (!ensure(InToolsContext))
+		{
+			return false;
+		}
+
+		if (UEditorSceneSnappingManager* FoundRegisteredSnappingManager = InToolsContext->ContextObjectStore->FindContext<UEditorSceneSnappingManager>())
+		{
+			FoundRegisteredSnappingManager->Shutdown();
+			InToolsContext->ContextObjectStore->RemoveContextObject(FoundRegisteredSnappingManager);
+		}
+
+		return true;
+	}
+
+	UEditorSceneSnappingManager* FindSceneSnappingManager(const UInteractiveToolManager* InToolManager)
+	{
+		if (!ensure(InToolManager))
+		{
+			return nullptr;
+		}
+
+		if (UEditorSceneSnappingManager* FoundRegisteredSnappingManager = InToolManager->GetContextObjectStore()->FindContext<UEditorSceneSnappingManager>())
+		{
+			return FoundRegisteredSnappingManager;
+		}
+
+		return nullptr;
+	}
+}
+
+void UEditorSceneSnappingManager::Initialize(const TObjectPtr<UInteractiveToolsContext>& InToolsContext)
+{
+	UInteractiveToolManager* ToolManager = InToolsContext ? InToolsContext->ToolManager : nullptr;
+	QueriesAPI = ToolManager ? ToolManager->GetContextQueriesAPI() : nullptr;
+
+	RegisterSnapQueryTargetHandler(MakeUnique<UE::Editor::Gizmos::FGizmoGridSnapper>(InToolsContext));
+
+	if (const TSharedPtr<ISnappingPolicy>& MergedSnappingPolicy = IViewportSnappingModule::GetSnapManager())
+	{
+		RegisterSnapQueryTargetHandler(MakeUnique<UE::Editor::Gizmos::FGizmoPolicySnapper>(MergedSnappingPolicy));
+	}
+}
+
+void UEditorSceneSnappingManager::Shutdown()
+{
+	QueryTargetHandlers.Reset();
+
+	QueriesAPI = nullptr;
+}
+
+bool UEditorSceneSnappingManager::ExecuteSceneHitQuery(const FSceneHitQueryRequest& InRequest, FSceneHitQueryResult& OutResult) const
+{
+	// @todo: implement when needed/supported
+	return false;
+}
+
+bool UEditorSceneSnappingManager::ExecuteSceneSnapQuery(const FSceneSnapQueryRequest& InRequest, TArray<FSceneSnapQueryResult>& OutResults) const
+{
+	switch (InRequest.RequestType)
+	{
+	case ESceneSnapQueryType::Position:
+		return ExecuteSceneSnapQueryPosition(InRequest, OutResults);
+	case ESceneSnapQueryType::Rotation:
+		return ExecuteSceneSnapQueryRotation(InRequest, OutResults);
+	case ESceneSnapQueryType::RotationAngle:
+		return ExecuteSceneSnapQueryRotationAngle(InRequest, OutResults);
+	case ESceneSnapQueryType::Scale:
+		return ExecuteSceneSnapQueryScale(InRequest, OutResults);
+	default:
+		ensureMsgf(false, TEXT("Only Position, Rotation and Scale Snap Queries are supported"));
+	}
+
+	return false;
+}
+
+void UEditorSceneSnappingManager::RegisterSnapQueryTargetHandler(TUniquePtr<FSceneSnapQueryTargetHandler>&& InHandler)
+{
+	QueryTargetHandlers.Emplace(MoveTemp(InHandler));
+}
+
+void UEditorSceneSnappingManager::DeregisterSnapQueryTargetHandler(const FName InHandlerName)
+{
+	const int32 HandlerIdx = QueryTargetHandlers.IndexOfByPredicate(
+			[InHandlerName](const TUniquePtr<FSceneSnapQueryTargetHandler>& InHandler)
+			{
+				return InHandler->TargetName == InHandlerName;
+			});
+	
+	if (HandlerIdx != INDEX_NONE)
+	{
+		QueryTargetHandlers.RemoveAt(HandlerIdx);
+	}
+}
+
+bool UEditorSceneSnappingManager::ExecuteSceneSnapQueryPosition(const FSceneSnapQueryRequest& InRequest, TArray<FSceneSnapQueryResult>& OutResults) const
+{
+	using namespace UE::Editor::Gizmos;
+
+	if (!QueriesAPI)
+	{
+		return false;
+	}
+
+	for (const TUniquePtr<FSceneSnapQueryTargetHandler>& Handler : QueryTargetHandlers)
+	{
+		if (Handler->IsQueryTypeSupported(ESceneSnapQueryType::Position))
+		{
+			FVector SnappedPosition;
+			if (Handler->SnapPosition(InRequest.Position, SnappedPosition) == ESceneSnapQueryTargetResult::Snapped)
+			{
+				FSceneSnapQueryResult SnapResult;
+				SnapResult.TargetType = Handler->GetTargetType();
+				SnapResult.Position = SnappedPosition;
+
+				OutResults.Emplace(SnapResult);
+
+				return true;
+			}
+		}
+	}
+
+	return false;
+}
+
+bool UEditorSceneSnappingManager::ExecuteSceneSnapQueryRotation(const FSceneSnapQueryRequest& InRequest, TArray<FSceneSnapQueryResult>& OutResults) const
+{
+	using namespace UE::Editor::Gizmos;
+
+	if (!QueriesAPI)
+	{
+		return false;
+	}
+
+	for (const TUniquePtr<FSceneSnapQueryTargetHandler>& Handler : QueryTargetHandlers)
+	{
+		if (Handler->IsQueryTypeSupported(ESceneSnapQueryType::Rotation))
+		{
+			FQuat SnappedRotation;
+			if (Handler->SnapRotation(InRequest.DeltaRotation, SnappedRotation) == ESceneSnapQueryTargetResult::Snapped)
+			{
+				FSceneSnapQueryResult SnapResult;
+				SnapResult.TargetType = Handler->GetTargetType();
+				SnapResult.DeltaRotation = SnappedRotation;
+
+				OutResults.Emplace(SnapResult);
+
+				return true;
+			}
+		}
+	}
+
+	return false;
+}
+
+bool UEditorSceneSnappingManager::ExecuteSceneSnapQueryRotationAngle(const FSceneSnapQueryRequest& InRequest, TArray<FSceneSnapQueryResult>& OutResults) const
+{
+	using namespace UE::Editor::Gizmos;
+
+	if (!QueriesAPI)
+	{
+		return false;
+	}
+	
+	// It's assumed the input AxisList only has a single active axis
+	const EAxis::Type Axis = EditorSnappingManagerLocals::GetAxisFromAxisList(InRequest.AxisList);
+	if (Axis == EAxis::None)
+	{
+		return false;
+	}
+
+	for (const TUniquePtr<FSceneSnapQueryTargetHandler>& Handler : QueryTargetHandlers)
+	{
+		if (Handler->IsQueryTypeSupported(ESceneSnapQueryType::RotationAngle))
+		{
+			double SnappedAngle;
+			if (Handler->SnapRotationAxisAngle(InRequest.RotationAngle, SnappedAngle, Axis) == ESceneSnapQueryTargetResult::Snapped)
+			{
+				FSceneSnapQueryResult SnapResult;
+				SnapResult.TargetType = Handler->GetTargetType();
+				SnapResult.RotationAngle = SnappedAngle;
+
+				OutResults.Emplace(SnapResult);
+
+				return true;
+			}
+		}
+	}
+
+	return false;
+}
+
+bool UEditorSceneSnappingManager::ExecuteSceneSnapQueryScale(const FSceneSnapQueryRequest& InRequest, TArray<FSceneSnapQueryResult>& OutResults) const
+{
+	using namespace UE::Editor::Gizmos;
+
+	if (!QueriesAPI)
+	{
+		return false;
+	}
+
+	for (const TUniquePtr<FSceneSnapQueryTargetHandler>& Handler : QueryTargetHandlers)
+	{
+		if (Handler->IsQueryTypeSupported(ESceneSnapQueryType::Scale))
+		{
+			FVector SnappedScale;
+			if (Handler->SnapScale(InRequest.Scale, SnappedScale) == ESceneSnapQueryTargetResult::Snapped)
+			{
+				FSceneSnapQueryResult SnapResult;
+				SnapResult.TargetType = Handler->GetTargetType();
+				SnapResult.Scale = SnappedScale;
+
+				OutResults.Emplace(SnapResult);
+
+				return true;
+			}
+		}
+	}
+
+	return false;
+}
diff --git a/Engine/Source/Editor/Experimental/EditorInteractiveToolsFramework/Private/Snapping/GridSnapping.cpp b/Engine/Source/Editor/Experimental/EditorInteractiveToolsFramework/Private/Snapping/GridSnapping.cpp
new file mode 100644
--- /dev/null
+++ b/Engine/Source/Editor/Experimental/EditorInteractiveToolsFramework/Private/Snapping/GridSnapping.cpp
@@ -0,0 +1,276 @@
+// Copyright Epic Games, Inc. All Rights Reserved.
+
+#include "GridSnapping.h"
+
+#include "BaseGizmos/GizmoMath.h"
+#include "InteractiveToolsContext.h"
+#include "SceneQueries/SceneSnappingManager.h"
+#include "ToolContextInterfaces.h"
+
+namespace UE::Editor::Gizmos
+{
+	FGizmoGridSnapper::FGizmoGridSnapper(const UInteractiveToolsContext* InToolsContext)
+	{
+		TargetName = "Grid";
+
+		UInteractiveToolManager* ToolManager = InToolsContext ? InToolsContext->ToolManager : nullptr;
+		QueriesAPI = ToolManager ? ToolManager->GetContextQueriesAPI() : nullptr;
+	}
+
+	ESceneSnapQueryTargetResult FGizmoGridSnapper::SnapPosition(const FGizmoGridSnapperRequest& InTargetRequest, const FVector& InPosition, FVector& OutSnappedPosition, const EAxisList::Type InAxisList) const
+	{
+		OutSnappedPosition = InPosition;
+
+		const bool bSnapX = InAxisList & EAxisList::X || InAxisList & EAxisList::Screen;
+		const bool bSnapY = InAxisList & EAxisList::Y || InAxisList & EAxisList::Screen;
+		const bool bSnapZ = InAxisList & EAxisList::Z || InAxisList & EAxisList::Screen;
+
+		if (!bSnapX && !bSnapY && !bSnapZ)
+		{
+			return ESceneSnapQueryTargetResult::NotSnapped;
+		}
+
+		const FVector GridSize = InTargetRequest.PositionGridSize;
+
+		OutSnappedPosition.X = bSnapX ? GizmoMath::SnapToIncrement(InPosition.X, GridSize.X) : InPosition.X;
+		OutSnappedPosition.Y = bSnapY ? GizmoMath::SnapToIncrement(InPosition.Y, GridSize.Y) : InPosition.Y;
+		OutSnappedPosition.Z = bSnapZ ? GizmoMath::SnapToIncrement(InPosition.Z, GridSize.Z) : InPosition.Z;
+
+		return ESceneSnapQueryTargetResult::Snapped;
+	}
+
+	ESceneSnapQueryTargetResult FGizmoGridSnapper::SnapPosition(const FVector& InPosition, FVector& OutSnappedPosition, const EAxisList::Type InAxisList) const
+	{
+		OutSnappedPosition = InPosition;
+
+		if (!QueriesAPI)
+		{
+			return ESceneSnapQueryTargetResult::Unsupported;
+		}
+
+		if (const FToolContextSnappingConfiguration SnappingConfig = QueriesAPI->GetCurrentSnappingSettings();
+			SnappingConfig.bEnablePositionGridSnapping)
+		{
+			FSceneSnapQueryResult SnapResult;
+			SnapResult.TargetType = ESceneSnapQueryTargetType::Grid;
+
+			FGizmoGridSnapperRequest GridRequest;
+			GridRequest.PositionGridSize = SnappingConfig.PositionGridDimensions;
+
+			return SnapPosition(GridRequest, InPosition, OutSnappedPosition, InAxisList);
+		}
+
+		return ESceneSnapQueryTargetResult::Disabled;
+	}
+
+	ESceneSnapQueryTargetResult FGizmoGridSnapper::SnapPositionAxis(const FGizmoGridSnapperRequest& InTargetRequest, const double InAxisValue, double& OutSnappedAxisValue, const EAxis::Type InAxis) const
+	{
+		OutSnappedAxisValue = InAxisValue;
+
+		const FVector GridSize = InTargetRequest.PositionGridSize;
+
+		OutSnappedAxisValue = GizmoMath::SnapToIncrement(InAxisValue, GridSize.GetComponentForAxis(InAxis));
+
+		return ESceneSnapQueryTargetResult::Snapped;
+	}
+
+	ESceneSnapQueryTargetResult FGizmoGridSnapper::SnapPositionAxis(const double InAxisValue, double& OutSnappedAxisValue, const EAxis::Type InAxis) const
+	{
+		OutSnappedAxisValue = InAxisValue;
+
+		if (!QueriesAPI)
+		{
+			return ESceneSnapQueryTargetResult::Unsupported;
+		}
+
+		if (const FToolContextSnappingConfiguration SnappingConfig = QueriesAPI->GetCurrentSnappingSettings();
+			SnappingConfig.bEnablePositionGridSnapping)
+		{
+			FSceneSnapQueryResult SnapResult;
+			SnapResult.TargetType = ESceneSnapQueryTargetType::Grid;
+
+			FGizmoGridSnapperRequest GridRequest;
+			GridRequest.PositionGridSize = SnappingConfig.PositionGridDimensions;
+
+			return SnapPositionAxis(GridRequest, InAxisValue, OutSnappedAxisValue, InAxis);
+		}
+
+		return ESceneSnapQueryTargetResult::Disabled;
+	}
+
+	ESceneSnapQueryTargetResult FGizmoGridSnapper::SnapRotation(const FGizmoGridSnapperRequest& InTargetRequest, const FQuat& InRotation, FQuat& OutSnappedRotation, const EAxisList::Type InAxisList) const
+	{
+		OutSnappedRotation = InRotation;
+
+		const bool bSnapRoll = InAxisList & EAxisList::X || InAxisList & EAxisList::Screen;
+		const bool bSnapPitch = InAxisList & EAxisList::Y || InAxisList & EAxisList::Screen;
+		const bool bSnapYaw = InAxisList & EAxisList::Z || InAxisList & EAxisList::Screen;
+
+		if (!bSnapRoll && !bSnapPitch && !bSnapYaw)
+		{
+			return ESceneSnapQueryTargetResult::NotSnapped;
+		}
+
+		const FRotator GridSize = InTargetRequest.RotationGridSize;
+
+		UE::Math::TVector<double> Axis;
+		double AngleRad;
+		InRotation.ToAxisAndAngle(Axis, AngleRad);
+
+		double AngleDeg = FMath::RadiansToDegrees(AngleRad);
+
+		AngleDeg = GizmoMath::SnapToIncrement(AngleDeg, GridSize.Pitch);
+		AngleRad = FMath::DegreesToRadians(AngleDeg);
+
+		OutSnappedRotation = FQuat(Axis, AngleRad);
+
+		return ESceneSnapQueryTargetResult::Snapped;
+	}
+
+	ESceneSnapQueryTargetResult FGizmoGridSnapper::SnapRotation(const FQuat& InRotation, FQuat& OutSnappedRotation, const EAxisList::Type InAxisList) const
+	{
+		OutSnappedRotation = InRotation;
+		
+		if (!QueriesAPI)
+		{
+			return ESceneSnapQueryTargetResult::Unsupported;
+		}
+
+		if (const FToolContextSnappingConfiguration SnappingConfig = QueriesAPI->GetCurrentSnappingSettings();
+			SnappingConfig.bEnableRotationGridSnapping)
+		{
+			FSceneSnapQueryResult SnapResult;
+			SnapResult.TargetType = ESceneSnapQueryTargetType::Grid;
+
+			FGizmoGridSnapperRequest GridRequest;
+			GridRequest.RotationGridSize = SnappingConfig.RotationGridAngles;
+
+			return SnapRotation(GridRequest, InRotation, OutSnappedRotation, InAxisList);
+		}
+
+		return ESceneSnapQueryTargetResult::Disabled;
+	}
+
+	ESceneSnapQueryTargetResult FGizmoGridSnapper::SnapRotationAxisAngle(const FGizmoGridSnapperRequest& InTargetRequest, const double& InAngle, double& OutSnappedAngle, const EAxis::Type InAxis) const
+	{
+		OutSnappedAngle = InAngle;
+
+		const FRotator GridSize = InTargetRequest.RotationGridSize;
+
+		const double SnapDelta = GridSize.GetComponentForAxis(InAxis);
+
+		OutSnappedAngle = FMath::RadiansToDegrees(InAngle);
+		OutSnappedAngle = GizmoMath::SnapToIncrement(OutSnappedAngle, SnapDelta);
+		OutSnappedAngle = FMath::DegreesToRadians(OutSnappedAngle);
+
+		return ESceneSnapQueryTargetResult::Snapped;
+	}
+
+	ESceneSnapQueryTargetResult FGizmoGridSnapper::SnapRotationAxisAngle(const double& InAngle, double& OutSnappedAngle, const EAxis::Type InAxis) const
+	{
+		OutSnappedAngle = InAngle;
+		
+		if (!QueriesAPI)
+		{
+			return ESceneSnapQueryTargetResult::Unsupported;
+		}
+
+		if (const FToolContextSnappingConfiguration SnappingConfig = QueriesAPI->GetCurrentSnappingSettings();
+			SnappingConfig.bEnableRotationGridSnapping)
+		{
+			FSceneSnapQueryResult SnapResult;
+			SnapResult.TargetType = ESceneSnapQueryTargetType::Grid;
+
+			FGizmoGridSnapperRequest GridRequest;
+			GridRequest.RotationGridSize = SnappingConfig.RotationGridAngles;
+
+			return SnapRotationAxisAngle(GridRequest, InAngle, OutSnappedAngle, InAxis);
+		}
+
+		return ESceneSnapQueryTargetResult::Disabled;
+	}
+
+	ESceneSnapQueryTargetResult FGizmoGridSnapper::SnapScale(const FGizmoGridSnapperRequest& InTargetRequest, const FVector& InScale, FVector& OutSnappedScale, const EAxisList::Type InAxisList) const
+	{
+		OutSnappedScale = InScale;
+
+		const bool bSnapX = InAxisList & EAxisList::X || InAxisList & EAxisList::Screen;
+		const bool bSnapY = InAxisList & EAxisList::Y || InAxisList & EAxisList::Screen;
+		const bool bSnapZ = InAxisList & EAxisList::Z || InAxisList & EAxisList::Screen;
+
+		if (!bSnapX && !bSnapY && !bSnapZ)
+		{
+			return ESceneSnapQueryTargetResult::NotSnapped;
+		}
+
+		const FVector GridSize = InTargetRequest.ScaleGridSize;
+
+		OutSnappedScale.X = bSnapX ? GizmoMath::SnapToIncrement(InScale.X, GridSize.X) : InScale.X;
+		OutSnappedScale.Y = bSnapY ? GizmoMath::SnapToIncrement(InScale.Y, GridSize.Y) : InScale.Y;
+		OutSnappedScale.Z = bSnapZ ? GizmoMath::SnapToIncrement(InScale.Z, GridSize.Z) : InScale.Z;
+
+		return ESceneSnapQueryTargetResult::Snapped;
+	}
+
+	ESceneSnapQueryTargetResult FGizmoGridSnapper::SnapScale(const FVector& InScale, FVector& OutSnappedScale, const EAxisList::Type InAxisList) const
+	{
+		if (!QueriesAPI)
+		{
+			return ESceneSnapQueryTargetResult::Unsupported;
+		}
+
+		if (const FToolContextSnappingConfiguration SnappingConfig = QueriesAPI->GetCurrentSnappingSettings();
+			SnappingConfig.bEnableScaleGridSnapping)
+		{
+			FSceneSnapQueryResult SnapResult;
+			SnapResult.TargetType = ESceneSnapQueryTargetType::Grid;
+
+			FGizmoGridSnapperRequest GridRequest;
+			GridRequest.ScaleGridSize = FVector(SnappingConfig.ScaleGridSize);
+
+			return SnapScale(GridRequest, InScale, OutSnappedScale, InAxisList);
+		}
+
+		return ESceneSnapQueryTargetResult::Disabled;
+	}
+
+	ESceneSnapQueryTargetResult FGizmoGridSnapper::SnapScaleAxis(const FGizmoGridSnapperRequest& InTargetRequest, const double InAxisValue, double& OutSnappedAxisValue, const EAxis::Type InAxis) const
+	{
+		OutSnappedAxisValue = InAxisValue;
+
+		const FVector GridSize = InTargetRequest.ScaleGridSize;
+
+		OutSnappedAxisValue = GizmoMath::SnapToIncrement(InAxisValue, GridSize.GetComponentForAxis(InAxis));
+
+		return ESceneSnapQueryTargetResult::Snapped;
+	}
+
+	ESceneSnapQueryTargetResult FGizmoGridSnapper::SnapScaleAxis(const double InAxisValue, double& OutSnappedAxisValue, const EAxis::Type InAxis) const
+	{
+		OutSnappedAxisValue = InAxisValue;
+
+		if (!QueriesAPI)
+		{
+			return ESceneSnapQueryTargetResult::Unsupported;
+		}
+
+		if (const FToolContextSnappingConfiguration SnappingConfig = QueriesAPI->GetCurrentSnappingSettings();
+			SnappingConfig.bEnableScaleGridSnapping)
+		{
+			FSceneSnapQueryResult SnapResult;
+			SnapResult.TargetType = ESceneSnapQueryTargetType::Grid;
+
+			FGizmoGridSnapperRequest GridRequest;
+			GridRequest.ScaleGridSize = FVector(SnappingConfig.ScaleGridSize);
+
+			return SnapScaleAxis(GridRequest, InAxisValue, OutSnappedAxisValue, InAxis);
+		}
+
+		return ESceneSnapQueryTargetResult::Disabled;
+	}
+
+	bool FGizmoGridSnapper::IsQueryTypeSupported(const ESceneSnapQueryType InQueryType) const
+	{
+		return InQueryType != ESceneSnapQueryType::Transform;
+	}
+}
diff --git a/Engine/Source/Editor/Experimental/EditorInteractiveToolsFramework/Private/Snapping/GridSnapping.h b/Engine/Source/Editor/Experimental/EditorInteractiveToolsFramework/Private/Snapping/GridSnapping.h
new file mode 100644
--- /dev/null
+++ b/Engine/Source/Editor/Experimental/EditorInteractiveToolsFramework/Private/Snapping/GridSnapping.h
@@ -0,0 +1,56 @@
+// Copyright Epic Games, Inc. All Rights Reserved.
+
+#pragma once
+
+#include "Math/MathFwd.h"
+#include "SceneQueries/SceneSnappingManager.h"
+
+class UInteractiveToolsContext;
+struct FSceneSnapQueryResult;
+struct FSceneSnapQueryRequest;
+struct FSceneHitQueryResult;
+struct FSceneHitQueryRequest;
+class USceneSnappingManager;
+class UInteractiveGizmoManager;
+class IToolsContextQueriesAPI;
+
+class UEditorSceneSnappingManager;
+
+namespace UE::Editor::Gizmos
+{
+	struct FGizmoGridSnapperRequest
+	{
+		FVector PositionGridSize;
+		FRotator RotationGridSize;
+		FVector ScaleGridSize;
+	};
+
+	class FGizmoGridSnapper final : public TSceneSnapQueryTargetHandler<FGizmoGridSnapperRequest>
+	{
+	public:
+		explicit FGizmoGridSnapper(const UInteractiveToolsContext* InToolsContext);
+		virtual ~FGizmoGridSnapper() override = default;
+
+		//~ Begin FSceneSnapPolicyTarget
+		virtual ESceneSnapQueryTargetResult SnapPosition(const FVector& InPosition, FVector& OutSnappedPosition, const EAxisList::Type InAxisList = EAxisList::All) const override;
+		virtual ESceneSnapQueryTargetResult SnapPositionAxis(const double InAxisValue, double& OutSnappedAxisValue, const EAxis::Type InAxis) const override;
+		virtual ESceneSnapQueryTargetResult SnapRotation(const FQuat& InRotation, FQuat& OutSnappedRotation, const EAxisList::Type InAxisList = EAxisList::All) const override;
+		virtual ESceneSnapQueryTargetResult SnapRotationAxisAngle(const double& InAngle, double& OutSnappedAngle, const EAxis::Type InAxis) const override;
+		virtual ESceneSnapQueryTargetResult SnapScale(const FVector& InScale, FVector& OutSnappedScale, const EAxisList::Type InAxisList = EAxisList::All) const override;
+		virtual ESceneSnapQueryTargetResult SnapScaleAxis(const double InAxisValue, double& OutSnappedAxisValue, const EAxis::Type InAxis) const override;
+		virtual bool IsQueryTypeSupported(const ESceneSnapQueryType InQueryType) const override;
+		//~ End FSceneSnapPolicyTarget
+
+		//~ Begin TSceneSnapQueryTarget
+		virtual ESceneSnapQueryTargetResult SnapPosition(const FGizmoGridSnapperRequest& InTargetRequest, const FVector& InPosition, FVector& OutSnappedPosition, const EAxisList::Type InAxisList = EAxisList::All) const override;
+		virtual ESceneSnapQueryTargetResult SnapPositionAxis(const FGizmoGridSnapperRequest& InTargetRequest, const double InAxisValue, double& OutSnappedPositionAxis, const EAxis::Type InAxis) const override;
+		virtual ESceneSnapQueryTargetResult SnapRotation(const FGizmoGridSnapperRequest& InTargetRequest, const FQuat& InRotation, FQuat& OutSnappedRotation, const EAxisList::Type InAxisList = EAxisList::All) const override;
+		virtual ESceneSnapQueryTargetResult SnapRotationAxisAngle(const FGizmoGridSnapperRequest& InTargetRequest, const double& InAngle, double& OutSnappedAngle, const EAxis::Type InAxis) const override;
+		virtual ESceneSnapQueryTargetResult SnapScale(const FGizmoGridSnapperRequest& InTargetRequest, const FVector& InScale, FVector& OutSnappedScale, const EAxisList::Type InAxisList = EAxisList::All) const override;
+		virtual ESceneSnapQueryTargetResult SnapScaleAxis(const FGizmoGridSnapperRequest& InTargetRequest, const double InAxisValue, double& OutSnappedAxisValue, const EAxis::Type InAxis) const override;
+		//~ End TSceneSnapQueryTarget
+
+	private:
+		const IToolsContextQueriesAPI* QueriesAPI = nullptr;
+	};
+}
diff --git a/Engine/Source/Editor/Experimental/EditorInteractiveToolsFramework/Private/Snapping/PolicySnapping.cpp b/Engine/Source/Editor/Experimental/EditorInteractiveToolsFramework/Private/Snapping/PolicySnapping.cpp
new file mode 100644
--- /dev/null
+++ b/Engine/Source/Editor/Experimental/EditorInteractiveToolsFramework/Private/Snapping/PolicySnapping.cpp
@@ -0,0 +1,92 @@
+// Copyright Epic Games, Inc. All Rights Reserved.
+
+#include "Snapping/PolicySnapping.h"
+
+#include "Editor.h"
+#include "ISnappingPolicy.h"
+#include "SceneQueries/SceneSnappingManager.h"
+
+namespace UE::Editor::Gizmos
+{
+	FGizmoPolicySnapper::FGizmoPolicySnapper(const TSharedPtr<ISnappingPolicy>& InSnappingPolicy)
+		: SnappingPolicy(InSnappingPolicy)
+	{
+		TargetName = "Policy";
+	}
+
+	ESceneSnapQueryTargetResult FGizmoPolicySnapper::SnapTransform(const FTransform& InTransform, FVector& OutSnappedTransform) const
+	{
+		return FSceneSnapQueryTargetHandler::SnapTransform(InTransform, OutSnappedTransform);
+	}
+
+	ESceneSnapQueryTargetResult FGizmoPolicySnapper::SnapPosition(const FVector& InPosition, FVector& OutSnappedPosition, const EAxisList::Type InAxisList) const
+	{
+		OutSnappedPosition = InPosition;
+		const FVector GridSize = FVector(GEditor->GetGridSize());
+		SnappingPolicy->SnapPointToGrid(OutSnappedPosition, GridSize);
+
+		return ESceneSnapQueryTargetResult::Snapped;
+	}
+
+	ESceneSnapQueryTargetResult FGizmoPolicySnapper::SnapPositionAxis(const double InAxisDelta, double& OutSnappedAxisDelta, const EAxis::Type InAxis) const
+	{
+		FVector SnappedPosition;
+		SnappedPosition.SetComponentForAxis(InAxis, InAxisDelta);
+
+		const FVector GridSize = FVector(GEditor->GetGridSize());
+		SnappingPolicy->SnapPointToGrid(SnappedPosition, GridSize);
+
+		OutSnappedAxisDelta = SnappedPosition.GetComponentForAxis(InAxis);
+
+		return ESceneSnapQueryTargetResult::Snapped;
+	}
+
+	ESceneSnapQueryTargetResult FGizmoPolicySnapper::SnapRotation(const FQuat& InRotation, FQuat& OutRotation, const EAxisList::Type InAxisList) const
+	{
+		FRotator Rotation = InRotation.Rotator();
+		SnappingPolicy->SnapRotatorToGrid(Rotation);
+
+		OutRotation = Rotation.Quaternion();
+
+		return ESceneSnapQueryTargetResult::Snapped;
+	}
+
+	ESceneSnapQueryTargetResult FGizmoPolicySnapper::SnapRotationAxisAngle(const double& InAngleDelta, double& OutSnappedAngleDelta, const EAxis::Type InAxis) const
+	{
+		FRotator Rotation;
+		Rotation.SetComponentForAxis(InAxis, InAngleDelta);
+
+		SnappingPolicy->SnapRotatorToGrid(Rotation);
+
+		OutSnappedAngleDelta = Rotation.GetComponentForAxis(InAxis);
+
+		return ESceneSnapQueryTargetResult::Snapped;
+	}
+
+	ESceneSnapQueryTargetResult FGizmoPolicySnapper::SnapScale(const FVector& InScale, FVector& OutSnappedScale, const EAxisList::Type InAxisList) const
+	{
+		OutSnappedScale = InScale;
+		const FVector GridSize = FVector(GEditor->GetScaleGridSize());
+		SnappingPolicy->SnapScale(OutSnappedScale, GridSize);
+
+		return ESceneSnapQueryTargetResult::Snapped;
+	}
+
+	ESceneSnapQueryTargetResult FGizmoPolicySnapper::SnapScaleAxis(const double InAxisDelta, double& OutSnappedAxisDelta, const EAxis::Type InAxis) const
+	{
+		FVector SnappedScale;
+		SnappedScale.SetComponentForAxis(InAxis, InAxisDelta);
+
+		const FVector GridSize = FVector(GEditor->GetScaleGridSize());
+		SnappingPolicy->SnapScale(SnappedScale, GridSize);
+
+		OutSnappedAxisDelta = SnappedScale.GetComponentForAxis(InAxis);
+
+		return ESceneSnapQueryTargetResult::Snapped;
+	}
+
+	bool FGizmoPolicySnapper::IsQueryTypeSupported(const ESceneSnapQueryType InQueryType) const
+	{
+		return true;
+	}
+}
diff --git a/Engine/Source/Editor/Experimental/EditorInteractiveToolsFramework/Private/Snapping/PolicySnapping.h b/Engine/Source/Editor/Experimental/EditorInteractiveToolsFramework/Private/Snapping/PolicySnapping.h
new file mode 100644
--- /dev/null
+++ b/Engine/Source/Editor/Experimental/EditorInteractiveToolsFramework/Private/Snapping/PolicySnapping.h
@@ -0,0 +1,42 @@
+// Copyright Epic Games, Inc. All Rights Reserved.
+
+#pragma once
+
+#include "Math/MathFwd.h"
+#include "SceneQueries/SceneSnappingManager.h"
+
+class UInteractiveToolsContext;
+struct FSceneSnapQueryResult;
+struct FSceneSnapQueryRequest;
+struct FSceneHitQueryResult;
+struct FSceneHitQueryRequest;
+class USceneSnappingManager;
+class UInteractiveGizmoManager;
+class IToolsContextQueriesAPI;
+
+class UEditorSceneSnappingManager;
+
+namespace UE::Editor::Gizmos
+{
+	/** Adapter for a given ISnappingPolicy. */
+	struct FGizmoPolicySnapper final : public FSceneSnapQueryTargetHandler
+	{
+	public:
+		explicit FGizmoPolicySnapper(const TSharedPtr<ISnappingPolicy>& InSnappingPolicy);
+		virtual ~FGizmoPolicySnapper() override = default;
+
+		//~ Begin FSceneSnapPolicyTarget
+		virtual ESceneSnapQueryTargetResult SnapTransform(const FTransform& InTransform, FVector& OutSnappedTransform) const override;
+		virtual ESceneSnapQueryTargetResult SnapPosition(const FVector& InPosition, FVector& OutSnappedPosition, const EAxisList::Type InAxisList) const override;
+		virtual ESceneSnapQueryTargetResult SnapPositionAxis(const double InAxisDelta, double& OutSnappedAxisDelta, const EAxis::Type InAxis) const override;
+		virtual ESceneSnapQueryTargetResult SnapRotation(const FQuat& InRotation, FQuat& OutRotation, const EAxisList::Type InAxisList) const override;
+		virtual ESceneSnapQueryTargetResult SnapRotationAxisAngle(const double& InAngleDelta, double& OutSnappedAngleDelta, const EAxis::Type InAxis) const override;
+		virtual ESceneSnapQueryTargetResult SnapScale(const FVector& InScale, FVector& OutSnappedScale, const EAxisList::Type InAxisList) const override;
+		virtual ESceneSnapQueryTargetResult SnapScaleAxis(const double InAxisDelta, double& OutSnappedAxisDelta, const EAxis::Type InAxis) const override;
+		virtual bool IsQueryTypeSupported(const ESceneSnapQueryType InQueryType) const override;
+		//~ End FSceneSnapPolicyTarget
+
+	private:
+		TSharedPtr<ISnappingPolicy> SnappingPolicy;
+	};
+}
diff --git a/Engine/Source/Editor/Experimental/EditorInteractiveToolsFramework/Public/EditorGizmos/TransformGizmo.h b/Engine/Source/Editor/Experimental/EditorInteractiveToolsFramework/Public/EditorGizmos/TransformGizmo.h
--- a/Engine/Source/Editor/Experimental/EditorInteractiveToolsFramework/Public/EditorGizmos/TransformGizmo.h
+++ b/Engine/Source/Editor/Experimental/EditorInteractiveToolsFramework/Public/EditorGizmos/TransformGizmo.h
@@ -628,8 +628,11 @@ protected:
 	/** Get screen-space axis based on a world space axis */
 	FVector2D GetWorldToScreenRotateAxisDir(const FInputDeviceRay& InPressPos, const FVector& InWorldAxis);
 
-	/** Compute rotate delta based on screen-space start/end positions */
-	virtual FQuat ComputeAxisRotateDelta(const FVector2D& InStartPos, const FInputDeviceRay& InDragPos);
+	/** Compute angle (degrees) delta based on screen-space start/end positions. Intended to be used in combination with ComputeAxisRotateDelta */
+	virtual double ComputeAxisRotateDeltaAngle(const FVector2D& InStartPos, const FInputDeviceRay& InDragPos);
+
+	/** Compute rotate delta based on screen-space angle delta (in degrees) */
+	virtual FQuat ComputeAxisRotateDelta(const double& InDeltaAngle);
 
 	/** Compute gimbal rotate delta based on screen-space start/end positions */
 	FQuat ComputeGimbalRotateDelta(const FVector2D& InStartPos, const FVector2D& InEndPos);
@@ -760,6 +763,11 @@ protected:
 	bool PositionSnapFunction(const FVector& WorldPosition, FVector& SnappedPositionOut) const;
 	FQuat RotationSnapFunction(const FQuat& DeltaRotation) const;
 
+	void SnapTranslateDelta(FVector& InOutWorldDelta) const;
+	void SnapRotateDelta(FQuat& InOutWorldDelta) const;
+	void SnapRotateAngleDelta(double& InOutAngleDelta) const;
+	void SnapScaleDelta(FVector& InOutLocalScaleDelta) const;
+
 	// Get max part identifier.
 	virtual uint32 GetMaxPartIdentifier() const;
 
@@ -854,11 +862,11 @@ protected:
 
 	/** Active interaction start hit param (only valid between state target BeginModify/EndModify) */
 	UPROPERTY()
-	float InteractionAxisStartParam;
+	double InteractionAxisStartParam;
 
 	/** Active interaction current hit param (only valid between state target BeginModify/EndModify) */
 	UPROPERTY()
-	float InteractionAxisCurrParam;
+	double InteractionAxisCurrParam;
 
 	/** Active world space planar origin (only valid between state target BeginModify/EndModify) */
 	UPROPERTY()
@@ -890,11 +898,11 @@ protected:
 
 	/** Active interaction rotation start angle (only valid between state target BeginModify/EndModify) */
 	UPROPERTY()
-	float InteractionStartAngle;
+	double InteractionStartAngle;
 
 	/** Active interaction rotation curr angle (only valid between state target BeginModify/EndModify) */
 	UPROPERTY()
-	float InteractionCurrAngle;
+	double InteractionCurrAngle;
 
 	/** Active interaction screen axis dir (only valid between state target BeginModify/EndModify) */
 	UPROPERTY()
diff --git a/Engine/Source/Editor/Experimental/EditorInteractiveToolsFramework/Public/Snapping/EditorSnappingManager.h b/Engine/Source/Editor/Experimental/EditorInteractiveToolsFramework/Public/Snapping/EditorSnappingManager.h
new file mode 100644
--- /dev/null
+++ b/Engine/Source/Editor/Experimental/EditorInteractiveToolsFramework/Public/Snapping/EditorSnappingManager.h
@@ -0,0 +1,55 @@
+// Copyright Epic Games, Inc. All Rights Reserved.
+
+#pragma once
+
+#include "Math/MathFwd.h"
+#include "SceneQueries/SceneSnappingManager.h"
+
+#include "EditorSnappingManager.generated.h"
+
+class UInteractiveToolsContext;
+struct FSceneSnapQueryResult;
+struct FSceneSnapQueryRequest;
+struct FSceneHitQueryResult;
+struct FSceneHitQueryRequest;
+class USceneSnappingManager;
+class UInteractiveGizmoManager;
+class IToolsContextQueriesAPI;
+
+class UEditorSceneSnappingManager;
+
+namespace UE::Editor::Gizmos
+{
+	EDITORINTERACTIVETOOLSFRAMEWORK_API bool RegisterSceneSnappingManager(UInteractiveToolsContext* InToolsContext);
+
+	EDITORINTERACTIVETOOLSFRAMEWORK_API bool DeregisterSceneSnappingManager(const UInteractiveToolsContext* InToolsContext);
+
+	EDITORINTERACTIVETOOLSFRAMEWORK_API UEditorSceneSnappingManager* FindSceneSnappingManager(const UInteractiveToolManager* InToolManager);
+}
+
+UCLASS(MinimalAPI)
+class UEditorSceneSnappingManager : public USceneSnappingManager
+{
+	GENERATED_BODY()
+
+public:
+	void Initialize(const TObjectPtr<UInteractiveToolsContext>& InToolsContext);
+	void Shutdown();
+
+	virtual bool ExecuteSceneHitQuery(const FSceneHitQueryRequest& InRequest, FSceneHitQueryResult& OutResult) const override;
+	virtual bool ExecuteSceneSnapQuery(const FSceneSnapQueryRequest& InRequest, TArray<FSceneSnapQueryResult>& OutResults) const override;
+
+	/** Register a snapping handler for a given snap target. */
+	EDITORINTERACTIVETOOLSFRAMEWORK_API void RegisterSnapQueryTargetHandler(TUniquePtr<FSceneSnapQueryTargetHandler>&& InHandler);
+	EDITORINTERACTIVETOOLSFRAMEWORK_API void DeregisterSnapQueryTargetHandler(const FName InHandlerName);
+
+protected:
+	bool ExecuteSceneSnapQueryPosition(const FSceneSnapQueryRequest& InRequest, TArray<FSceneSnapQueryResult>& OutResults) const;
+	bool ExecuteSceneSnapQueryRotation(const FSceneSnapQueryRequest& InRequest, TArray<FSceneSnapQueryResult>& OutResults) const;
+	bool ExecuteSceneSnapQueryRotationAngle(const FSceneSnapQueryRequest& InRequest, TArray<FSceneSnapQueryResult>& OutResults) const;
+	bool ExecuteSceneSnapQueryScale(const FSceneSnapQueryRequest& InRequest, TArray<FSceneSnapQueryResult>& OutResults) const;
+
+	TArray<TUniquePtr<FSceneSnapQueryTargetHandler>> QueryTargetHandlers;
+
+	const IToolsContextQueriesAPI* QueriesAPI = nullptr;
+};
diff --git a/Engine/Source/Editor/UnrealEd/Private/WidgetModeManager.cpp b/Engine/Source/Editor/UnrealEd/Private/WidgetModeManager.cpp
--- a/Engine/Source/Editor/UnrealEd/Private/WidgetModeManager.cpp
+++ b/Engine/Source/Editor/UnrealEd/Private/WidgetModeManager.cpp
@@ -3,6 +3,7 @@
 #include "WidgetModeManager.h"
 
 #include "EditorGizmos/EditorTransformGizmoUtil.h"
+#include "Snapping/EditorSnappingManager.h"
 #include "ToolContexts/WidgetToolsContext.h"
 
 //////////////////////////////////////////////////////////////////////////
@@ -15,6 +16,7 @@ FWidgetModeManager::FWidgetModeManager()
 	if (UObjectInitialized())
 	{
 		UE::EditorTransformGizmoUtil::UnregisterTransformGizmoContextObject(this);
+		UE::Editor::Gizmos::DeregisterSceneSnappingManager(InteractiveToolsContext);
 		InteractiveToolsContext->ShutdownContext();
 		InteractiveToolsContext = nullptr;
 	}
@@ -23,6 +25,7 @@ FWidgetModeManager::FWidgetModeManager()
 	InteractiveToolsContext = CachedWidgetToolContext;
 	InteractiveToolsContext->InitializeContextWithEditorModeManager(this);
 	UE::EditorTransformGizmoUtil::RegisterTransformGizmoContextObject(this);
+	UE::Editor::Gizmos::RegisterSceneSnappingManager(InteractiveToolsContext);
 }
 
 bool FWidgetModeManager::OnKeyChar(const FGeometry& MyGeometry, const FCharacterEvent& InCharacterEvent)
diff --git a/Engine/Source/Runtime/InteractiveToolsFramework/Private/BaseGizmos/GizmoMath.cpp b/Engine/Source/Runtime/InteractiveToolsFramework/Private/BaseGizmos/GizmoMath.cpp
--- a/Engine/Source/Runtime/InteractiveToolsFramework/Private/BaseGizmos/GizmoMath.cpp
+++ b/Engine/Source/Runtime/InteractiveToolsFramework/Private/BaseGizmos/GizmoMath.cpp
@@ -702,7 +702,7 @@ RealType GizmoMath::SnapToIncrement(RealType Value, RealType Increment)
 	Value = FMath::Abs(Value);
 	int IntIncrement = (int)(Value / Increment);
 	RealType Remainder = (RealType)fmod(Value, Increment);
-	if (Remainder > IntIncrement / 2)
+	if (Remainder > Increment / 2)
 	{
 		++IntIncrement;
 	}
diff --git a/Engine/Source/Runtime/InteractiveToolsFramework/Private/SceneQueries/SceneSnappingManager.cpp b/Engine/Source/Runtime/InteractiveToolsFramework/Private/SceneQueries/SceneSnappingManager.cpp
--- a/Engine/Source/Runtime/InteractiveToolsFramework/Private/SceneQueries/SceneSnappingManager.cpp
+++ b/Engine/Source/Runtime/InteractiveToolsFramework/Private/SceneQueries/SceneSnappingManager.cpp
@@ -90,8 +90,15 @@ void FSceneHitQueryResult::InitializeHitResult(const FSceneHitQueryRequest& From
 	// initialize .Time? Need start/end for that...
 }
 
+EToolContextCoordinateSystem FSceneSnapQueryRequest::GetResultCoordinateSpace() const
+{
+	return ResultCoordinateSpace.Get(RequestCoordinateSpace);
+}
 
-
+EToolContextCoordinateSystem FSceneSnapQueryResult::GetCoordinateSpace(const FSceneSnapQueryRequest& InFromRequest) const
+{
+	return CoordinateSpace.Get(InFromRequest.GetResultCoordinateSpace());
+}
 
 USceneSnappingManager* USceneSnappingManager::Find(UInteractiveToolManager* ToolManager)
 {
diff --git a/Engine/Source/Runtime/InteractiveToolsFramework/Public/SceneQueries/SceneSnappingManager.h b/Engine/Source/Runtime/InteractiveToolsFramework/Public/SceneQueries/SceneSnappingManager.h
--- a/Engine/Source/Runtime/InteractiveToolsFramework/Public/SceneQueries/SceneSnappingManager.h
+++ b/Engine/Source/Runtime/InteractiveToolsFramework/Public/SceneQueries/SceneSnappingManager.h
@@ -7,16 +7,17 @@
 #include "Math/Vector.h"
 #include "Math/Ray.h"
 #include "Engine/HitResult.h"
+#include "ToolContextInterfaces.h"
+
 #include "SceneSnappingManager.generated.h"
 
+class ISnappingPolicy;
 class AActor;
 class UActorComponent;
 class UPrimitiveComponent;
 class UInteractiveToolManager;
 class UInteractiveGizmoManager;
 
-
-
 struct FSceneQueryVisibilityFilter
 {
 	/** Optional: components to consider invisible even if they aren't. */
@@ -76,9 +77,16 @@ struct FSceneHitQueryResult
 UENUM()
 enum class ESceneSnapQueryType : uint8
 {
-	/** snapping a position */
+	/** Position/Translation snapping */
 	Position = 1,
-	Rotation = 2
+	/** Full rotation snapping (which could affect the axis of rotation) */
+	Rotation = 2,
+	/** Snaps the rotation angle only (in degrees). */
+	RotationAngle = 3,
+	/** Scale snapping */
+	Scale = 4,
+	/** Transform snapping (ie. snap to surface is both position and rotation) */
+	Transform = 5
 };
 
 
@@ -93,28 +101,138 @@ enum class ESceneSnapQueryTargetType : uint8
 	MeshEdge = 2,
 	/** Grid Snapping */
 	Grid = 4,
+	/** Any object with a transform representation  */
+	ObjectTransform = 8,
+	/** As defined by an FName */
+	Custom = 16,
 
-	All = MeshVertex | MeshEdge | Grid
+	All = MeshVertex | MeshEdge | Grid | ObjectTransform | Custom
 };
 ENUM_CLASS_FLAGS(ESceneSnapQueryTargetType);
 
+/** Describes the result of the attempted snap operation.
+ * Note that this differs to FSceneSnapQueryResult in that it only
+ * indicates the success or failure of the snap, not the resulting data. */
+UENUM()
+enum class ESceneSnapQueryTargetResult : uint8
+{
+	/** Snapping was valid and attempted, but did not snap to target. */
+	NotSnapped,
+
+	/** Snapping was valid and attempted, and did snap to target. This may indicate no further snapping operations should be attempted for the given interaction. */
+	Snapped,
+
+	/** Snapping is valid and supported, but is disabled (usually due to a user setting). */
+	Disabled,
+
+	/** The snapping query isn't valid for this Query Target. */
+	Unsupported
+};
+
+/** A custom query target implementation, allowing extensibility and handling
+ * of custom targets not defined in the ESceneSnapQueryTargetType enum. */
+class FSceneSnapQueryTargetHandler
+{
+public:
+	virtual ~FSceneSnapQueryTargetHandler() = default;
+
+	/** Snaps the given transform, potentially affecting all query types (position, rotation, scale). The given transform can be absolute or a relative delta. */
+	virtual ESceneSnapQueryTargetResult SnapTransform(const FTransform& InTransform, FVector& OutSnappedTransform) const { return ESceneSnapQueryTargetResult::Unsupported; }
+
+	virtual ESceneSnapQueryTargetResult SnapPosition(const FVector& InPosition, FVector& OutSnappedPosition, const EAxisList::Type InAxisList = EAxisList::All) const { return ESceneSnapQueryTargetResult::Unsupported; }
+	virtual ESceneSnapQueryTargetResult SnapPositionAxis(const double InAxisValue, double& OutSnappedAxisValue, const EAxis::Type InAxis) const { return ESceneSnapQueryTargetResult::Unsupported; }
+
+	virtual ESceneSnapQueryTargetResult SnapRotation(const FQuat& InRotation, FQuat& OutSnappedRotation, const EAxisList::Type InAxisList = EAxisList::All) const { return ESceneSnapQueryTargetResult::Unsupported; }
+	virtual ESceneSnapQueryTargetResult SnapRotationAxisAngle(const double& InAngle, double& OutSnappedAngle, const EAxis::Type InAxis) const { return ESceneSnapQueryTargetResult::Unsupported; }
+
+	virtual ESceneSnapQueryTargetResult SnapScale(const FVector& InScale, FVector& OutSnappedScale, const EAxisList::Type InAxisList = EAxisList::All) const { return ESceneSnapQueryTargetResult::Unsupported; }
+	virtual ESceneSnapQueryTargetResult SnapScaleAxis(const double InAxisValue, double& OutSnappedAxisValue, const EAxis::Type InAxis) const { return ESceneSnapQueryTargetResult::Unsupported; }
+
+	virtual ESceneSnapQueryTargetType GetTargetType() const { return ESceneSnapQueryTargetType::Custom; }
+
+	/** Returns true if the given query type can -ever- be supported, regardless of any variable state. */
+	virtual bool IsQueryTypeSupported(const ESceneSnapQueryType InQueryType) const { return false; }
+
+	bool operator==(const FSceneSnapQueryTargetHandler& InOther) const
+	{
+		return TargetName == InOther.TargetName;
+	}
+
+	bool operator !=(const FSceneSnapQueryTargetHandler& InOther) const
+	{
+		return TargetName != InOther.TargetName;
+	}
+
+public:
+	/** Unique query target name, ie. "Grid" */
+	FName TargetName;
+};
+
+inline uint32 GetTypeHash(const FSceneSnapQueryTargetHandler& InSnapQueryTargetHandler)
+{
+	return GetTypeHash(InSnapQueryTargetHandler.TargetName);
+}
+
+/** A custom query target implementation with typed parameters, allowing extensibility and handling
+ * of custom targets not defined in the ESceneSnapQueryTargetType enum. */
+template <typename TargetRequestParametersType UE_REQUIRES(std::is_class_v<TargetRequestParametersType>)>
+class TSceneSnapQueryTargetHandler : public FSceneSnapQueryTargetHandler
+{
+public:
+	/** Snaps the given transform, potentially affecting all query types (position, rotation, scale). The given transform can be absolute or a relative delta. */
+	virtual ESceneSnapQueryTargetResult SnapTransform(const FTransform& InTransform, FVector& OutSnappedTransform) const override { return ESceneSnapQueryTargetResult::Unsupported; }
+
+	/** Snaps the given transform, potentially affecting all query types (position, rotation, scale). The given transform can be absolute or a relative delta. */
+	virtual ESceneSnapQueryTargetResult SnapTransform(const TargetRequestParametersType& InTargetRequest, const FTransform& InTransform, FVector& OutSnappedTransform) const { return ESceneSnapQueryTargetResult::Unsupported; }
+
+	virtual ESceneSnapQueryTargetResult SnapPosition(const FVector& InPosition, FVector& OutSnappedPosition, const EAxisList::Type InAxisList = EAxisList::All) const override { return ESceneSnapQueryTargetResult::Unsupported; }
+	virtual ESceneSnapQueryTargetResult SnapPosition(const TargetRequestParametersType& InTargetRequest, const FVector& InPosition, FVector& OutSnappedPosition, const EAxisList::Type InAxisList = EAxisList::All) const { return ESceneSnapQueryTargetResult::Unsupported; }
+
+	virtual ESceneSnapQueryTargetResult SnapPositionAxis(const double InAxisValue, double& OutSnappedAxisValue, const EAxis::Type InAxis) const override { return ESceneSnapQueryTargetResult::Unsupported; }
+	virtual ESceneSnapQueryTargetResult SnapPositionAxis(const TargetRequestParametersType& InTargetRequest, const double InAxisValue, double& OutSnappedAxisValue, const EAxis::Type InAxis) const { return ESceneSnapQueryTargetResult::Unsupported; }
+
+	virtual ESceneSnapQueryTargetResult SnapRotation(const FQuat& InRotation, FQuat& OutSnappedRotation, const EAxisList::Type InAxisList = EAxisList::All) const override { return ESceneSnapQueryTargetResult::Unsupported; }
+	virtual ESceneSnapQueryTargetResult SnapRotation(const TargetRequestParametersType& InTargetRequest, const FQuat& InRotation, FQuat& OutSnappedRotation, const EAxisList::Type InAxisList = EAxisList::All) const { return ESceneSnapQueryTargetResult::Unsupported; }
+
+	virtual ESceneSnapQueryTargetResult SnapRotationAxisAngle(const double& InAngle, double& OutSnappedAngle, const EAxis::Type InAxis) const override { return ESceneSnapQueryTargetResult::Unsupported; }
+	virtual ESceneSnapQueryTargetResult SnapRotationAxisAngle(const TargetRequestParametersType& InTargetRequest, const double& InAngle, double& OutSnappedAngle, const EAxis::Type InAxis) const { return ESceneSnapQueryTargetResult::Unsupported; }
+
+	virtual ESceneSnapQueryTargetResult SnapScale(const FVector& InScale, FVector& OutSnappedScale, const EAxisList::Type InAxisList = EAxisList::All) const override { return ESceneSnapQueryTargetResult::Unsupported; }
+	virtual ESceneSnapQueryTargetResult SnapScale(const TargetRequestParametersType& InTargetRequest, const FVector& InScale, FVector& OutSnappedScale, const EAxisList::Type InAxisList = EAxisList::All) const { return ESceneSnapQueryTargetResult::Unsupported; }
+
+	virtual ESceneSnapQueryTargetResult SnapScaleAxis(const double InAxisValue, double& OutSnappedAxisValue, const EAxis::Type InAxis) const override { return ESceneSnapQueryTargetResult::Unsupported; }
+	virtual ESceneSnapQueryTargetResult SnapScaleAxis(const TargetRequestParametersType& InTargetRequest, const double InAxisValue, double& OutSnappedAxisValue, const EAxis::Type InAxis) const { return ESceneSnapQueryTargetResult::Unsupported; }
+};
 
 /**
  * Configuration variables for a USceneSnappingManager snap query request.
  */
 struct FSceneSnapQueryRequest
 {
+	/** Coordinate space of the request data */
+	EToolContextCoordinateSystem RequestCoordinateSpace;
+
+	/** Optionally specify the desired result coordinate space, which can prevent unnecessary transforms. Defaults to the requested RequestCoordinateSpace */
+	TOptional<EToolContextCoordinateSystem> ResultCoordinateSpace;
+
 	/** What type of snap query geometry is this */
 	ESceneSnapQueryType RequestType = ESceneSnapQueryType::Position;
+
 	/** What does caller want to try to snap to */
 	ESceneSnapQueryTargetType TargetTypes = ESceneSnapQueryTargetType::Grid;
 
+	/** The name of the custom TargetTypes, if applicable. */
+	TSet<FName> CustomTargetTypes;
+
 	/** Optional explicitly specified position grid */
 	TOptional<FVector> GridSize{};
 
 	/** Optional explicitly specified rotation grid */
 	TOptional<FRotator> RotGridSize{};
 
+	/** Optional explicitly specified scale grid */
+	TOptional<FVector> ScaleGridSize{};
+
 	/** Snap input position */
 	FVector Position = FVector::ZeroVector;
 
@@ -127,10 +245,23 @@ struct FSceneSnapQueryRequest
 	/** Snap input rotation delta */
 	FQuat DeltaRotation = FQuat(EForceInit::ForceInitToZero);
 
+	/** Snap input rotation angle */
+	double RotationAngle = 0.0;
+
+	/** Snap input axis constraints, defaults to all (unconstrained) */
+	EAxisList::Type AxisList;
+
+	/** Snap input scale */
+	FVector Scale = FVector::OneVector;
+
 	/** Optional: components to consider invisible even if they aren't. */
 	const TArray<const UPrimitiveComponent*>* ComponentsToIgnore = nullptr;
+
 	/** Optional: components to consider visible even if they aren't. */
 	const TArray<const UPrimitiveComponent*>* InvisibleComponentsToInclude = nullptr;
+
+	/** @return The ResultCoordinateSpace if set, otherwise the RequestCoordinateSpace */
+	INTERACTIVETOOLSFRAMEWORK_API EToolContextCoordinateSystem GetResultCoordinateSpace() const;
 };
 
 
@@ -139,6 +270,9 @@ struct FSceneSnapQueryRequest
  */
 struct FSceneSnapQueryResult
 {
+	/** (Optionally set) Coordinate space of the result data */
+	TOptional<EToolContextCoordinateSystem> CoordinateSpace;
+
 	/** Actor that owns snap target */
 	AActor* TargetActor = nullptr;
 	/** Component that owns snap target */
@@ -152,12 +286,18 @@ struct FSceneSnapQueryResult
 	FVector Normal;
 	/** Snap rotation delta (may not be set depending on query types) */
 	FQuat   DeltaRotation;
+	/** Snap rotation angle (may not be set depending on query types) */
+	double  RotationAngle;
+	/** Snap scale (may not be set depending on query types) */
+	FVector Scale;
 
 	/** Vertices of triangle that contains result (for debugging, may not be set) */
 	FVector TriVertices[3];
 	/** Vertex/Edge index we snapped to in triangle */
 	int TriSnapIndex;
 
+	/** @return The CoordinateSpace if set, otherwise the ResultCoordinateSpace from the provided request */
+	INTERACTIVETOOLSFRAMEWORK_API EToolContextCoordinateSystem GetCoordinateSpace(const FSceneSnapQueryRequest& InFromRequest) const;
 };
 
 
diff --git a/Engine/Source/Runtime/InteractiveToolsFramework/Public/ToolContextInterfaces.h b/Engine/Source/Runtime/InteractiveToolsFramework/Public/ToolContextInterfaces.h
--- a/Engine/Source/Runtime/InteractiveToolsFramework/Public/ToolContextInterfaces.h
+++ b/Engine/Source/Runtime/InteractiveToolsFramework/Public/ToolContextInterfaces.h
@@ -128,7 +128,9 @@ enum class EToolContextCoordinateSystem : uint8
 	/** World space coordinate system */
 	World = 0,
 	/** Local coordinate system */
-	Local = 1
+	Local = 1,
+	/** Screen space coordinate system */
+	Screen = 2
 };
 
 
