[ITF Drag Tools] Move DragToolBehaviorSource instance from Editor Viewport Client to Context

- UModeManagerInteractiveToolsContext hosting Drag Tools Behavior Source, including their rendering in DrawHUD method

- FEditorModeTools can be used to enable/disable them per editor (disabled by default)

- Level Editor enables them in SLevelEditor constructor

- Editor Viewport client now has no direct knowledge about ITF based drag tools

diff --git a/Engine/Source/Editor/Experimental/EditorInteractiveToolsFramework/Private/EdModeInteractiveToolsContext.cpp b/Engine/Source/Editor/Experimental/EditorInteractiveToolsFramework/Private/EdModeInteractiveToolsContext.cpp
--- a/Engine/Source/Editor/Experimental/EditorInteractiveToolsFramework/Private/EdModeInteractiveToolsContext.cpp
+++ b/Engine/Source/Editor/Experimental/EditorInteractiveToolsFramework/Private/EdModeInteractiveToolsContext.cpp
@@ -33,9 +33,11 @@
 #include "BaseBehaviors/ClickDragBehavior.h"
 #include "EditorModeManager.h"
 #include "EdMode.h"
+#include "EditorDragTools.h"
 
 #include "BaseGizmos/GizmoRenderingUtil.h"
 #include "UnrealClient.h"
+#include "Editor/UnrealEd/Private/Tools/DragToolsBehaviourSource.h"
 #include "UObject/ObjectSaveContext.h"
 
 //#define ENABLE_DEBUG_PRINTING
@@ -885,6 +887,11 @@ void UModeManagerInteractiveToolsContext::DrawHUD(FViewportClient* ViewportClien
 	{
 		EdModeContext->DrawHUD(ViewportClient, Viewport, View, Canvas);
 	}
+
+	if (DragToolsBehaviorSource)
+	{
+		DragToolsBehaviorSource->RenderTools(View, Canvas);
+	}
 }
 
 
@@ -1224,10 +1231,21 @@ void UModeManagerInteractiveToolsContext::Initialize(IToolsContextQueriesAPI* Qu
 	{
 		RestoreEditorState();
 	});
+
+	GetParentEditorModeManager()->OnEditorModeIDChanged().AddUObject(this, &UModeManagerInteractiveToolsContext::OnEditorModeChanged);
+
+	// Delegates used to handle SelectionMode.EnableITFTools CVAR changes
+	UE::DragTools::OnITFDragToolsActivated().AddUObject(this, &UModeManagerInteractiveToolsContext::RegisterDragTools);
+	UE::DragTools::OnITFDragToolsDeactivated().AddUObject(this, &UModeManagerInteractiveToolsContext::UnregisterDragTools);
 }
 
 void UModeManagerInteractiveToolsContext::Shutdown()
 {
+	UnregisterDragTools();
+	GetParentEditorModeManager()->OnEditorModeIDChanged().RemoveAll(this);
+	UE::DragTools::OnITFDragToolsActivated().RemoveAll(this);
+	UE::DragTools::OnITFDragToolsDeactivated().RemoveAll(this);
+
 	if (ToolManager)
 	{
 		ToolManager->OnToolPostBuild.RemoveAll(this);
@@ -1291,6 +1309,69 @@ void UModeManagerInteractiveToolsContext::SetBypassAltModifier(bool bInBypassAlt
 	bBypassAltModifier = bInBypassAltModifier;
 }
 
+void UModeManagerInteractiveToolsContext::RegisterDragTools()
+{
+	FEditorModeTools* EditorModeTools = GetParentEditorModeManager();
+	if (!EditorModeTools)
+	{
+		return;
+	}
+
+	// Ignore uninitialized Mode Tools missing a Viewport Client
+	if (!EditorModeTools->GetFocusedViewportClient())
+	{
+		return;
+	}
+
+	// This Mode Manager does not require Drag Tools
+	if (!UsesDragTools())
+	{
+		return;
+	}
+
+	// Not initialized yet
+	if (!DragToolsBehaviorSource)
+	{
+		// Alt modifier input required by Drag Tools, but currently skipped by ITF in Tool Context.
+		// Let's make sure Alt keys are processed
+		SetBypassAltModifier(false);
+
+		DragToolsBehaviorSource = NewObject<UDragToolsBehaviorSource>();
+		DragToolsBehaviorSource->Initialize(this);
+
+		if (DragToolsBehaviorSource->GetInputBehaviors())
+		{
+			InputRouter->RegisterSource(DragToolsBehaviorSource);
+		}
+	}
+}
+void UModeManagerInteractiveToolsContext::UnregisterDragTools()
+{
+	if (DragToolsBehaviorSource)
+	{
+		// Revert to previous state, with Alt Modifier not being processed
+		SetBypassAltModifier(true);
+		InputRouter->DeregisterSource(DragToolsBehaviorSource);
+		DragToolsBehaviorSource = nullptr;
+	}
+}
+void UModeManagerInteractiveToolsContext::OnEditorModeChanged(const FEditorModeID& InModeID, bool bInIsEnteringMode)
+{
+	if (!UE::DragTools::UseITFTools())
+	{
+		return;
+	}
+
+	if (bInIsEnteringMode)
+	{
+		RegisterDragTools();
+	}
+	else
+	{
+		UnregisterDragTools();
+	}
+}
+
 void UEdModeInteractiveToolsContext::InitializeContextFromModeManagerContext(UModeManagerInteractiveToolsContext* ModeManagerToolsContext)
 {
 	check(ModeManagerToolsContext != nullptr);
diff --git a/Engine/Source/Editor/Experimental/EditorInteractiveToolsFramework/Public/EdModeInteractiveToolsContext.h b/Engine/Source/Editor/Experimental/EditorInteractiveToolsFramework/Public/EdModeInteractiveToolsContext.h
--- a/Engine/Source/Editor/Experimental/EditorInteractiveToolsFramework/Public/EdModeInteractiveToolsContext.h
+++ b/Engine/Source/Editor/Experimental/EditorInteractiveToolsFramework/Public/EdModeInteractiveToolsContext.h
@@ -6,6 +6,7 @@
 #include "Containers/Map.h"
 #include "Containers/UnrealString.h"
 #include "Delegates/Delegate.h"
+#include "Editor.h"
 #include "Engine/EngineBaseTypes.h"
 #include "Engine/World.h"
 #include "HAL/Platform.h"
@@ -36,6 +37,7 @@ class ILevelEditor;
 class IToolsContextQueriesAPI;
 class IToolsContextRenderAPI;
 class IToolsContextTransactionsAPI;
+class UDragToolsBehaviorSource;
 class UEdModeInteractiveToolsContext;
 class UGizmoViewContext;
 class UInputRouter;
@@ -338,7 +340,40 @@ public:
 	 */
 	bool OnChildEdModeDeactivated(UEdModeInteractiveToolsContext* ChildToolsContext);
 
+	
+	/**
+	 * Mark Editor Modes Tools for Drag Tools Support
+	 */
+	void SetDragToolsEnabled(bool bInEnabled)
+	{
+		bUsesDragTools = bInEnabled;
+	}
+
+	/**
+	 * Are Drag Tools supported? e.g. Marquee Select
+	 * Enable/Disable support by using SetDragToolsEnabled(...)
+	 */
+	bool UsesDragTools() const
+	{
+		return bUsesDragTools;
+	};
+
 protected:
+	/**
+	 * Activating Drag Tools if current Mode Manager requires them
+	 */
+	void RegisterDragTools();
+
+	/**
+	 * Deactivating Drag Tools
+	 */
+	void UnregisterDragTools();
+
+	/**
+	 * Responding to Mode Changes in order to deactivate/activate Drag Tools accordingly
+	 */
+	void OnEditorModeChanged(const FEditorModeID& InModeID, bool bInIsEnteringMode);
+
 	//~Begin ITF Alt processing Bypass
 
 	/**
@@ -346,9 +381,6 @@ protected:
 	 * The Alt modifier bypass will be removed once a complete switch to ITF happens
 	 */
 
-	// Needs to use Alt modifier for Drag Tools to work
-	friend class FEditorViewportClient;
-
 	/** Enable/Disable Alt modifier bypass */
 	void SetBypassAltModifier(bool bInBypassAltModifier);
 
@@ -359,6 +391,14 @@ protected:
 	bool bBypassAltModifier = true;
 
 	//~End ITF Alt processing Bypass
+
+	/**
+	 *
+	 */
+	UPROPERTY(Transient)
+	TObjectPtr<UDragToolsBehaviorSource> DragToolsBehaviorSource;
+
+	bool bUsesDragTools = false;
 };
 
 
diff --git a/Engine/Source/Editor/LevelEditor/LevelEditor.Build.cs b/Engine/Source/Editor/LevelEditor/LevelEditor.Build.cs
--- a/Engine/Source/Editor/LevelEditor/LevelEditor.Build.cs
+++ b/Engine/Source/Editor/LevelEditor/LevelEditor.Build.cs
@@ -13,6 +13,7 @@ public class LevelEditor : ModuleRules
 			new string[] {
 				"AssetTools",
 				"ClassViewer",
+				"EditorInteractiveToolsFramework",
 				"MainFrame",
                 "PlacementMode",
 				"SlateReflector",
diff --git a/Engine/Source/Editor/LevelEditor/Private/SLevelEditor.cpp b/Engine/Source/Editor/LevelEditor/Private/SLevelEditor.cpp
--- a/Engine/Source/Editor/LevelEditor/Private/SLevelEditor.cpp
+++ b/Engine/Source/Editor/LevelEditor/Private/SLevelEditor.cpp
@@ -79,6 +79,7 @@
 #include "SLevelViewport.h"
 #include "LevelEditorOutlinerSettings.h"
 #include "IWorldHierarchy.h"
+#include "EdModeInteractiveToolsContext.h"
 
 #define LOCTEXT_NAMESPACE "SLevelEditor"
 
@@ -94,6 +95,11 @@ SLevelEditor::SLevelEditor()
 	// Ensure the Actor main menu has a default label before the first selection change.
 	, CachedViewportContextMenuTitle(FLevelEditorContextMenu::GetContextMenuTitle(ELevelEditorMenuContext::MainMenu, nullptr))
 {
+	if (UModeManagerInteractiveToolsContext* const ModeManagerInteractiveToolsContext =
+			SLevelEditor::GetEditorModeManager().GetInteractiveToolsContext())
+	{
+		ModeManagerInteractiveToolsContext->SetDragToolsEnabled(true);
+	}
 }
 
 void SLevelEditor::BindCommands()
diff --git a/Engine/Source/Editor/UnrealEd/Private/EditorDragTools.cpp b/Engine/Source/Editor/UnrealEd/Private/EditorDragTools.cpp
--- a/Engine/Source/Editor/UnrealEd/Private/EditorDragTools.cpp
+++ b/Engine/Source/Editor/UnrealEd/Private/EditorDragTools.cpp
@@ -11,22 +11,6 @@
 
 namespace UE::DragTools
 {
-void UnregisterDragTools(FEditorViewportClient* InViewportClient)
-{
-	if (InViewportClient)
-	{
-		InViewportClient->UnregisterDragTools();
-	}
-}
-
-void RegisterDragTools(FEditorViewportClient* InViewportClient)
-{
-	if (InViewportClient)
-	{
-		InViewportClient->RegisterDragTools();
-	}
-}
-
 // CVar initializer
 static TAutoConsoleVariable<int32> CVarEnableITFTools(
 	TEXT("SelectionMode.EnableITFTools"),
@@ -35,16 +19,13 @@ static TAutoConsoleVariable<int32> CVarEnableITFTools(
 	FConsoleVariableDelegate::CreateLambda(
 		[](const IConsoleVariable* InVariable)
 		{
-			if (FEditorViewportClient* ViewportClient = GLevelEditorModeTools().GetFocusedViewportClient())
+			if (InVariable->GetBool())
+			{
+				OnITFDragToolsActivated().Broadcast();
+			}
+			else
 			{
-				if (InVariable->GetBool())
-				{
-					RegisterDragTools(ViewportClient);
-				}
-				else
-				{
-					UnregisterDragTools(ViewportClient);
-				}
+				OnITFDragToolsDeactivated().Broadcast();
 			}
 		}
 	)
@@ -54,8 +35,21 @@ bool UseITFTools()
 {
 	return CVarEnableITFTools->GetBool();
 }
+
+FOnDragToolsToggleDelegate& OnITFDragToolsActivated()
+{
+	static FOnDragToolsToggleDelegate OnDragToolsActivated;
+	return OnDragToolsActivated;
 }
 
+FOnDragToolsToggleDelegate& OnITFDragToolsDeactivated()
+{
+	static FOnDragToolsToggleDelegate OnDragToolsDeactivated;
+	return OnDragToolsDeactivated;
+}
+
+} // namespace UE::DragTools
+
 ///////////////////////////////////////////////////////////////////////////////////////////////////////////////////////
 //
 // FDragTool
diff --git a/Engine/Source/Editor/UnrealEd/Private/EditorViewportClient.cpp b/Engine/Source/Editor/UnrealEd/Private/EditorViewportClient.cpp
--- a/Engine/Source/Editor/UnrealEd/Private/EditorViewportClient.cpp
+++ b/Engine/Source/Editor/UnrealEd/Private/EditorViewportClient.cpp
@@ -80,7 +80,6 @@
 #include "Settings/EditorStyleSettings.h"
 #include "GameFramework/ActorPrimitiveColorHandler.h"
 #include "LandscapeRender.h"
-#include "Tools/DragToolsBehaviourSource.h"
 
 #define LOCTEXT_NAMESPACE "EditorViewportClient"
 
@@ -544,8 +543,6 @@ FEditorViewportClient::FEditorViewportClient(FEditorModeTools* InModeTools, FPre
 
 	FSlateApplication::Get().OnWindowDPIScaleChanged().AddRaw(this, &FEditorViewportClient::HandleWindowDPIScaleChanged);
 
-	RegisterDragTools();
-
 	double MaxViewLocation = HALF_WORLD_MAX;
 	if(GetScene() != nullptr && GetScene()->GetWorld() != nullptr)
 	{
@@ -596,8 +593,6 @@ FEditorViewportClient::~FEditorViewportClient()
 		FSlateApplication::Get().OnWindowDPIScaleChanged().RemoveAll(this);
 	}
 
-	UnregisterDragTools();
-
 	ModeTools.Reset();
 }
 
@@ -3980,11 +3975,6 @@ void FEditorViewportClient::AddReferencedObjects( FReferenceCollector& Collector
 			StereoViewState.GetReference()->AddReferencedObjects(Collector);
 		}
 	}
-
-	if (DragToolsBehaviorSource) 
-	{
-		Collector.AddReferencedObject(DragToolsBehaviorSource);
-	}
 }
 
 FString FEditorViewportClient::GetReferencerName() const
@@ -4625,14 +4615,6 @@ void FEditorViewportClient::DrawPreviewLightVisualization(const FSceneView* View
 
 void FEditorViewportClient::RenderDragTool(const FSceneView* View, FCanvas* Canvas)
 {
-	if (DragToolsBehaviorSource)
-	{
-		if (FDragTool* ActiveDragTool = DragToolsBehaviorSource->GetActiveTool())
-		{
-			ActiveDragTool->Render(View, Canvas);
-		}
-	}
-
 	MouseDeltaTracker->RenderDragTool(View, Canvas);
 }
 
@@ -6683,50 +6665,6 @@ float FEditorViewportClient::GetMinimumOrthoZoom() const
 	return FMath::Max(GetDefault<ULevelEditorViewportSettings>()->MinimumOrthographicZoom, MIN_ORTHOZOOM);
 }
 
-void FEditorViewportClient::RegisterDragTools()
-{
-	if (UE::DragTools::UseITFTools())
-	{
-		if (!DragToolsBehaviorSource)
-		{
-			if (const FEditorModeTools* EditorModeTools = GetModeTools())
-			{
-				if (UModeManagerInteractiveToolsContext* InteractiveToolContext = EditorModeTools->GetInteractiveToolsContext())
-				{
-					// Alt modifier input required by Drag Tools, but currently skipped by ITF in Tool Context.
-					// Let's make sure Alt keys are processed
-					InteractiveToolContext->SetBypassAltModifier(false);
-
-					DragToolsBehaviorSource = NewObject<UDragToolsBehaviorSource>();
-					DragToolsBehaviorSource->Initialize(InteractiveToolContext, this);
-
-					if (ensure(DragToolsBehaviorSource->GetInputBehaviors()))
-					{
-						InteractiveToolContext->InputRouter->RegisterSource(DragToolsBehaviorSource);
-					}
-				}
-			}
-		}
-	}
-}
-
-void FEditorViewportClient::UnregisterDragTools()
-{
-	if (DragToolsBehaviorSource)
-	{
-		if (const FEditorModeTools* EditorModeTools = GetModeTools())
-		{
-			if (UModeManagerInteractiveToolsContext* InteractiveToolContext = EditorModeTools->GetInteractiveToolsContext())
-			{
-				// Revert to previous state, with Alt Modifier not being processed
-				InteractiveToolContext->SetBypassAltModifier(true);
-				InteractiveToolContext->InputRouter->DeregisterSource(DragToolsBehaviorSource);
-				DragToolsBehaviorSource = nullptr;
-			}
-		}
-	}
-}
-
 ////////////////
 
 bool FEditorViewportStats::bInitialized(false);
diff --git a/Engine/Source/Editor/UnrealEd/Private/Tools/DragToolsBehaviourSource.cpp b/Engine/Source/Editor/UnrealEd/Private/Tools/DragToolsBehaviourSource.cpp
--- a/Engine/Source/Editor/UnrealEd/Private/Tools/DragToolsBehaviourSource.cpp
+++ b/Engine/Source/Editor/UnrealEd/Private/Tools/DragToolsBehaviourSource.cpp
@@ -5,8 +5,9 @@
 #include "BaseBehaviors/KeyAsModifierInputBehavior.h"
 #include "DragTool_FrustumSelect.h"
 #include "EdModeInteractiveToolsContext.h"
+#include "EditorModeManager.h"
 
-void UDragToolsBehaviorSource::Initialize(const UEditorInteractiveToolsContext* InInteractiveToolsContext, FEditorViewportClient* InViewportClient)
+void UDragToolsBehaviorSource::Initialize(const UEditorInteractiveToolsContext* InInteractiveToolsContext)
 {
 	if (!InInteractiveToolsContext)
 	{
@@ -25,7 +26,13 @@ void UDragToolsBehaviorSource::Initialize(const UEditorInteractiveToolsContext*
 		return;
 	}
 
-	FEditorViewportClient* EditorViewportClient = InViewportClient;
+	FEditorViewportClient* EditorViewportClient = nullptr;
+
+	if (const FEditorModeTools* const ModeManager = InInteractiveToolsContext->GetParentEditorModeManager())
+	{
+		EditorViewportClient = ModeManager->GetFocusedViewportClient();
+	}
+
 	if (!EditorViewportClient)
 	{
 		FViewport* Viewport = ContextQueriesAPI->GetFocusedViewport();
@@ -138,3 +145,11 @@ void UDragToolsBehaviorSource::OnDeactivateTool(EDragTool::Type InDragToolType)
 		ActiveToolType = -1;
 	}
 }
+
+void UDragToolsBehaviorSource::RenderTools(const FSceneView* InSceneView, FCanvas* InCanvas) const
+{
+	if (FDragTool* ActiveDragTool = GetActiveTool())
+	{
+		ActiveDragTool->Render(InSceneView, InCanvas);
+	}
+}
diff --git a/Engine/Source/Editor/UnrealEd/Private/Tools/DragToolsBehaviourSource.h b/Engine/Source/Editor/UnrealEd/Private/Tools/DragToolsBehaviourSource.h
--- a/Engine/Source/Editor/UnrealEd/Private/Tools/DragToolsBehaviourSource.h
+++ b/Engine/Source/Editor/UnrealEd/Private/Tools/DragToolsBehaviourSource.h
@@ -2,13 +2,12 @@
 
 #pragma once
 
+#include "DragTool_FrustumSelect.h"
 #include "InputBehaviorSet.h"
 #include "UObject/Object.h"
-#include "DragTool_FrustumSelect.h"
 #include "DragToolsBehaviourSource.generated.h"
 
 class UEditorInteractiveToolsContext;
-class FEditorViewportClient;
 
 /**
  * Hosts Drag Tools and needed behaviors.
@@ -29,10 +28,14 @@ public:
 	 * Instantiates the BehaviorSet hosting all behaviors required by drag tools.
 	 *
 	 * @param InInteractiveToolsContext: used to retrieve the viewport client used by drag tools. Can be gathered e.g.
-	 * @param InViewportClient: the viewport client used by drag tools. (the focussed viewport will be retrieved from the context if set to null).
 	 * from UEdMode::GetInteractiveToolsContext.
 	 */
-	void Initialize(const UEditorInteractiveToolsContext* InInteractiveToolsContext, FEditorViewportClient* InViewportClient = nullptr);
+	UNREALED_API void Initialize(const UEditorInteractiveToolsContext* InInteractiveToolsContext);
+
+	/**
+	 * Renders the active tool on the specified View/Canvas
+	 */
+	UNREALED_API void RenderTools(const FSceneView* InSceneView, FCanvas* InCanvas) const;
 
 	/**
 	 * Returns the drag tool currently being used, if any
diff --git a/Engine/Source/Editor/UnrealEd/Public/EditorDragTools.h b/Engine/Source/Editor/UnrealEd/Public/EditorDragTools.h
--- a/Engine/Source/Editor/UnrealEd/Public/EditorDragTools.h
+++ b/Engine/Source/Editor/UnrealEd/Public/EditorDragTools.h
@@ -19,17 +19,16 @@ class FUICommandInfo;
  */
 namespace UE::DragTools
 {
+
+DECLARE_MULTICAST_DELEGATE(FOnDragToolsToggleDelegate)
 /**
  * Returns true if ITF-based drag tools should be used.
  * If an ITF version of a drag tool is not available yet, legacy will be used.
  */
 UNREALED_API bool UseITFTools();
 
-/** Register/Activate available ITF-based drag tools for default ed mode ("Selection Mode") */
-void RegisterDragTools(FEditorViewportClient* InViewportClient);
-
-/** Deregisters/Deactivates ITF-based drag tools for default ed mode ("Selection Mode") */
-void UnregisterDragTools(FEditorViewportClient* InViewportClient);
+UNREALED_API FOnDragToolsToggleDelegate& OnITFDragToolsActivated();
+UNREALED_API FOnDragToolsToggleDelegate& OnITFDragToolsDeactivated();
 
 } // namespace UE::DragTools
 
diff --git a/Engine/Source/Editor/UnrealEd/Public/EditorViewportClient.h b/Engine/Source/Editor/UnrealEd/Public/EditorViewportClient.h
--- a/Engine/Source/Editor/UnrealEd/Public/EditorViewportClient.h
+++ b/Engine/Source/Editor/UnrealEd/Public/EditorViewportClient.h
@@ -36,7 +36,6 @@ struct FTypedElementHandle;
 class IAssetFactoryInterface;
 class SEditorViewport;
 class UActorFactory;
-class UDragToolsBehaviorSource;
 class UTypedElementViewportInteraction;
 enum class EViewStatusForScreenPercentage;
 struct FGizmoState;
@@ -2147,23 +2146,6 @@ private:
 	 * true during FEditorViewportClient::StopTracking. Used to guard against reentry.
 	 */
 	bool bIsTrackingBeingStopped;
-
-public:
-	/**
-	 * Initializes and registers ITF-based drag tools
-	 */
-	UNREALED_API void RegisterDragTools();
-
-	/**
-	 * Unregistering of ITF drag tools
-	 */
-	UNREALED_API void UnregisterDragTools();
-
-protected:
-	/**
-	 * Hosting and accessing Drag Tools
-	 */
-	TObjectPtr<UDragToolsBehaviorSource> DragToolsBehaviorSource;
 };
 
 
