Add and address deprecation warning after UEditorInteractiveToolsContext classes move to UnrealEd

See CL 40464919.

diff --git a/Engine/Plugins/Animation/ControlRig/Source/ControlRigEditor/Private/EditMode/ControlRigEditMode.cpp b/Engine/Plugins/Animation/ControlRig/Source/ControlRigEditor/Private/EditMode/ControlRigEditMode.cpp
--- a/Engine/Plugins/Animation/ControlRig/Source/ControlRigEditor/Private/EditMode/ControlRigEditMode.cpp
+++ b/Engine/Plugins/Animation/ControlRig/Source/ControlRigEditor/Private/EditMode/ControlRigEditMode.cpp
@@ -57,7 +57,7 @@
 #include "LevelSequence.h"
 #include "LevelEditor.h"
 #include "InteractiveToolManager.h"
-#include "EdModeInteractiveToolsContext.h"
+#include "Tools/EdModeInteractiveToolsContext.h"
 #include "Constraints/MovieSceneConstraintChannelHelper.h"
 #include "Editor/EditorPerProjectUserSettings.h"
 #include "Transform/TransformConstraint.h"
diff --git a/Engine/Plugins/Animation/ControlRig/Source/ControlRigEditor/Private/EditMode/SControlRigEditModeTools.cpp b/Engine/Plugins/Animation/ControlRig/Source/ControlRigEditor/Private/EditMode/SControlRigEditModeTools.cpp
--- a/Engine/Plugins/Animation/ControlRig/Source/ControlRigEditor/Private/EditMode/SControlRigEditModeTools.cpp
+++ b/Engine/Plugins/Animation/ControlRig/Source/ControlRigEditor/Private/EditMode/SControlRigEditModeTools.cpp
@@ -30,7 +30,7 @@
 #include "LevelEditor.h"
 #include "EditorModeManager.h"
 #include "InteractiveToolManager.h"
-#include "EdModeInteractiveToolsContext.h"
+#include "Tools/EdModeInteractiveToolsContext.h"
 #include "ControlRigSpaceChannelEditors.h"
 #include "IKeyArea.h"
 #include "Widgets/Notifications/SNotificationList.h"
diff --git a/Engine/Plugins/Animation/ControlRig/Source/ControlRigEditor/Private/Editor/ControlRigViewportToolbarExtensions.cpp b/Engine/Plugins/Animation/ControlRig/Source/ControlRigEditor/Private/Editor/ControlRigViewportToolbarExtensions.cpp
--- a/Engine/Plugins/Animation/ControlRig/Source/ControlRigEditor/Private/Editor/ControlRigViewportToolbarExtensions.cpp
+++ b/Engine/Plugins/Animation/ControlRig/Source/ControlRigEditor/Private/Editor/ControlRigViewportToolbarExtensions.cpp
@@ -24,7 +24,7 @@
 #include "LevelEditor.h"
 #include "InteractiveToolManager.h"
 #include "EditorModeManager.h"
-#include "EdModeInteractiveToolsContext.h"
+#include "Tools/EdModeInteractiveToolsContext.h"
 
 #define LOCTEXT_NAMESPACE "ControlRigViewportToolbar"
 
diff --git a/Engine/Plugins/Animation/MLDeformer/MLDeformerFramework/Source/MLDeformerFrameworkEditor/Private/MLDeformerPaintMode.cpp b/Engine/Plugins/Animation/MLDeformer/MLDeformerFramework/Source/MLDeformerFrameworkEditor/Private/MLDeformerPaintMode.cpp
--- a/Engine/Plugins/Animation/MLDeformer/MLDeformerFramework/Source/MLDeformerFrameworkEditor/Private/MLDeformerPaintMode.cpp
+++ b/Engine/Plugins/Animation/MLDeformer/MLDeformerFramework/Source/MLDeformerFrameworkEditor/Private/MLDeformerPaintMode.cpp
@@ -4,7 +4,7 @@
 #include "MLDeformerPaintModeToolkit.h"
 #include "MeshAttributePaintTool.h"
 #include "EdMode.h"
-#include "EdModeInteractiveToolsContext.h"
+#include "Tools/EdModeInteractiveToolsContext.h"
 #include "EditorModeManager.h"
 #include "ToolTargets/SkeletalMeshComponentToolTarget.h"
 #include "ModelingToolsManagerActions.h"
diff --git a/Engine/Plugins/Animation/MLDeformer/MLDeformerFramework/Source/MLDeformerFrameworkEditor/Private/MLDeformerPaintModeToolkit.cpp b/Engine/Plugins/Animation/MLDeformer/MLDeformerFramework/Source/MLDeformerFrameworkEditor/Private/MLDeformerPaintModeToolkit.cpp
--- a/Engine/Plugins/Animation/MLDeformer/MLDeformerFramework/Source/MLDeformerFrameworkEditor/Private/MLDeformerPaintModeToolkit.cpp
+++ b/Engine/Plugins/Animation/MLDeformer/MLDeformerFramework/Source/MLDeformerFrameworkEditor/Private/MLDeformerPaintModeToolkit.cpp
@@ -2,7 +2,7 @@
 
 #include "MLDeformerPaintModeToolkit.h"
 #include "MLDeformerEditorToolkit.h"
-#include "EdModeInteractiveToolsContext.h"
+#include "Tools/EdModeInteractiveToolsContext.h"
 #include "Framework/MultiBox/MultiBoxBuilder.h"
 #include "IDetailsView.h"
 #include "InteractiveTool.h"
diff --git a/Engine/Plugins/Animation/SkeletalMeshModelingTools/Source/SkeletalMeshModelingTools/Private/SkeletalMeshModelingToolsEditorMode.cpp b/Engine/Plugins/Animation/SkeletalMeshModelingTools/Source/SkeletalMeshModelingTools/Private/SkeletalMeshModelingToolsEditorMode.cpp
--- a/Engine/Plugins/Animation/SkeletalMeshModelingTools/Source/SkeletalMeshModelingTools/Private/SkeletalMeshModelingToolsEditorMode.cpp
+++ b/Engine/Plugins/Animation/SkeletalMeshModelingTools/Source/SkeletalMeshModelingTools/Private/SkeletalMeshModelingToolsEditorMode.cpp
@@ -9,7 +9,7 @@
 
 #include "BaseGizmos/TransformGizmoUtil.h"
 #include "EdMode.h"
-#include "EdModeInteractiveToolsContext.h"
+#include "Tools/EdModeInteractiveToolsContext.h"
 #include "EditorModeManager.h"
 #include "EditorViewportClient.h"
 #include "Features/IModularFeatures.h"
diff --git a/Engine/Plugins/Animation/SkeletalMeshModelingTools/Source/SkeletalMeshModelingTools/Private/SkeletalMeshModelingToolsEditorModeToolkit.cpp b/Engine/Plugins/Animation/SkeletalMeshModelingTools/Source/SkeletalMeshModelingTools/Private/SkeletalMeshModelingToolsEditorModeToolkit.cpp
--- a/Engine/Plugins/Animation/SkeletalMeshModelingTools/Source/SkeletalMeshModelingTools/Private/SkeletalMeshModelingToolsEditorModeToolkit.cpp
+++ b/Engine/Plugins/Animation/SkeletalMeshModelingTools/Source/SkeletalMeshModelingTools/Private/SkeletalMeshModelingToolsEditorModeToolkit.cpp
@@ -4,7 +4,7 @@
 
 #include "SkeletalMeshModelingToolsCommands.h"
 
-#include "EdModeInteractiveToolsContext.h"
+#include "Tools/EdModeInteractiveToolsContext.h"
 #include "Features/IModularFeatures.h"
 #include "Framework/MultiBox/MultiBoxBuilder.h"
 #include "IDetailsView.h"
diff --git a/Engine/Plugins/ChaosClothAssetEditor/Source/ChaosClothAssetEditor/Private/ChaosClothAsset/ClothEditor3DViewportClient.cpp b/Engine/Plugins/ChaosClothAssetEditor/Source/ChaosClothAssetEditor/Private/ChaosClothAsset/ClothEditor3DViewportClient.cpp
--- a/Engine/Plugins/ChaosClothAssetEditor/Source/ChaosClothAssetEditor/Private/ChaosClothAsset/ClothEditor3DViewportClient.cpp
+++ b/Engine/Plugins/ChaosClothAssetEditor/Source/ChaosClothAssetEditor/Private/ChaosClothAsset/ClothEditor3DViewportClient.cpp
@@ -10,7 +10,7 @@
 #include "Dataflow/DataflowObject.h"
 #include "Dataflow/DataflowActor.h"
 #include "Dataflow/DataflowRenderingFactory.h"
-#include "EdModeInteractiveToolsContext.h"
+#include "Tools/EdModeInteractiveToolsContext.h"
 #include "BaseGizmos/TransformProxy.h"
 #include "BaseGizmos/CombinedTransformGizmo.h"
 #include "BaseGizmos/TransformGizmoUtil.h"
diff --git a/Engine/Plugins/ChaosClothAssetEditor/Source/ChaosClothAssetEditor/Private/ChaosClothAsset/ClothEditorMode.cpp b/Engine/Plugins/ChaosClothAssetEditor/Source/ChaosClothAssetEditor/Private/ChaosClothAsset/ClothEditorMode.cpp
--- a/Engine/Plugins/ChaosClothAssetEditor/Source/ChaosClothAssetEditor/Private/ChaosClothAsset/ClothEditorMode.cpp
+++ b/Engine/Plugins/ChaosClothAssetEditor/Source/ChaosClothAssetEditor/Private/ChaosClothAsset/ClothEditorMode.cpp
@@ -16,7 +16,7 @@
 #include "AssetEditorModeManager.h"
 #include "Drawing/MeshElementsVisualizer.h"
 #include "EditorViewportClient.h"
-#include "EdModeInteractiveToolsContext.h"
+#include "Tools/EdModeInteractiveToolsContext.h"
 #include "EngineAnalytics.h"
 #include "Framework/Commands/UICommandList.h"
 #include "InteractiveTool.h"
diff --git a/Engine/Plugins/ChaosClothAssetEditor/Source/ChaosClothAssetEditor/Private/ChaosClothAsset/ClothEditorModeToolkit.cpp b/Engine/Plugins/ChaosClothAssetEditor/Source/ChaosClothAssetEditor/Private/ChaosClothAsset/ClothEditorModeToolkit.cpp
--- a/Engine/Plugins/ChaosClothAssetEditor/Source/ChaosClothAssetEditor/Private/ChaosClothAsset/ClothEditorModeToolkit.cpp
+++ b/Engine/Plugins/ChaosClothAssetEditor/Source/ChaosClothAssetEditor/Private/ChaosClothAsset/ClothEditorModeToolkit.cpp
@@ -14,7 +14,7 @@
 #include "Framework/MultiBox/MultiBoxBuilder.h"
 #include "InteractiveToolManager.h"
 #include "SEditorViewport.h"
-#include "EdModeInteractiveToolsContext.h"
+#include "Tools/EdModeInteractiveToolsContext.h"
 #include "SBaseCharacterFXEditorViewport.h"
 #include "Widgets/Input/SButton.h"
 #include "SPrimaryButton.h"
diff --git a/Engine/Plugins/ChaosClothAssetEditor/Source/ChaosClothAssetEditor/Private/ChaosClothAsset/ClothEditorRestSpaceViewportClient.cpp b/Engine/Plugins/ChaosClothAssetEditor/Source/ChaosClothAssetEditor/Private/ChaosClothAsset/ClothEditorRestSpaceViewportClient.cpp
--- a/Engine/Plugins/ChaosClothAssetEditor/Source/ChaosClothAssetEditor/Private/ChaosClothAsset/ClothEditorRestSpaceViewportClient.cpp
+++ b/Engine/Plugins/ChaosClothAssetEditor/Source/ChaosClothAssetEditor/Private/ChaosClothAsset/ClothEditorRestSpaceViewportClient.cpp
@@ -9,7 +9,7 @@
 #include "EditorModeManager.h"
 #include "Engine/Selection.h"
 #include "Behaviors/2DViewportBehaviorTargets.h"
-#include "EdModeInteractiveToolsContext.h"
+#include "Tools/EdModeInteractiveToolsContext.h"
 #include "CameraController.h"
 #include "Components/DynamicMeshComponent.h"
 #include "Framework/Application/SlateApplication.h"
diff --git a/Engine/Plugins/ChaosClothAssetEditor/Source/ChaosClothAssetEditor/Private/ChaosClothAsset/ClothEditorToolkit.cpp b/Engine/Plugins/ChaosClothAssetEditor/Source/ChaosClothAssetEditor/Private/ChaosClothAsset/ClothEditorToolkit.cpp
--- a/Engine/Plugins/ChaosClothAssetEditor/Source/ChaosClothAssetEditor/Private/ChaosClothAsset/ClothEditorToolkit.cpp
+++ b/Engine/Plugins/ChaosClothAssetEditor/Source/ChaosClothAssetEditor/Private/ChaosClothAsset/ClothEditorToolkit.cpp
@@ -18,7 +18,7 @@
 #include "ChaosClothAsset/ClothSimulationNodeDetailExtender.h"
 #include "ChaosClothAsset/AddWeightMapNode.h"
 #include "ChaosClothAsset/TransferSkinWeightsNode.h"
-#include "EdModeInteractiveToolsContext.h"
+#include "Tools/EdModeInteractiveToolsContext.h"
 #include "Framework/Docking/LayoutExtender.h"
 #include "AssetEditorModeManager.h"
 #include "AdvancedPreviewScene.h"
diff --git a/Engine/Plugins/Editor/ModelingToolsEditorMode/Source/ModelingToolsEditorMode/Private/ModelingModeToolkit_Toolbars.cpp b/Engine/Plugins/Editor/ModelingToolsEditorMode/Source/ModelingToolsEditorMode/Private/ModelingModeToolkit_Toolbars.cpp
--- a/Engine/Plugins/Editor/ModelingToolsEditorMode/Source/ModelingToolsEditorMode/Private/ModelingModeToolkit_Toolbars.cpp
+++ b/Engine/Plugins/Editor/ModelingToolsEditorMode/Source/ModelingToolsEditorMode/Private/ModelingModeToolkit_Toolbars.cpp
@@ -1,7 +1,7 @@
 // Copyright Epic Games, Inc. All Rights Reserved.
 
 #include "ModelingToolsEditorModeToolkit.h"
-#include "EdModeInteractiveToolsContext.h"
+#include "Tools/EdModeInteractiveToolsContext.h"
 #include "ModelingToolsEditorMode.h"
 #include "Framework/Commands/UICommandList.h"
 #include "ModelingToolsManagerActions.h"
diff --git a/Engine/Plugins/Editor/ModelingToolsEditorMode/Source/ModelingToolsEditorMode/Private/ModelingToolsEditorMode.cpp b/Engine/Plugins/Editor/ModelingToolsEditorMode/Source/ModelingToolsEditorMode/Private/ModelingToolsEditorMode.cpp
--- a/Engine/Plugins/Editor/ModelingToolsEditorMode/Source/ModelingToolsEditorMode/Private/ModelingToolsEditorMode.cpp
+++ b/Engine/Plugins/Editor/ModelingToolsEditorMode/Source/ModelingToolsEditorMode/Private/ModelingToolsEditorMode.cpp
@@ -4,7 +4,7 @@
 #include "Components/StaticMeshComponent.h"
 #include "IGeometryProcessingInterfacesModule.h"
 #include "Editor/UnrealEdEngine.h"
-#include "EdModeInteractiveToolsContext.h"
+#include "Tools/EdModeInteractiveToolsContext.h"
 #include "GeometryProcessingInterfaces/IUVEditorModularFeature.h"
 #include "IAnalyticsProviderET.h"
 #include "ModelingToolsEditorModeToolkit.h"
diff --git a/Engine/Plugins/Editor/ModelingToolsEditorMode/Source/ModelingToolsEditorMode/Private/ModelingToolsEditorModeToolkit.cpp b/Engine/Plugins/Editor/ModelingToolsEditorMode/Source/ModelingToolsEditorMode/Private/ModelingToolsEditorModeToolkit.cpp
--- a/Engine/Plugins/Editor/ModelingToolsEditorMode/Source/ModelingToolsEditorMode/Private/ModelingToolsEditorModeToolkit.cpp
+++ b/Engine/Plugins/Editor/ModelingToolsEditorMode/Source/ModelingToolsEditorMode/Private/ModelingToolsEditorModeToolkit.cpp
@@ -5,7 +5,7 @@
 #include "EditorModeManager.h"
 #include "IGeometryProcessingInterfacesModule.h"
 #include "GeometryProcessingInterfaces/IUVEditorModularFeature.h"
-#include "EdModeInteractiveToolsContext.h"
+#include "Tools/EdModeInteractiveToolsContext.h"
 #include "ModelingToolsManagerActions.h"
 #include "InteractiveToolManager.h"
 #include "ModelingToolsEditorModeSettings.h"
diff --git a/Engine/Plugins/Editor/ScriptableToolsEditorMode/Source/ScriptableToolsEditorMode/Private/ScriptableToolsEditorMode.cpp b/Engine/Plugins/Editor/ScriptableToolsEditorMode/Source/ScriptableToolsEditorMode/Private/ScriptableToolsEditorMode.cpp
--- a/Engine/Plugins/Editor/ScriptableToolsEditorMode/Source/ScriptableToolsEditorMode/Private/ScriptableToolsEditorMode.cpp
+++ b/Engine/Plugins/Editor/ScriptableToolsEditorMode/Source/ScriptableToolsEditorMode/Private/ScriptableToolsEditorMode.cpp
@@ -1,7 +1,7 @@
 // Copyright Epic Games, Inc. All Rights Reserved.
 
 #include "ScriptableToolsEditorMode.h"
-#include "EdModeInteractiveToolsContext.h"
+#include "Tools/EdModeInteractiveToolsContext.h"
 #include "Modules/ModuleManager.h"
 #include "ILevelEditor.h"
 #include "LevelEditor.h"
diff --git a/Engine/Plugins/Editor/ScriptableToolsEditorMode/Source/ScriptableToolsEditorMode/Private/ScriptableToolsEditorModeToolkit.cpp b/Engine/Plugins/Editor/ScriptableToolsEditorMode/Source/ScriptableToolsEditorMode/Private/ScriptableToolsEditorModeToolkit.cpp
--- a/Engine/Plugins/Editor/ScriptableToolsEditorMode/Source/ScriptableToolsEditorMode/Private/ScriptableToolsEditorModeToolkit.cpp
+++ b/Engine/Plugins/Editor/ScriptableToolsEditorMode/Source/ScriptableToolsEditorMode/Private/ScriptableToolsEditorModeToolkit.cpp
@@ -1,7 +1,7 @@
 // Copyright Epic Games, Inc. All Rights Reserved.
 
 #include "ScriptableToolsEditorModeToolkit.h"
-#include "EdModeInteractiveToolsContext.h"
+#include "Tools/EdModeInteractiveToolsContext.h"
 #include "ScriptableToolsEditorMode.h"
 #include "ScriptableToolsEditorModeManagerCommands.h"
 #include "ScriptableToolsEditorModeStyle.h"
diff --git a/Engine/Plugins/Editor/ScriptableToolsEditorMode/Source/ScriptableToolsEditorMode/Public/ScriptableToolsEditorMode.h b/Engine/Plugins/Editor/ScriptableToolsEditorMode/Source/ScriptableToolsEditorMode/Public/ScriptableToolsEditorMode.h
--- a/Engine/Plugins/Editor/ScriptableToolsEditorMode/Source/ScriptableToolsEditorMode/Public/ScriptableToolsEditorMode.h
+++ b/Engine/Plugins/Editor/ScriptableToolsEditorMode/Source/ScriptableToolsEditorMode/Public/ScriptableToolsEditorMode.h
@@ -7,7 +7,7 @@
 
 #include "InputState.h"
 #include "InteractiveToolManager.h"
-#include "EdModeInteractiveToolsContext.h"
+#include "Tools/EdModeInteractiveToolsContext.h"
 
 #include "ScriptableToolsEditorMode.generated.h"
 
diff --git a/Engine/Plugins/Editor/SequencerAnimTools/Source/SequencerAnimTools/Private/SequencerAnimEditPivotTool.cpp b/Engine/Plugins/Editor/SequencerAnimTools/Source/SequencerAnimTools/Private/SequencerAnimEditPivotTool.cpp
--- a/Engine/Plugins/Editor/SequencerAnimTools/Source/SequencerAnimTools/Private/SequencerAnimEditPivotTool.cpp
+++ b/Engine/Plugins/Editor/SequencerAnimTools/Source/SequencerAnimTools/Private/SequencerAnimEditPivotTool.cpp
@@ -24,7 +24,7 @@
 #include "LevelSequenceEditorBlueprintLibrary.h"
 #include "ILevelSequenceEditorToolkit.h"
 #include "InteractiveToolManager.h"
-#include "EdModeInteractiveToolsContext.h"
+#include "Tools/EdModeInteractiveToolsContext.h"
 #include "SceneView.h"
 #include "UnrealEdGlobals.h"
 #include "Editor/UnrealEdEngine.h"
diff --git a/Engine/Plugins/Editor/SequencerAnimTools/Source/SequencerAnimTools/Private/SequencerAnimToolsModule.cpp b/Engine/Plugins/Editor/SequencerAnimTools/Source/SequencerAnimTools/Private/SequencerAnimToolsModule.cpp
--- a/Engine/Plugins/Editor/SequencerAnimTools/Source/SequencerAnimTools/Private/SequencerAnimToolsModule.cpp
+++ b/Engine/Plugins/Editor/SequencerAnimTools/Source/SequencerAnimTools/Private/SequencerAnimToolsModule.cpp
@@ -4,7 +4,7 @@
 #include "EditorModeManager.h"
 #include "ILevelEditor.h"
 #include "LevelEditor.h"
-#include "EdModeInteractiveToolsContext.h"
+#include "Tools/EdModeInteractiveToolsContext.h"
 #include "InteractiveToolManager.h"
 #include "LevelEditorSequencerIntegration.h"
 
diff --git a/Engine/Plugins/Editor/SequencerAnimTools/Source/SequencerAnimTools/Private/SequencerToolsEditMode.cpp b/Engine/Plugins/Editor/SequencerAnimTools/Source/SequencerAnimTools/Private/SequencerToolsEditMode.cpp
--- a/Engine/Plugins/Editor/SequencerAnimTools/Source/SequencerAnimTools/Private/SequencerToolsEditMode.cpp
+++ b/Engine/Plugins/Editor/SequencerAnimTools/Source/SequencerAnimTools/Private/SequencerToolsEditMode.cpp
@@ -2,7 +2,7 @@
 
 #include "SequencerToolsEditMode.h"
 #include "EditorModeManager.h"
-#include "EdModeInteractiveToolsContext.h"
+#include "Tools/EdModeInteractiveToolsContext.h"
 #include "ILevelEditor.h"
 #include "LevelEditor.h"
 #include "Modules/ModuleManager.h"
diff --git a/Engine/Plugins/Editor/StaticMeshEditorModeling/Source/StaticMeshEditorModeling/Private/StaticMeshEditorModelingMode.cpp b/Engine/Plugins/Editor/StaticMeshEditorModeling/Source/StaticMeshEditorModeling/Private/StaticMeshEditorModelingMode.cpp
--- a/Engine/Plugins/Editor/StaticMeshEditorModeling/Source/StaticMeshEditorModeling/Private/StaticMeshEditorModelingMode.cpp
+++ b/Engine/Plugins/Editor/StaticMeshEditorModeling/Source/StaticMeshEditorModeling/Private/StaticMeshEditorModelingMode.cpp
@@ -2,7 +2,7 @@
 
 #include "StaticMeshEditorModelingMode.h"
 #include "StaticMeshEditorModelingToolkit.h"
-#include "EdModeInteractiveToolsContext.h"
+#include "Tools/EdModeInteractiveToolsContext.h"
 #include "ToolTargetManager.h"
 #include "ModelingToolsManagerActions.h"
 #include "ModelingToolsEditorModeSettings.h"
diff --git a/Engine/Plugins/Editor/StaticMeshEditorModeling/Source/StaticMeshEditorModeling/Private/StaticMeshEditorModelingToolkit.cpp b/Engine/Plugins/Editor/StaticMeshEditorModeling/Source/StaticMeshEditorModeling/Private/StaticMeshEditorModelingToolkit.cpp
--- a/Engine/Plugins/Editor/StaticMeshEditorModeling/Source/StaticMeshEditorModeling/Private/StaticMeshEditorModelingToolkit.cpp
+++ b/Engine/Plugins/Editor/StaticMeshEditorModeling/Source/StaticMeshEditorModeling/Private/StaticMeshEditorModelingToolkit.cpp
@@ -7,7 +7,7 @@
 #include "ModelingToolsManagerActions.h"
 #include "Tools/UEdMode.h"
 #include "InteractiveTool.h"
-#include "EdModeInteractiveToolsContext.h"
+#include "Tools/EdModeInteractiveToolsContext.h"
 #include "InteractiveToolManager.h"
 #include "Widgets/Text/STextBlock.h"
 #include "Styling/AppStyle.h"
diff --git a/Engine/Plugins/Editor/UVEditor/Source/UVEditor/Private/UVEditor.cpp b/Engine/Plugins/Editor/UVEditor/Source/UVEditor/Private/UVEditor.cpp
--- a/Engine/Plugins/Editor/UVEditor/Source/UVEditor/Private/UVEditor.cpp
+++ b/Engine/Plugins/Editor/UVEditor/Source/UVEditor/Private/UVEditor.cpp
@@ -7,7 +7,7 @@
 #include "UVEditorToolkit.h"
 #include "UVEditorSubsystem.h"
 #include "EditorModeManager.h"
-#include "EdModeInteractiveToolsContext.h"
+#include "Tools/EdModeInteractiveToolsContext.h"
 
 #include "AssetEditorModeManager.h"
 #include "UVEditorMode.h"
diff --git a/Engine/Plugins/Editor/UVEditor/Source/UVEditor/Private/UVEditor2DViewportClient.cpp b/Engine/Plugins/Editor/UVEditor/Source/UVEditor/Private/UVEditor2DViewportClient.cpp
--- a/Engine/Plugins/Editor/UVEditor/Source/UVEditor/Private/UVEditor2DViewportClient.cpp
+++ b/Engine/Plugins/Editor/UVEditor/Source/UVEditor/Private/UVEditor2DViewportClient.cpp
@@ -6,7 +6,7 @@
 #include "BaseBehaviors/MouseWheelBehavior.h"
 #include "ContextObjectStore.h"
 #include "EditorModeManager.h"
-#include "EdModeInteractiveToolsContext.h"
+#include "Tools/EdModeInteractiveToolsContext.h"
 #include "Drawing/MeshDebugDrawing.h"
 #include "FrameTypes.h"
 #include "UVEditorMode.h"
diff --git a/Engine/Plugins/Editor/UVEditor/Source/UVEditor/Private/UVEditor3DViewportClient.cpp b/Engine/Plugins/Editor/UVEditor/Source/UVEditor/Private/UVEditor3DViewportClient.cpp
--- a/Engine/Plugins/Editor/UVEditor/Source/UVEditor/Private/UVEditor3DViewportClient.cpp
+++ b/Engine/Plugins/Editor/UVEditor/Source/UVEditor/Private/UVEditor3DViewportClient.cpp
@@ -6,7 +6,7 @@
 #include "BaseBehaviors/MouseWheelBehavior.h"
 #include "ContextObjectStore.h"
 #include "EditorModeManager.h"
-#include "EdModeInteractiveToolsContext.h"
+#include "Tools/EdModeInteractiveToolsContext.h"
 #include "Drawing/MeshDebugDrawing.h"
 #include "FrameTypes.h"
 #include "UVEditorMode.h"
@@ -59,4 +59,4 @@ void FUVEditor3DViewportClient::FocusCameraOnSelection()
 	{
 		ViewportButtonsAPI->InitiateFocusCameraOnSelection();
 	}
-}
\ No newline at end of file
+}
diff --git a/Engine/Plugins/Editor/UVEditor/Source/UVEditor/Private/UVEditorMode.cpp b/Engine/Plugins/Editor/UVEditor/Source/UVEditor/Private/UVEditorMode.cpp
--- a/Engine/Plugins/Editor/UVEditor/Source/UVEditor/Private/UVEditorMode.cpp
+++ b/Engine/Plugins/Editor/UVEditor/Source/UVEditor/Private/UVEditorMode.cpp
@@ -19,7 +19,7 @@
 #include "EditorModelingObjectsCreationAPI.h"
 #include "EditorModes.h"
 #include "EditorViewportClient.h"
-#include "EdModeInteractiveToolsContext.h" //ToolsContext, EditorInteractiveToolsContext
+#include "Tools/EdModeInteractiveToolsContext.h" //ToolsContext, EditorInteractiveToolsContext
 #include "EngineAnalytics.h"
 #include "Framework/Commands/UICommandList.h"
 #include "InteractiveTool.h"
diff --git a/Engine/Plugins/Editor/UVEditor/Source/UVEditor/Private/UVEditorModeToolkit.cpp b/Engine/Plugins/Editor/UVEditor/Source/UVEditor/Private/UVEditorModeToolkit.cpp
--- a/Engine/Plugins/Editor/UVEditor/Source/UVEditor/Private/UVEditorModeToolkit.cpp
+++ b/Engine/Plugins/Editor/UVEditor/Source/UVEditor/Private/UVEditorModeToolkit.cpp
@@ -20,7 +20,7 @@
 #include "Widgets/Layout/SBorder.h"
 #include "Widgets/Text/STextBlock.h"
 #include "Widgets/Images/SImage.h"
-#include "EdModeInteractiveToolsContext.h"
+#include "Tools/EdModeInteractiveToolsContext.h"
 #include "EditorModeManager.h"
 #include "UVEditorModeUILayer.h"
 #include "UVEditorStyle.h"
diff --git a/Engine/Plugins/Editor/UVEditor/Source/UVEditor/Private/UVEditorToolkit.cpp b/Engine/Plugins/Editor/UVEditor/Source/UVEditor/Private/UVEditorToolkit.cpp
--- a/Engine/Plugins/Editor/UVEditor/Source/UVEditor/Private/UVEditorToolkit.cpp
+++ b/Engine/Plugins/Editor/UVEditor/Source/UVEditor/Private/UVEditorToolkit.cpp
@@ -32,7 +32,7 @@
 
 #include "SLevelViewport.h"
 
-#include "EdModeInteractiveToolsContext.h"
+#include "Tools/EdModeInteractiveToolsContext.h"
 
 #define LOCTEXT_NAMESPACE "UVEditorToolkit"
 
diff --git a/Engine/Plugins/Enterprise/LidarPointCloud/Source/LidarPointCloudEditor/Private/LidarPointCloudEdModeToolkit.cpp b/Engine/Plugins/Enterprise/LidarPointCloud/Source/LidarPointCloudEditor/Private/LidarPointCloudEdModeToolkit.cpp
--- a/Engine/Plugins/Enterprise/LidarPointCloud/Source/LidarPointCloudEditor/Private/LidarPointCloudEdModeToolkit.cpp
+++ b/Engine/Plugins/Enterprise/LidarPointCloud/Source/LidarPointCloudEditor/Private/LidarPointCloudEdModeToolkit.cpp
@@ -1,7 +1,7 @@
 // Copyright Epic Games, Inc. All Rights Reserved.
 
 #include "LidarPointCloudEdModeToolkit.h"
-#include "EdModeInteractiveToolsContext.h"
+#include "Tools/EdModeInteractiveToolsContext.h"
 #include "InteractiveToolManager.h"
 #include "LidarPointCloudEditorCommands.h"
 #include "LidarPointCloudEditorTools.h"
diff --git a/Engine/Plugins/Experimental/AnimNext/Source/AnimNextEditor/Private/AnimNextEditorModule.cpp b/Engine/Plugins/Experimental/AnimNext/Source/AnimNextEditor/Private/AnimNextEditorModule.cpp
--- a/Engine/Plugins/Experimental/AnimNext/Source/AnimNextEditor/Private/AnimNextEditorModule.cpp
+++ b/Engine/Plugins/Experimental/AnimNext/Source/AnimNextEditor/Private/AnimNextEditorModule.cpp
@@ -53,7 +53,7 @@
 #include "Common/AnimNextRigVMAssetEditorDataCustomization.h"
 #include "DataInterface/AnimNextDataInterface.h"
 #include "Editor/RigVMEditorStyle.h"
-#include "EdModeInteractiveToolsContext.h"
+#include "Tools/EdModeInteractiveToolsContext.h"
 #include "Entries/AnimNextEventGraphEntry.h"
 #include "Logging/MessageLog.h"
 #include "Graph/AnimNextEdGraphCustomization.h"
diff --git a/Engine/Plugins/Experimental/AnimNextStateTree/Source/AnimNextStateTreeEditor/Private/AnimNextStateTreeEditorModule.cpp b/Engine/Plugins/Experimental/AnimNextStateTree/Source/AnimNextStateTreeEditor/Private/AnimNextStateTreeEditorModule.cpp
--- a/Engine/Plugins/Experimental/AnimNextStateTree/Source/AnimNextStateTreeEditor/Private/AnimNextStateTreeEditorModule.cpp
+++ b/Engine/Plugins/Experimental/AnimNextStateTree/Source/AnimNextStateTreeEditor/Private/AnimNextStateTreeEditorModule.cpp
@@ -13,7 +13,7 @@
 #include "StateTreeEditorMode.h"
 #include "WorkspaceItemMenuContext.h"
 #include "Toolkits/AssetEditorToolkitMenuContext.h"
-#include "EdModeInteractiveToolsContext.h"
+#include "Tools/EdModeInteractiveToolsContext.h"
 #include "AnimNextStateTreeEditorHost.h"
 #include "StateTreeEditingSubsystem.h"
 #include "StateTreeEditorStyle.h"
diff --git a/Engine/Plugins/Experimental/Avalanche/Source/AvalancheInteractiveTools/Private/EdMode/AvaInteractiveToolsEdMode.cpp b/Engine/Plugins/Experimental/Avalanche/Source/AvalancheInteractiveTools/Private/EdMode/AvaInteractiveToolsEdMode.cpp
--- a/Engine/Plugins/Experimental/Avalanche/Source/AvalancheInteractiveTools/Private/EdMode/AvaInteractiveToolsEdMode.cpp
+++ b/Engine/Plugins/Experimental/Avalanche/Source/AvalancheInteractiveTools/Private/EdMode/AvaInteractiveToolsEdMode.cpp
@@ -9,7 +9,7 @@
 #include "ContextObjectStore.h"
 #include "EditorModeManager.h"
 #include "EditorModes.h"
-#include "EdModeInteractiveToolsContext.h"
+#include "Tools/EdModeInteractiveToolsContext.h"
 #include "IAvalancheInteractiveToolsModule.h"
 #include "Selection.h"
 #include "Tools/AvaInteractiveToolsToolBase.h"
diff --git a/Engine/Plugins/Experimental/ChaosEditor/Source/FractureEditor/Private/FractureEditorMode.cpp b/Engine/Plugins/Experimental/ChaosEditor/Source/FractureEditor/Private/FractureEditorMode.cpp
--- a/Engine/Plugins/Experimental/ChaosEditor/Source/FractureEditor/Private/FractureEditorMode.cpp
+++ b/Engine/Plugins/Experimental/ChaosEditor/Source/FractureEditor/Private/FractureEditorMode.cpp
@@ -27,7 +27,7 @@
 #include "ScopedTransaction.h"
 #include "Settings/LevelEditorViewportSettings.h"
 #include "GeometryCollection/GeometryCollectionAlgo.h"
-#include "EdModeInteractiveToolsContext.h"
+#include "Tools/EdModeInteractiveToolsContext.h"
 #include "BaseGizmos/TransformGizmoUtil.h"
 #include "Snapping/ModelingSceneSnappingManager.h"
 #include "Components/InstancedStaticMeshComponent.h"
diff --git a/Engine/Plugins/Experimental/ChaosEditor/Source/FractureEditor/Private/FractureToolSelection.cpp b/Engine/Plugins/Experimental/ChaosEditor/Source/FractureEditor/Private/FractureToolSelection.cpp
--- a/Engine/Plugins/Experimental/ChaosEditor/Source/FractureEditor/Private/FractureToolSelection.cpp
+++ b/Engine/Plugins/Experimental/ChaosEditor/Source/FractureEditor/Private/FractureToolSelection.cpp
@@ -13,7 +13,7 @@
 #include "CanvasTypes.h"
 #include "EditorModeManager.h"
 #include "FractureEditorMode.h"
-#include "EdModeInteractiveToolsContext.h"
+#include "Tools/EdModeInteractiveToolsContext.h"
 #include "BaseBehaviors/SingleClickOrDragBehavior.h"
 #include "GameFramework/Actor.h"
 #include "Engine/Selection.h"
diff --git a/Engine/Plugins/Experimental/ChaosEditor/Source/FractureEditor/Public/FractureTool.h b/Engine/Plugins/Experimental/ChaosEditor/Source/FractureEditor/Public/FractureTool.h
--- a/Engine/Plugins/Experimental/ChaosEditor/Source/FractureEditor/Public/FractureTool.h
+++ b/Engine/Plugins/Experimental/ChaosEditor/Source/FractureEditor/Public/FractureTool.h
@@ -10,7 +10,7 @@
 #include "FractureEditorCommands.h"
 #include "FractureEditorModeToolkit.h"
 
-#include "EdModeInteractiveToolsContext.h"
+#include "Tools/EdModeInteractiveToolsContext.h"
 
 #include "FractureTool.generated.h"
 
diff --git a/Engine/Plugins/Experimental/CharacterFXEditor/BaseCharacterFXEditor/Source/BaseCharacterFXEditor/Private/BaseCharacterFXEditorMode.cpp b/Engine/Plugins/Experimental/CharacterFXEditor/BaseCharacterFXEditor/Source/BaseCharacterFXEditor/Private/BaseCharacterFXEditorMode.cpp
--- a/Engine/Plugins/Experimental/CharacterFXEditor/BaseCharacterFXEditor/Source/BaseCharacterFXEditor/Private/BaseCharacterFXEditorMode.cpp
+++ b/Engine/Plugins/Experimental/CharacterFXEditor/BaseCharacterFXEditor/Source/BaseCharacterFXEditor/Private/BaseCharacterFXEditorMode.cpp
@@ -4,7 +4,7 @@
 
 #include "Components/DynamicMeshComponent.h"
 #include "AssetEditorModeManager.h"
-#include "EdModeInteractiveToolsContext.h"
+#include "Tools/EdModeInteractiveToolsContext.h"
 #include "ModelingToolTargetUtil.h"
 #include "BaseCharacterFXEditorCommands.h"
 #include "BaseCharacterFXEditorModeToolkit.h"
diff --git a/Engine/Plugins/Experimental/CharacterFXEditor/BaseCharacterFXEditor/Source/BaseCharacterFXEditor/Private/BaseCharacterFXEditorModeToolkit.cpp b/Engine/Plugins/Experimental/CharacterFXEditor/BaseCharacterFXEditor/Source/BaseCharacterFXEditor/Private/BaseCharacterFXEditorModeToolkit.cpp
--- a/Engine/Plugins/Experimental/CharacterFXEditor/BaseCharacterFXEditor/Source/BaseCharacterFXEditor/Private/BaseCharacterFXEditorModeToolkit.cpp
+++ b/Engine/Plugins/Experimental/CharacterFXEditor/BaseCharacterFXEditor/Source/BaseCharacterFXEditor/Private/BaseCharacterFXEditorModeToolkit.cpp
@@ -13,7 +13,7 @@
 #include "Widgets/Layout/SBorder.h"
 #include "Widgets/Text/STextBlock.h"
 #include "Widgets/Images/SImage.h"
-#include "EdModeInteractiveToolsContext.h"
+#include "Tools/EdModeInteractiveToolsContext.h"
 #include "EditorModeManager.h"
 #include "AssetEditorModeManager.h"
 #include "Toolkits/AssetEditorModeUILayer.h"
diff --git a/Engine/Plugins/Experimental/CharacterFXEditor/ExampleCharacterFXEditor/Source/ExampleCharacterFXEditor/Private/ExampleCharacterFXEditorMode.cpp b/Engine/Plugins/Experimental/CharacterFXEditor/ExampleCharacterFXEditor/Source/ExampleCharacterFXEditor/Private/ExampleCharacterFXEditorMode.cpp
--- a/Engine/Plugins/Experimental/CharacterFXEditor/ExampleCharacterFXEditor/Source/ExampleCharacterFXEditor/Private/ExampleCharacterFXEditorMode.cpp
+++ b/Engine/Plugins/Experimental/CharacterFXEditor/ExampleCharacterFXEditor/Source/ExampleCharacterFXEditor/Private/ExampleCharacterFXEditorMode.cpp
@@ -8,7 +8,7 @@
 #include "TargetInterfaces/MaterialProvider.h"
 #include "TargetInterfaces/DynamicMeshCommitter.h"
 #include "TargetInterfaces/DynamicMeshProvider.h"
-#include "EdModeInteractiveToolsContext.h"
+#include "Tools/EdModeInteractiveToolsContext.h"
 #include "ToolTargetManager.h"
 #include "ToolTargets/StaticMeshComponentToolTarget.h"
 #include "ToolTargets/DynamicMeshComponentToolTarget.h"
diff --git a/Engine/Plugins/Experimental/Dataflow/Source/DataflowEditor/Private/Dataflow/DataflowConstructionViewportClient.cpp b/Engine/Plugins/Experimental/Dataflow/Source/DataflowEditor/Private/Dataflow/DataflowConstructionViewportClient.cpp
--- a/Engine/Plugins/Experimental/Dataflow/Source/DataflowEditor/Private/Dataflow/DataflowConstructionViewportClient.cpp
+++ b/Engine/Plugins/Experimental/Dataflow/Source/DataflowEditor/Private/Dataflow/DataflowConstructionViewportClient.cpp
@@ -14,7 +14,7 @@
 #include "Dataflow/DataflowEditorPreviewSceneBase.h"
 #include "Dataflow/DataflowRenderingViewMode.h"
 #include "EditorModeManager.h"
-#include "EdModeInteractiveToolsContext.h"
+#include "Tools/EdModeInteractiveToolsContext.h"
 #include "GraphEditor.h"
 #include "PreviewScene.h"
 #include "Selection.h"
diff --git a/Engine/Plugins/Experimental/Dataflow/Source/DataflowEditor/Private/Dataflow/DataflowEditorMode.cpp b/Engine/Plugins/Experimental/Dataflow/Source/DataflowEditor/Private/Dataflow/DataflowEditorMode.cpp
--- a/Engine/Plugins/Experimental/Dataflow/Source/DataflowEditor/Private/Dataflow/DataflowEditorMode.cpp
+++ b/Engine/Plugins/Experimental/Dataflow/Source/DataflowEditor/Private/Dataflow/DataflowEditorMode.cpp
@@ -26,7 +26,7 @@
 #include "Dataflow/DataflowToolTarget.h"
 #include "Dataflow/DataflowToolRegistry.h"
 #include "EditorModeManager.h"
-#include "EdModeInteractiveToolsContext.h"
+#include "Tools/EdModeInteractiveToolsContext.h"
 #include "Elements/Framework/EngineElementsLibrary.h"
 #include "EngineAnalytics.h"
 #include "MeshSelectionTool.h"
diff --git a/Engine/Plugins/Experimental/Dataflow/Source/DataflowEditor/Private/Dataflow/DataflowEditorModeToolkit.cpp b/Engine/Plugins/Experimental/Dataflow/Source/DataflowEditor/Private/Dataflow/DataflowEditorModeToolkit.cpp
--- a/Engine/Plugins/Experimental/Dataflow/Source/DataflowEditor/Private/Dataflow/DataflowEditorModeToolkit.cpp
+++ b/Engine/Plugins/Experimental/Dataflow/Source/DataflowEditor/Private/Dataflow/DataflowEditorModeToolkit.cpp
@@ -9,7 +9,7 @@
 #include "Dataflow/DataflowConstructionViewport.h"
 #include "Dataflow/DataflowEditorPreviewSceneBase.h"
 #include "Dataflow/DataflowConstructionScene.h"
-#include "EdModeInteractiveToolsContext.h"
+#include "Tools/EdModeInteractiveToolsContext.h"
 #include "Framework/Application/SlateApplication.h"
 #include "InteractiveToolManager.h"
 #include "MeshAttributePaintTool.h"
diff --git a/Engine/Plugins/Experimental/Dataflow/Source/DataflowEditor/Private/Dataflow/DataflowEditorToolkit.cpp b/Engine/Plugins/Experimental/Dataflow/Source/DataflowEditor/Private/Dataflow/DataflowEditorToolkit.cpp
--- a/Engine/Plugins/Experimental/Dataflow/Source/DataflowEditor/Private/Dataflow/DataflowEditorToolkit.cpp
+++ b/Engine/Plugins/Experimental/Dataflow/Source/DataflowEditor/Private/Dataflow/DataflowEditorToolkit.cpp
@@ -50,7 +50,7 @@
 #include "EditorViewportTabContent.h"
 #include "EditorViewportLayout.h"
 #include "EditorViewportCommands.h"
-#include "EdModeInteractiveToolsContext.h"
+#include "Tools/EdModeInteractiveToolsContext.h"
 #include "Engine/SkeletalMesh.h"
 #include "Framework/Commands/GenericCommands.h"
 #include "GameFramework/Actor.h"
diff --git a/Engine/Plugins/Experimental/Dataflow/Source/DataflowEditor/Private/Dataflow/DataflowEditorViewportClientBase.cpp b/Engine/Plugins/Experimental/Dataflow/Source/DataflowEditor/Private/Dataflow/DataflowEditorViewportClientBase.cpp
--- a/Engine/Plugins/Experimental/Dataflow/Source/DataflowEditor/Private/Dataflow/DataflowEditorViewportClientBase.cpp
+++ b/Engine/Plugins/Experimental/Dataflow/Source/DataflowEditor/Private/Dataflow/DataflowEditorViewportClientBase.cpp
@@ -3,7 +3,7 @@
 #include "Dataflow/DataflowEditorViewportClientBase.h"
 #include "Dataflow/DataflowEditorPreviewSceneBase.h"
 #include "EditorModeManager.h"
-#include "EdModeInteractiveToolsContext.h"
+#include "Tools/EdModeInteractiveToolsContext.h"
 #include "BaseBehaviors/ClickDragBehavior.h"
 #include "AssetViewerSettings.h"
 
diff --git a/Engine/Plugins/Experimental/Dataflow/Source/DataflowEditor/Private/Dataflow/DataflowSimulationViewportClient.cpp b/Engine/Plugins/Experimental/Dataflow/Source/DataflowEditor/Private/Dataflow/DataflowSimulationViewportClient.cpp
--- a/Engine/Plugins/Experimental/Dataflow/Source/DataflowEditor/Private/Dataflow/DataflowSimulationViewportClient.cpp
+++ b/Engine/Plugins/Experimental/Dataflow/Source/DataflowEditor/Private/Dataflow/DataflowSimulationViewportClient.cpp
@@ -13,7 +13,7 @@
 #include "Dataflow/DataflowSimulationScene.h"
 #include "Dataflow/DataflowSimulationVisualization.h"
 #include "EditorModeManager.h"
-#include "EdModeInteractiveToolsContext.h"
+#include "Tools/EdModeInteractiveToolsContext.h"
 #include "GraphEditor.h"
 #include "PreviewScene.h"
 #include "Selection.h"
diff --git a/Engine/Plugins/Experimental/GizmoEdMode/Source/GizmoEdMode/Private/GizmoEdMode.cpp b/Engine/Plugins/Experimental/GizmoEdMode/Source/GizmoEdMode/Private/GizmoEdMode.cpp
--- a/Engine/Plugins/Experimental/GizmoEdMode/Source/GizmoEdMode/Private/GizmoEdMode.cpp
+++ b/Engine/Plugins/Experimental/GizmoEdMode/Source/GizmoEdMode/Private/GizmoEdMode.cpp
@@ -3,7 +3,7 @@
 #include "GizmoEdMode.h"
 #include "DefaultAssetEditorGizmoFactory.h"
 #include "EditorModeManager.h"
-#include "EdModeInteractiveToolsContext.h"
+#include "Tools/EdModeInteractiveToolsContext.h"
 #include "InteractiveGizmo.h"
 #include "InteractiveGizmoManager.h"
 #include "Settings/LevelEditorViewportSettings.h"
diff --git a/Engine/Plugins/Experimental/GizmoEdMode/Source/LightGizmos/Private/LightGizmosModule.cpp b/Engine/Plugins/Experimental/GizmoEdMode/Source/LightGizmos/Private/LightGizmosModule.cpp
--- a/Engine/Plugins/Experimental/GizmoEdMode/Source/LightGizmos/Private/LightGizmosModule.cpp
+++ b/Engine/Plugins/Experimental/GizmoEdMode/Source/LightGizmos/Private/LightGizmosModule.cpp
@@ -15,7 +15,7 @@
 #include "SpotLightGizmo.h"
 #include "DirectionalLightGizmo.h"
 #include "LevelEditor.h"
-#include "EdModeInteractiveToolsContext.h"
+#include "Tools/EdModeInteractiveToolsContext.h"
 #include "InteractiveGizmoManager.h"
 
 #define LOCTEXT_NAMESPACE "FLightGizmosModule"
diff --git a/Engine/Plugins/Experimental/UserToolBoxBasicCommand/Source/UserToolBoxBasicCommand/Private/FlipNormals.cpp b/Engine/Plugins/Experimental/UserToolBoxBasicCommand/Source/UserToolBoxBasicCommand/Private/FlipNormals.cpp
--- a/Engine/Plugins/Experimental/UserToolBoxBasicCommand/Source/UserToolBoxBasicCommand/Private/FlipNormals.cpp
+++ b/Engine/Plugins/Experimental/UserToolBoxBasicCommand/Source/UserToolBoxBasicCommand/Private/FlipNormals.cpp
@@ -7,7 +7,7 @@
 
 #include "UDynamicMesh.h"
 #include "GeometryScript/MeshAssetFunctions.h"
-#include "EdModeInteractiveToolsContext.h"
+#include "Tools/EdModeInteractiveToolsContext.h"
 #include "Engine/Selection.h"
 #include "Engine/StaticMeshActor.h"
 #include "GeometryScript/MeshNormalsFunctions.h"
diff --git a/Engine/Plugins/Experimental/WidgetEditorToolPalette/Source/WidgetEditorToolPalette/Private/WidgetEditorToolPaletteMode.cpp b/Engine/Plugins/Experimental/WidgetEditorToolPalette/Source/WidgetEditorToolPalette/Private/WidgetEditorToolPaletteMode.cpp
--- a/Engine/Plugins/Experimental/WidgetEditorToolPalette/Source/WidgetEditorToolPalette/Private/WidgetEditorToolPaletteMode.cpp
+++ b/Engine/Plugins/Experimental/WidgetEditorToolPalette/Source/WidgetEditorToolPalette/Private/WidgetEditorToolPaletteMode.cpp
@@ -4,7 +4,7 @@
 #include "WidgetEditorToolPaletteCommands.h"
 #include "WidgetEditorToolPaletteToolkit.h"
 #include "ContextObjectStore.h"
-#include "EdModeInteractiveToolsContext.h"
+#include "Tools/EdModeInteractiveToolsContext.h"
 #include "ToolContexts/ToolStackContext.h"
 #include "ToolTargetManager.h"
 #include "DefaultTools/RectangleSelectTool.h"
diff --git a/Engine/Plugins/Experimental/WidgetEditorToolPalette/Source/WidgetEditorToolPalette/Private/WidgetEditorToolPaletteToolkit.cpp b/Engine/Plugins/Experimental/WidgetEditorToolPalette/Source/WidgetEditorToolPalette/Private/WidgetEditorToolPaletteToolkit.cpp
--- a/Engine/Plugins/Experimental/WidgetEditorToolPalette/Source/WidgetEditorToolPalette/Private/WidgetEditorToolPaletteToolkit.cpp
+++ b/Engine/Plugins/Experimental/WidgetEditorToolPalette/Source/WidgetEditorToolPalette/Private/WidgetEditorToolPaletteToolkit.cpp
@@ -7,7 +7,7 @@
 
 #include "Tools/UEdMode.h"
 #include "InteractiveTool.h"
-#include "EdModeInteractiveToolsContext.h"
+#include "Tools/EdModeInteractiveToolsContext.h"
 #include "InteractiveToolManager.h"
 #include "Widgets/Text/STextBlock.h"
 #include "IDetailsView.h"
diff --git a/Engine/Plugins/Experimental/Workspace/Source/WorkspaceEditor/Private/WorkspaceEditor.cpp b/Engine/Plugins/Experimental/Workspace/Source/WorkspaceEditor/Private/WorkspaceEditor.cpp
--- a/Engine/Plugins/Experimental/Workspace/Source/WorkspaceEditor/Private/WorkspaceEditor.cpp
+++ b/Engine/Plugins/Experimental/Workspace/Source/WorkspaceEditor/Private/WorkspaceEditor.cpp
@@ -22,7 +22,7 @@
 #include "AssetEditorModeManager.h"
 #include "Selection.h"
 #include "EditorModeManager.h"
-#include "EdModeInteractiveToolsContext.h"
+#include "Tools/EdModeInteractiveToolsContext.h"
 #include "ContextObjectStore.h"
 #include "SGraphDocument.h"
 #include "WorkspaceSchema.h"
@@ -861,4 +861,4 @@ void FWorkspaceEditor::RemoveEditingObject(UObject* Object)
 }
 }
 
-#undef LOCTEXT_NAMESPACE
\ No newline at end of file
+#undef LOCTEXT_NAMESPACE
diff --git a/Engine/Plugins/MeshPainting/Source/MeshPaintEditorMode/Private/MeshPaintMode.cpp b/Engine/Plugins/MeshPainting/Source/MeshPaintEditorMode/Private/MeshPaintMode.cpp
--- a/Engine/Plugins/MeshPainting/Source/MeshPaintEditorMode/Private/MeshPaintMode.cpp
+++ b/Engine/Plugins/MeshPainting/Source/MeshPaintEditorMode/Private/MeshPaintMode.cpp
@@ -9,7 +9,7 @@
 #include "IMeshPaintComponentAdapter.h"
 #include "MeshPaintModeSettings.h"
 #include "InteractiveToolManager.h"
-#include "EdModeInteractiveToolsContext.h"
+#include "Tools/EdModeInteractiveToolsContext.h"
 #include "Engine/SkeletalMesh.h"
 #include "ComponentReregisterContext.h"
 #include "Dialogs/Dialogs.h"
diff --git a/Engine/Plugins/MeshPainting/Source/MeshPaintEditorMode/Private/MeshPaintModeToolkit.cpp b/Engine/Plugins/MeshPainting/Source/MeshPaintEditorMode/Private/MeshPaintModeToolkit.cpp
--- a/Engine/Plugins/MeshPainting/Source/MeshPaintEditorMode/Private/MeshPaintModeToolkit.cpp
+++ b/Engine/Plugins/MeshPainting/Source/MeshPaintEditorMode/Private/MeshPaintModeToolkit.cpp
@@ -1,7 +1,7 @@
 // Copyright Epic Games, Inc. All Rights Reserved.
 
 #include "MeshPaintModeToolkit.h"
-#include "EdModeInteractiveToolsContext.h"
+#include "Tools/EdModeInteractiveToolsContext.h"
 #include "MeshPaintMode.h"
 #include "InteractiveToolManager.h"
 
diff --git a/Engine/Plugins/Runtime/HairStrands/Source/HairStrandsEditor/Private/GroomEditorMode.cpp b/Engine/Plugins/Runtime/HairStrands/Source/HairStrandsEditor/Private/GroomEditorMode.cpp
--- a/Engine/Plugins/Runtime/HairStrands/Source/HairStrandsEditor/Private/GroomEditorMode.cpp
+++ b/Engine/Plugins/Runtime/HairStrands/Source/HairStrandsEditor/Private/GroomEditorMode.cpp
@@ -3,7 +3,7 @@
 #include "GroomEditorMode.h"
 #include "GroomEditorCommands.h"
 #include "Framework/Commands/UICommandList.h"
-#include "EdModeInteractiveToolsContext.h"
+#include "Tools/EdModeInteractiveToolsContext.h"
 #include "InteractiveToolManager.h"
 
 #include UE_INLINE_GENERATED_CPP_BY_NAME(GroomEditorMode)
diff --git a/Engine/Plugins/Runtime/StateTree/Source/StateTreeEditorModule/Private/StateTreeEditor.cpp b/Engine/Plugins/Runtime/StateTree/Source/StateTreeEditorModule/Private/StateTreeEditor.cpp
--- a/Engine/Plugins/Runtime/StateTree/Source/StateTreeEditorModule/Private/StateTreeEditor.cpp
+++ b/Engine/Plugins/Runtime/StateTree/Source/StateTreeEditorModule/Private/StateTreeEditor.cpp
@@ -14,7 +14,7 @@
 #include "Customizations/StateTreeBindingExtension.h"
 #include "DetailsViewArgs.h"
 #include "EditorModeManager.h"
-#include "EdModeInteractiveToolsContext.h"
+#include "Tools/EdModeInteractiveToolsContext.h"
 #include "IDetailsView.h"
 #include "IContentBrowserSingleton.h"
 #include "ISourceCodeAccessModule.h"
diff --git a/Engine/Plugins/Runtime/StateTree/Source/StateTreeEditorModule/Private/StateTreeEditorMode.cpp b/Engine/Plugins/Runtime/StateTree/Source/StateTreeEditorModule/Private/StateTreeEditorMode.cpp
--- a/Engine/Plugins/Runtime/StateTree/Source/StateTreeEditorModule/Private/StateTreeEditorMode.cpp
+++ b/Engine/Plugins/Runtime/StateTree/Source/StateTreeEditorModule/Private/StateTreeEditorMode.cpp
@@ -11,7 +11,7 @@
 #include "StateTreeEditorSettings.h"
 #include "StateTreeCompilerLog.h"
 #include "StateTreeDelegates.h"
-#include "EdModeInteractiveToolsContext.h"
+#include "Tools/EdModeInteractiveToolsContext.h"
 #include "IMessageLogListing.h"
 #include "InteractiveToolManager.h"
 #include "PropertyPath.h"
diff --git a/Engine/Source/Editor/AnimationEditMode/Private/AnimationEditMode.cpp b/Engine/Source/Editor/AnimationEditMode/Private/AnimationEditMode.cpp
--- a/Engine/Source/Editor/AnimationEditMode/Private/AnimationEditMode.cpp
+++ b/Engine/Source/Editor/AnimationEditMode/Private/AnimationEditMode.cpp
@@ -3,7 +3,7 @@
 #include "AnimationEditMode.h"
 
 #include "ContextObjectStore.h"
-#include "EdModeInteractiveToolsContext.h"
+#include "Tools/EdModeInteractiveToolsContext.h"
 #include "EditorModeManager.h"
 #include "Misc/AssertionMacros.h"
 #include "Tools/Modes.h"
diff --git a/Engine/Source/Editor/Experimental/EditorInteractiveToolsFramework/Private/EditorGizmos/EditorTransformGizmoBuilder.cpp b/Engine/Source/Editor/Experimental/EditorInteractiveToolsFramework/Private/EditorGizmos/EditorTransformGizmoBuilder.cpp
--- a/Engine/Source/Editor/Experimental/EditorInteractiveToolsFramework/Private/EditorGizmos/EditorTransformGizmoBuilder.cpp
+++ b/Engine/Source/Editor/Experimental/EditorInteractiveToolsFramework/Private/EditorGizmos/EditorTransformGizmoBuilder.cpp
@@ -7,7 +7,7 @@
 #include "BaseGizmos/GizmoViewContext.h"
 #include "ContextObjectStore.h"
 #include "EditorModeManager.h"
-#include "EdModeInteractiveToolsContext.h"
+#include "Tools/EdModeInteractiveToolsContext.h"
 #include "EditorGizmos/EditorTransformGizmo.h"
 #include "EditorGizmos/EditorTransformGizmoSource.h"
 #include "EditorGizmos/EditorTransformProxy.h"
diff --git a/Engine/Source/Editor/Experimental/EditorInteractiveToolsFramework/Private/EditorGizmos/EditorTransformGizmoUtil.cpp b/Engine/Source/Editor/Experimental/EditorInteractiveToolsFramework/Private/EditorGizmos/EditorTransformGizmoUtil.cpp
--- a/Engine/Source/Editor/Experimental/EditorInteractiveToolsFramework/Private/EditorGizmos/EditorTransformGizmoUtil.cpp
+++ b/Engine/Source/Editor/Experimental/EditorInteractiveToolsFramework/Private/EditorGizmos/EditorTransformGizmoUtil.cpp
@@ -7,7 +7,7 @@
 #include "EditorInteractiveGizmoManager.h"
 #include "InteractiveToolsContext.h"
 #include "EditorModeManager.h"
-#include "EdModeInteractiveToolsContext.h"
+#include "Tools/EdModeInteractiveToolsContext.h"
 #include "EditorGizmos/TransformGizmo.h"
 #include "EditorGizmos/EditorTransformGizmoBuilder.h"
 #include "EditorGizmos/EditorTransformGizmoDataBinder.h"
diff --git a/Engine/Source/Editor/Experimental/EditorInteractiveToolsFramework/Private/EditorInteractiveGizmoManager.cpp b/Engine/Source/Editor/Experimental/EditorInteractiveToolsFramework/Private/EditorInteractiveGizmoManager.cpp
--- a/Engine/Source/Editor/Experimental/EditorInteractiveToolsFramework/Private/EditorInteractiveGizmoManager.cpp
+++ b/Engine/Source/Editor/Experimental/EditorInteractiveToolsFramework/Private/EditorInteractiveGizmoManager.cpp
@@ -9,7 +9,7 @@
 #include "EditorInteractiveGizmoSelectionBuilder.h"
 #include "EditorInteractiveGizmoSubsystem.h"
 #include "EditorModeManager.h"
-#include "EdModeInteractiveToolsContext.h"
+#include "Tools/EdModeInteractiveToolsContext.h"
 #include "HAL/IConsoleManager.h"
 #include "HAL/Platform.h"
 #include "HAL/PlatformCrt.h"
diff --git a/Engine/Source/Editor/Experimental/EditorInteractiveToolsFramework/Private/InteractiveToolStack.cpp b/Engine/Source/Editor/Experimental/EditorInteractiveToolsFramework/Private/InteractiveToolStack.cpp
--- a/Engine/Source/Editor/Experimental/EditorInteractiveToolsFramework/Private/InteractiveToolStack.cpp
+++ b/Engine/Source/Editor/Experimental/EditorInteractiveToolsFramework/Private/InteractiveToolStack.cpp
@@ -5,7 +5,7 @@
 
 #include "ContextObjectStore.h"
 #include "Delegates/Delegate.h"
-#include "EdModeInteractiveToolsContext.h"
+#include "Tools/EdModeInteractiveToolsContext.h"
 #include "Framework/Commands/Contexts/UIContentContext.h"
 #include "Framework/Commands/Contexts/UIIdentifierContext.h"
 #include "Framework/Commands/UIAction.h"
diff --git a/Engine/Source/Editor/Experimental/EditorInteractiveToolsFramework/Private/ToolContexts/ToolStackContext.cpp b/Engine/Source/Editor/Experimental/EditorInteractiveToolsFramework/Private/ToolContexts/ToolStackContext.cpp
--- a/Engine/Source/Editor/Experimental/EditorInteractiveToolsFramework/Private/ToolContexts/ToolStackContext.cpp
+++ b/Engine/Source/Editor/Experimental/EditorInteractiveToolsFramework/Private/ToolContexts/ToolStackContext.cpp
@@ -2,7 +2,7 @@
 
 #include "ToolContexts/ToolStackContext.h"
 
-#include "EdModeInteractiveToolsContext.h"
+#include "Tools/EdModeInteractiveToolsContext.h"
 #include "Framework/Commands/InputChord.h"
 #include "Framework/Commands/UICommandInfo.h"
 #include "Framework/Commands/UICommandList.h"
@@ -109,4 +109,4 @@ void UToolStackContext::UnregisterToolStack(const FString& Identifier)
 FInteractiveToolStack* UToolStackContext::AccessToolStack(const FString& Identifier)
 {
 	return ToolStacks.Find(Identifier);
-}
\ No newline at end of file
+}
diff --git a/Engine/Source/Editor/Experimental/EditorInteractiveToolsFramework/Public/EdModeInteractiveToolsContext.h b/Engine/Source/Editor/Experimental/EditorInteractiveToolsFramework/Public/EdModeInteractiveToolsContext.h
--- a/Engine/Source/Editor/Experimental/EditorInteractiveToolsFramework/Public/EdModeInteractiveToolsContext.h
+++ b/Engine/Source/Editor/Experimental/EditorInteractiveToolsFramework/Public/EdModeInteractiveToolsContext.h
@@ -2,6 +2,6 @@
 #pragma once
 #include "Misc/CoreMiscDefines.h"
 
-//todo: have a dedicated CL to uncomment this deprecation notice and update all files using this header
-//UE_DEPRECATED_HEADER(5.6, "Use Tools/EditorInteractiveToolsContext.h instead of EdModeInteractiveToolsContext.h.")
-#include "Tools/EdModeInteractiveToolsContext.h"
\ No newline at end of file
+// Moved to UnrealEd module
+UE_DEPRECATED_HEADER(5.6, "Use Tools/EditorInteractiveToolsContext.h instead of EdModeInteractiveToolsContext.h.")
+#include "Tools/EdModeInteractiveToolsContext.h"
diff --git a/Engine/Source/Editor/Experimental/EditorInteractiveToolsFramework/Public/ToolContexts/WidgetToolsContext.h b/Engine/Source/Editor/Experimental/EditorInteractiveToolsFramework/Public/ToolContexts/WidgetToolsContext.h
--- a/Engine/Source/Editor/Experimental/EditorInteractiveToolsFramework/Public/ToolContexts/WidgetToolsContext.h
+++ b/Engine/Source/Editor/Experimental/EditorInteractiveToolsFramework/Public/ToolContexts/WidgetToolsContext.h
@@ -2,7 +2,7 @@
 
 #pragma once
 
-#include "EdModeInteractiveToolsContext.h"
+#include "Tools/EdModeInteractiveToolsContext.h"
 #include "HAL/Platform.h"
 #include "Input/Events.h"
 #include "Layout/Geometry.h"
@@ -47,4 +47,4 @@ public:
 	void OnMouseCaptureLost(const FCaptureLostEvent& CaptureLostEvent);
 	int32 OnPaint(const FPaintArgs& Args, const FGeometry& AllottedGeometry, const FSlateRect& MyCullingRect, FSlateWindowElementList& OutDrawElements, int32 LayerId, const FWidgetStyle& InWidgetStyle, bool bParentEnabled) const;
 	// ~End SWidget "Proxy" Interface
-};
\ No newline at end of file
+};
diff --git a/Engine/Source/Editor/LevelInstanceEditor/Private/LevelInstanceEditorMode.cpp b/Engine/Source/Editor/LevelInstanceEditor/Private/LevelInstanceEditorMode.cpp
--- a/Engine/Source/Editor/LevelInstanceEditor/Private/LevelInstanceEditorMode.cpp
+++ b/Engine/Source/Editor/LevelInstanceEditor/Private/LevelInstanceEditorMode.cpp
@@ -20,7 +20,7 @@
 #include "Framework/Application/SlateApplication.h"
 #include "Modules/ModuleManager.h"
 #include "InteractiveToolManager.h"
-#include "EdModeInteractiveToolsContext.h"
+#include "Tools/EdModeInteractiveToolsContext.h"
 #include "BaseBehaviors/MouseWheelBehavior.h"
 #include "InputRouter.h"
 #include "ToolContextInterfaces.h"
diff --git a/Engine/Source/Editor/LevelInstanceEditor/Private/LevelInstanceEditorModule.cpp b/Engine/Source/Editor/LevelInstanceEditor/Private/LevelInstanceEditorModule.cpp
--- a/Engine/Source/Editor/LevelInstanceEditor/Private/LevelInstanceEditorModule.cpp
+++ b/Engine/Source/Editor/LevelInstanceEditor/Private/LevelInstanceEditorModule.cpp
@@ -49,7 +49,7 @@
 #include "WorldPartition/WorldPartitionActorLoaderInterface.h"
 #include "ScopedTransaction.h"
 #include "ISCSEditorUICustomization.h"
-#include "EdModeInteractiveToolsContext.h"
+#include "Tools/EdModeInteractiveToolsContext.h"
 #include "SceneOutlinerModule.h"
 #include "SceneOutlinerFwd.h"
 
@@ -1261,4 +1261,4 @@ void FLevelInstanceEditorModule::AddReferencedObjects(FReferenceCollector& Colle
 	DefaultBehaviorSource.AddReferencedObjects(Collector);
 }
 
-#undef LOCTEXT_NAMESPACE
\ No newline at end of file
+#undef LOCTEXT_NAMESPACE
diff --git a/Engine/Source/Editor/Persona/Private/AnimationEditorViewportClient.cpp b/Engine/Source/Editor/Persona/Private/AnimationEditorViewportClient.cpp
--- a/Engine/Source/Editor/Persona/Private/AnimationEditorViewportClient.cpp
+++ b/Engine/Source/Editor/Persona/Private/AnimationEditorViewportClient.cpp
@@ -42,7 +42,7 @@
 #include "RawIndexBuffer.h"
 #include "CameraController.h"
 #include "ContextObjectStore.h"
-#include "EdModeInteractiveToolsContext.h"
+#include "Tools/EdModeInteractiveToolsContext.h"
 #include "IPersonaToolkit.h"
 #include "Animation/MorphTarget.h"
 #include "Rendering/SkeletalMeshModel.h"
diff --git a/Engine/Source/Editor/Persona/Private/IPersonaEditorModeManager.cpp b/Engine/Source/Editor/Persona/Private/IPersonaEditorModeManager.cpp
--- a/Engine/Source/Editor/Persona/Private/IPersonaEditorModeManager.cpp
+++ b/Engine/Source/Editor/Persona/Private/IPersonaEditorModeManager.cpp
@@ -3,7 +3,7 @@
 #include "IPersonaEditorModeManager.h"
 
 #include "ContextObjectStore.h"
-#include "EdModeInteractiveToolsContext.h"
+#include "Tools/EdModeInteractiveToolsContext.h"
 #include "SEditorViewport.h"
 #include "Slate/SceneViewport.h"
 #include "Framework/Application/SlateApplication.h"
@@ -74,4 +74,4 @@ void IPersonaEditorModeManager::AddReferencedObjects(FReferenceCollector& Collec
 {
 	FAssetEditorModeManager::AddReferencedObjects(Collector);
 	Collector.AddReferencedObject(PersonaModeManagerContext);
-}
\ No newline at end of file
+}
diff --git a/Engine/Source/Editor/Persona/Private/PersonaEditorModeManager.cpp b/Engine/Source/Editor/Persona/Private/PersonaEditorModeManager.cpp
--- a/Engine/Source/Editor/Persona/Private/PersonaEditorModeManager.cpp
+++ b/Engine/Source/Editor/Persona/Private/PersonaEditorModeManager.cpp
@@ -6,7 +6,7 @@
 #include "IPersonaPreviewScene.h"
 #include "Selection.h"
 #include "Animation/DebugSkelMeshComponent.h"
-#include "EdModeInteractiveToolsContext.h"
+#include "Tools/EdModeInteractiveToolsContext.h"
 #include "ContextObjectStore.h"
 
 bool FPersonaEditorModeManager::GetCameraTarget(FSphere& OutTarget) const
diff --git a/Engine/Source/Editor/Persona/Private/SAnimationEditorViewport.cpp b/Engine/Source/Editor/Persona/Private/SAnimationEditorViewport.cpp
--- a/Engine/Source/Editor/Persona/Private/SAnimationEditorViewport.cpp
+++ b/Engine/Source/Editor/Persona/Private/SAnimationEditorViewport.cpp
@@ -32,7 +32,7 @@
 #include "PreviewProfileController.h"
 #include "Materials/Material.h"
 #include "EditorFontGlyphs.h"
-#include "EdModeInteractiveToolsContext.h"
+#include "Tools/EdModeInteractiveToolsContext.h"
 #include "ContextObjectStore.h"
 #include "IPersonaEditMode.h"
 
diff --git a/Engine/Source/Editor/UnrealEd/Private/EditorModeManager.cpp b/Engine/Source/Editor/UnrealEd/Private/EditorModeManager.cpp
--- a/Engine/Source/Editor/UnrealEd/Private/EditorModeManager.cpp
+++ b/Engine/Source/Editor/UnrealEd/Private/EditorModeManager.cpp
@@ -27,7 +27,7 @@
 #include "Subsystems/BrushEditingSubsystem.h"
 #include "Tools/UEdMode.h"
 #include "InputRouter.h"
-#include "EdModeInteractiveToolsContext.h"
+#include "Tools/EdModeInteractiveToolsContext.h"
 #include "Tools/LegacyEdModeInterfaces.h"
 #include "CanvasTypes.h"
 #include "CanvasItem.h"
@@ -2136,4 +2136,4 @@ bool FEditorModeTools::EndTransform(const FGizmoState& InState) const
 bool FEditorModeTools::HasOngoingTransform() const
 {
 	return bHasOngoingTransform;
-}
\ No newline at end of file
+}
diff --git a/Engine/Source/Editor/UnrealEd/Private/EditorViewportClient.cpp b/Engine/Source/Editor/UnrealEd/Private/EditorViewportClient.cpp
--- a/Engine/Source/Editor/UnrealEd/Private/EditorViewportClient.cpp
+++ b/Engine/Source/Editor/UnrealEd/Private/EditorViewportClient.cpp
@@ -69,7 +69,7 @@
 #include "VirtualShadowMapVisualizationData.h"
 #include "GPUSkinCacheVisualizationData.h"
 #include "UnrealWidget.h"
-#include "EdModeInteractiveToolsContext.h"
+#include "Tools/EdModeInteractiveToolsContext.h"
 #include "Engine/World.h"
 #include "ProfilingDebugging/MiscTrace.h"
 #include "ProfilingDebugging/TraceScreenshot.h"
diff --git a/Engine/Source/Editor/UnrealEd/Private/LevelEditorViewport.cpp b/Engine/Source/Editor/UnrealEd/Private/LevelEditorViewport.cpp
--- a/Engine/Source/Editor/UnrealEd/Private/LevelEditorViewport.cpp
+++ b/Engine/Source/Editor/UnrealEd/Private/LevelEditorViewport.cpp
@@ -109,7 +109,7 @@
 #include "Materials/MaterialExpressionBounds.h"
 #include "LevelEditorDragDropHandler.h"
 #include "UnrealWidget.h"
-#include "EdModeInteractiveToolsContext.h"
+#include "Tools/EdModeInteractiveToolsContext.h"
 #include "GenericPlatform/GenericPlatformInputDeviceMapper.h"
 #include "Rendering/StaticLightingSystemInterface.h"
 #include "TextureResource.h"
diff --git a/Engine/Source/Editor/UnrealEd/Private/Toolkits/BaseToolkit.cpp b/Engine/Source/Editor/UnrealEd/Private/Toolkits/BaseToolkit.cpp
--- a/Engine/Source/Editor/UnrealEd/Private/Toolkits/BaseToolkit.cpp
+++ b/Engine/Source/Editor/UnrealEd/Private/Toolkits/BaseToolkit.cpp
@@ -1,7 +1,7 @@
 // Copyright Epic Games, Inc. All Rights Reserved.
 
 #include "Toolkits/BaseToolkit.h"
-#include "EdModeInteractiveToolsContext.h"
+#include "Tools/EdModeInteractiveToolsContext.h"
 #include "Toolkits/ToolkitManager.h"
 #include "PropertyEditorModule.h"
 #include "Modules/ModuleManager.h"
@@ -891,4 +891,4 @@ void FModeToolkit::RebuildModeToolBar()
 	}
 }
 
-#undef LOCTEXT_NAMESPACE
\ No newline at end of file
+#undef LOCTEXT_NAMESPACE
diff --git a/Engine/Source/Editor/UnrealEd/Private/Tools/DragToolsBehaviourSource.cpp b/Engine/Source/Editor/UnrealEd/Private/Tools/DragToolsBehaviourSource.cpp
--- a/Engine/Source/Editor/UnrealEd/Private/Tools/DragToolsBehaviourSource.cpp
+++ b/Engine/Source/Editor/UnrealEd/Private/Tools/DragToolsBehaviourSource.cpp
@@ -4,7 +4,7 @@
 #include "BaseBehaviors/ClickDragBehavior.h"
 #include "BaseBehaviors/KeyAsModifierInputBehavior.h"
 #include "DragTool_FrustumSelect.h"
-#include "EdModeInteractiveToolsContext.h"
+#include "Tools/EdModeInteractiveToolsContext.h"
 
 void UDragToolsBehaviorSource::Initialize(const UEditorInteractiveToolsContext* InInteractiveToolsContext, FEditorViewportClient* InViewportClient)
 {
diff --git a/Engine/Source/Editor/UnrealEd/Private/Tools/UEdMode.cpp b/Engine/Source/Editor/UnrealEd/Private/Tools/UEdMode.cpp
--- a/Engine/Source/Editor/UnrealEd/Private/Tools/UEdMode.cpp
+++ b/Engine/Source/Editor/UnrealEd/Private/Tools/UEdMode.cpp
@@ -5,7 +5,7 @@
 #include "EditorModeManager.h"
 #include "EditorModes.h"
 #include "Toolkits/BaseToolkit.h"
-#include "EdModeInteractiveToolsContext.h"
+#include "Tools/EdModeInteractiveToolsContext.h"
 #include "Toolkits/ToolkitManager.h"
 #include "InteractiveToolManager.h"
 #include "InteractiveToolQueryInterfaces.h"
