[EditorGizmos] Added further debug visualizations

diff --git a/Engine/Source/Editor/Experimental/EditorInteractiveToolsFramework/Private/EditorGizmos/TransformGizmo.cpp b/Engine/Source/Editor/Experimental/EditorInteractiveToolsFramework/Private/EditorGizmos/TransformGizmo.cpp
--- a/Engine/Source/Editor/Experimental/EditorInteractiveToolsFramework/Private/EditorGizmos/TransformGizmo.cpp
+++ b/Engine/Source/Editor/Experimental/EditorInteractiveToolsFramework/Private/EditorGizmos/TransformGizmo.cpp
@@ -22,6 +22,7 @@
 #include "Intersection/IntersectionUtil.h"
 #include "PrimitiveDrawingUtils.h"
 #include "SceneQueries/SceneSnappingManager.h"
+#include "ToolDataVisualizer.h"
 
 #define LOCTEXT_NAMESPACE "UTransformGizmo"
 
@@ -417,18 +418,73 @@ void UTransformGizmo::Render(IToolsContextRenderAPI* RenderAPI)
 		RenderState.Initialize(RenderAPI->GetSceneView(), GetGizmoTransform());
 		GizmoElementRoot->Render(RenderAPI, RenderState);
 
-		if (GizmoLocals::bDebugDraw)
+		RenderDebug(RenderAPI);
+	}
+}
+
+void UTransformGizmo::RenderDebug(IToolsContextRenderAPI* RenderAPI)
+{
+	if (!GizmoLocals::bDebugDraw)
+	{
+		return;
+	}
+
+	DebugData.TransformCurrent = CurrentTransform;
+
+	FToolDataVisualizer Draw;
+	Draw.BeginFrame(RenderAPI);
+
+	static constexpr float ScreenSpaceHandleRadius = 15.0f;
+
+	const FVector CameraFacingNormal = -RenderAPI->GetSceneView()->GetViewDirection();
+
+	static const FLinearColor ActorStartColor = FLinearColor::Blue.CopyWithNewOpacity(0.5f);
+	static const FLinearColor ActorCurrentColor = FLinearColor(0.0f, 1.0f, 1.0f).CopyWithNewOpacity(0.5f);
+
+	static const FLinearColor InteractionStartColor = FLinearColor::Red.CopyWithNewOpacity(0.5f);
+	static const FLinearColor InteractionCurrentColor = FLinearColor::Yellow.CopyWithNewOpacity(0.5f);
+
+	// From initial cursor pos to current for drag operation
+	if (DebugData.bIsEditing)
+	{
+		// Actual Object
+		Draw.DrawCircle(DebugData.TransformStart.GetLocation(), CameraFacingNormal, ScreenSpaceHandleRadius, 16, ActorStartColor, 0.0f, false);
+		Draw.DrawCircle(DebugData.TransformCurrent.GetLocation(), CameraFacingNormal, ScreenSpaceHandleRadius, 16, ActorCurrentColor, 0.0f, false);
+		Draw.DrawLine(DebugData.TransformStart.GetLocation(), DebugData.TransformCurrent.GetLocation(), ActorCurrentColor, 0.0f, false);
+
+		// Rotate mode displays as arcs, etc.
+		if (CurrentMode == EGizmoTransformMode::Rotate)
 		{
-			if (bDebugRotate)
-			{
-				const float RotateRadius = InteractionAxisList == EAxisList::Screen ? RotateScreenSpaceRadius : RotateAxisOuterRadius;
-				const float Radius = 2.f * GetWorldRadius(RotateRadius);
-				FPrimitiveDrawInterface* PDI = RenderAPI->GetPrimitiveDrawInterface();
-				PDI->DrawLine(DebugClosest - (DebugDirection * Radius), DebugClosest + (DebugDirection * Radius), FLinearColor::Yellow, SDPG_Foreground, 0.f, 0.01f);
-				PDI->DrawLine(DebugClosest, DebugClosest + (DebugNormalRemoved * Radius), FLinearColor::Red, SDPG_Foreground, 0.f, 0.01f);
-			}
+			const float InteractionRadius = static_cast<float>(DebugData.InteractionRadius);
+			float ArcAngleMin = static_cast<float>(FMath::RadiansToDegrees(DebugData.InteractionAngleStart));
+			float ArcAngleMax = static_cast<float>(FMath::RadiansToDegrees(DebugData.InteractionAngleCurrent));
+
+			FVector StartNormal = (DebugData.InteractionStart.GetLocation() - DebugData.TransformStart.GetLocation()).GetSafeNormal();
+			FVector InteractionStartOvershoot = DebugData.TransformStart.GetLocation() + (StartNormal * InteractionRadius);
+
+			// Cursor
+			Draw.DrawCircle(DebugData.InteractionStart.GetLocation(), CameraFacingNormal, ScreenSpaceHandleRadius, 16, InteractionStartColor, 0.0f, false);
+			Draw.DrawCircle(InteractionStartOvershoot, CameraFacingNormal, ScreenSpaceHandleRadius, 16, InteractionStartColor, 0.0f, false);
+			Draw.DrawCircle(DebugData.InteractionCurrent.GetLocation(), CameraFacingNormal, ScreenSpaceHandleRadius, 16, InteractionCurrentColor, 0.0f, false);
+
+			Draw.DrawCircle(DebugData.TransformStart.GetLocation(), DebugData.InteractionPlaneNormal, static_cast<float>(DebugData.InteractionRadius), 128, InteractionCurrentColor, 0.0f, false);
+
+			// @todo: min/max angles should be relative (to InteractionPlanarNormal?), also doesn't seem to be centered
+			// Draw.DrawArc(DebugData.TransformStart.GetLocation(), InteractionPlanarNormal, static_cast<float>(DebugData.InteractionRadius * 1.1f), 128, ArcAngleMin, ArcAngleMax, FLinearColor::Green, 0.0f, false);
+
+			Draw.DrawLine(DebugData.TransformStart.GetLocation(), InteractionStartOvershoot, InteractionStartColor.CopyWithNewOpacity(0.5f), 0.0f, false);
+			Draw.DrawLine(DebugData.TransformStart.GetLocation(), DebugData.InteractionCurrent.GetLocation(), InteractionCurrentColor, 0.0f, false);
+		}
+		else
+		{
+			// Cursor
+			Draw.DrawCircle(DebugData.InteractionStart.GetLocation(), CameraFacingNormal, ScreenSpaceHandleRadius, 16, InteractionStartColor, 0.0f, false);
+			Draw.DrawCircle(DebugData.InteractionCurrent.GetLocation(), CameraFacingNormal, ScreenSpaceHandleRadius, 16, InteractionCurrentColor, 0.0f, false);
+			Draw.DrawLine(DebugData.InteractionStart.GetLocation(), DebugData.InteractionCurrent.GetLocation(), InteractionCurrentColor, 0.0f, false);
 		}
 	}
+
+	Draw.EndFrame();
 }
 
 FInputRayHit UTransformGizmo::BeginHoverSequenceHitTest(const FInputDeviceRay& DevicePos)
@@ -1682,6 +1738,12 @@ void UTransformGizmo::ResetInteractingStates(const EGizmoTransformMode InMode)
 
 void UTransformGizmo::BeginTransformEditSequence()
 {
+	// Update Debug
+	{
+		DebugData.bIsEditing = true;
+		DebugData.TransformStart = CurrentTransform;
+	}
+
 	if (ensure(StateTarget))
 	{
 		StateTarget->BeginUpdate();
@@ -1702,6 +1764,11 @@ void UTransformGizmo::BeginTransformEditSequence()
 
 void UTransformGizmo::EndTransformEditSequence()
 {
+	// Update Debug
+	{
+		DebugData.bIsEditing = false;
+	}
+
 	if (ensure(StateTarget))
 	{
 		StateTarget->EndUpdate();
@@ -1892,6 +1959,15 @@ void UTransformGizmo::OnClickPressAxis(const FInputDeviceRay& PressPos)
 	{
 		InteractionPlanarStartPoint = PressPos.WorldRay.Origin + PressPos.WorldRay.Direction * HitDepth;
 		InteractionPlanarCurrPoint = InteractionPlanarStartPoint;
+		
+		// Update Debug
+		if (GizmoLocals::bDebugDraw)
+		{
+			DebugData.TransformStart.SetLocation(InteractionPlanarOrigin);
+			DebugData.InteractionStart.SetLocation(InteractionPlanarStartPoint);
+			DebugData.InteractionCurrent.SetLocation(InteractionPlanarStartPoint);
+			DebugData.InteractionPlaneNormal = InteractionPlanarNormal;
+		}
 	}
 
 	bInInteraction = true;
@@ -2087,6 +2163,12 @@ void UTransformGizmo::OnClickDragTranslateAxis(const FInputDeviceRay& DragPos)
 		ApplyTranslateDelta(Delta);
 
 		InteractionPlanarCurrPoint += Delta;
+
+		// Update Debug
+		if (GizmoLocals::bDebugDraw)
+		{
+			DebugData.InteractionCurrent.SetLocation(InteractionPlanarCurrPoint);
+		}
 	}
 }
 
@@ -2134,6 +2216,15 @@ void UTransformGizmo::OnClickPressPlanar(const FInputDeviceRay& PressPos)
 		InteractionPlanarCurrPoint = InteractionPlanarStartPoint;
 		bInInteraction = true;
 		
+		// Update Debug
+		if (GizmoLocals::bDebugDraw)
+		{
+			DebugData.TransformStart.SetLocation(InteractionPlanarOrigin);
+			DebugData.InteractionStart.SetLocation(InteractionPlanarStartPoint);
+			DebugData.InteractionCurrent.SetLocation(InteractionPlanarStartPoint);
+			DebugData.InteractionPlaneNormal = InteractionPlanarNormal;
+		}
+		
 		SetModeLastHitPart(EGizmoTransformMode::Translate, LastHitPart);
 	}
 }
@@ -2148,6 +2239,12 @@ void UTransformGizmo::OnClickDragTranslatePlanar(const FInputDeviceRay& DragPos)
 		SnapTranslateDelta(Delta);
 		ApplyTranslateDelta(Delta);
 		InteractionPlanarCurrPoint += Delta;
+		
+		// Update Debug
+		if (GizmoLocals::bDebugDraw)
+		{
+			DebugData.InteractionCurrent.SetLocation(InteractionPlanarCurrPoint);
+		}
 	}
 }
 
@@ -2252,6 +2349,19 @@ void UTransformGizmo::OnClickPressScale(const FInputDeviceRay& PressPos)
 	InteractionScreenStartPos = InteractionScreenEndPos = InteractionScreenCurrPos = PressPos.ScreenPosition;
 	bInInteraction = true;
 
+	// Update Debug
+	if (GizmoLocals::bDebugDraw)
+	{
+		FVector::FReal HitDepth;
+		if (GetRayParamIntersectionWithInteractionPlane(PressPos, HitDepth))
+		{
+			const FVector HitPoint = PressPos.WorldRay.Origin + PressPos.WorldRay.Direction * HitDepth;
+			DebugData.InteractionStart.SetLocation(HitPoint);
+			DebugData.InteractionCurrent.SetLocation(HitPoint);
+			DebugData.InteractionPlaneNormal = InteractionPlanarNormal;
+		}
+	}
+
 	SetModeLastHitPart(EGizmoTransformMode::Scale, LastHitPart);
 }
 
@@ -2299,6 +2409,18 @@ void UTransformGizmo::OnClickDragScale(const FInputDeviceRay& DragPos)
 		InteractionScreenEndPos -= ScreenDelta;
 		InteractionScreenCurrPos = DragPos.ScreenPosition;
 	}
+
+	// Update Debug
+	if (GizmoLocals::bDebugDraw)
+	{
+		FVector::FReal HitDepth;
+		if (GetRayParamIntersectionWithInteractionPlane(DragPos, HitDepth))
+		{
+			const FVector HitPoint = DragPos.WorldRay.Origin + DragPos.WorldRay.Direction * HitDepth;
+			DebugData.InteractionCurrent.SetLocation(HitPoint);
+			DebugData.InteractionPlaneNormal = InteractionPlanarNormal;
+		}
+	}
 }
 
 void UTransformGizmo::OnClickReleaseScaleAxis(const FInputDeviceRay& InReleasePos)
@@ -2366,6 +2488,18 @@ bool UTransformGizmo::OnClickPressRotateArc(
 				InteractionPlanarOrigin, InteractionPlanarNormal, InteractionPlanarAxisX, InteractionPlanarAxisY);
 			InteractionCurrAngle = InteractionStartAngle;
 			bInInteraction = true;
+
+			// Update Debug
+			if (GizmoLocals::bDebugDraw)
+			{
+				DebugData.TransformStart.SetLocation(InteractionPlanarOrigin);
+				DebugData.InteractionStart.SetLocation(HitPoint);
+				DebugData.InteractionAngleStart = InteractionStartAngle;
+				DebugData.InteractionAngleCurrent = InteractionCurrAngle;
+				DebugData.InteractionRadius = FVector::Distance(DebugData.TransformStart.GetLocation(), DebugData.InteractionStart.GetLocation());
+				DebugData.InteractionPlaneNormal = InteractionPlanarNormal;
+			}
+
 			return true;
 		}
 	}
@@ -2388,7 +2522,7 @@ void UTransformGizmo::OnClickPressRotateAxis(const FInputDeviceRay& InPressPos)
 
 	static const TArray AxisList({EAxisList::X, EAxisList::Y, EAxisList::Z});
 
-	bDebugRotate = true;
+	DebugData.bDebugRotate = true;
 	
 	// initialize pull data
 	InteractionScreenAxisDirection = GetScreenRotateAxisDir(InPressPos);
@@ -2438,7 +2572,7 @@ FVector2D UTransformGizmo::GetWorldToScreenRotateAxisDir(const FInputDeviceRay&
 	// NOTE that function is not intended to remain here indefinitely, its purpose is to debug closest point computation
 	const auto PrintProjection = [this](const TCHAR* const InMessage)
 	{
-		if (bDebugRotate && GizmoLocals::bDebugDraw)
+		if (DebugData.bDebugRotate && GizmoLocals::bDebugDraw)
 		{
 			UE_LOG(LogTransformGizmo, Warning, TEXT("%s"), InMessage);
 		}
@@ -2503,15 +2637,18 @@ FVector2D UTransformGizmo::GetWorldToScreenRotateAxisDir(const FInputDeviceRay&
 	const double DotAxis = FMath::Abs( FVector2D::DotProduct(PullProjection, AxisProjection) );
 	const double DotClosest = FMath::Abs( FVector2D::DotProduct(PullProjection, ToClosestProjection) );
 	NormalProjectionToRemove = (DotAxis < DotClosest) ? AxisProjection : ToClosestProjection;
-	DebugNormalRemoved = (DotAxis < DotClosest) ? InWorldAxis : ToClosestDirection;
-	DebugNormalSkip = (DotAxis < DotClosest) ? ToClosestDirection : InWorldAxis;
 
-	// debug
+	// Update Debug
+	if (GizmoLocals::bDebugDraw)
 	{
-		DebugDirection = PullDirection;
-		DebugClosest = ClosestPointOnCircle;
+		DebugData.InteractionStart.SetLocation(QueryPoint);
+		DebugData.DebugNormalRemoved = (DotAxis < DotClosest) ? InWorldAxis : ToClosestDirection;
+        DebugData.DebugNormalSkip = (DotAxis < DotClosest) ? ToClosestDirection : InWorldAxis;
+        DebugData.DebugDirection = PullDirection;
+        DebugData.DebugClosest = ClosestPointOnCircle;
+		DebugData.InteractionRadius = Radius;
 	}
-	
+
 	return PullProjection;
 }
 
@@ -2563,6 +2700,20 @@ void UTransformGizmo::OnClickDragRotateAxis(const FInputDeviceRay& DragPos)
 
 		InteractionScreenCurrPos = DragPos.ScreenPosition;
 
+		// Update Debug
+		if (GizmoLocals::bDebugDraw)
+		{
+			FVector::FReal HitDepth;
+			GetRayParamIntersectionWithInteractionPlane(DragPos, HitDepth);
+
+			const FVector HitPoint = DragPos.WorldRay.Origin + DragPos.WorldRay.Direction * HitDepth;
+			DebugData.InteractionCurrent.SetLocation(HitPoint);
+
+			const double HitAngle = InteractionCurrAngle + DeltaAngle;
+			DebugData.InteractionAngleCurrent += FMath::FindDeltaAngleRadians(DebugData.InteractionAngleCurrent, HitAngle);
+			DebugData.InteractionRadius = FVector::Distance(DebugData.TransformStart.GetLocation(), DebugData.InteractionCurrent.GetLocation());
+		}
+
 		break;
 	}
 	case EAxisRotateMode::Arc:
@@ -2585,6 +2736,15 @@ void UTransformGizmo::OnClickDragRotateAxis(const FInputDeviceRay& DragPos)
 			ApplyRotateDelta(Delta);
 
 			InteractionCurrAngle = HitAngle;
+			
+			// Update Debug
+			if (GizmoLocals::bDebugDraw)
+			{
+				DebugData.InteractionCurrent.SetLocation(HitPoint);
+				DebugData.InteractionAngleCurrent += FMath::FindDeltaAngleRadians(DebugData.InteractionAngleCurrent, HitAngle);
+				DebugData.InteractionRadius = FVector::Distance(DebugData.TransformStart.GetLocation(), DebugData.InteractionCurrent.GetLocation());
+				DebugData.InteractionPlaneNormal = InteractionPlanarNormal;
+			}
 		}
 		break;
 	}
@@ -2605,7 +2765,7 @@ double UTransformGizmo::ComputeAxisRotateDeltaAngle(const FVector2D& InStartPos,
 		if (FMath::Abs(DotNormal) > FMath::Abs(DotTangent))
 		{
 			::Swap(NormalProjectionToRemove, InteractionScreenAxisDirection);
-			::Swap(DebugDirection, DebugNormalRemoved);
+			::Swap(DebugData.DebugDirection, DebugData.DebugNormalRemoved);
 			InteractionScreenAxisDirection *= FMath::Sign(DotTangent) * FMath::Sign(DotNormal);
 		}
 		bTrySwitchingToNormalPull = false;
@@ -2654,7 +2814,7 @@ FQuat UTransformGizmo::ComputeGimbalRotateDelta(const FVector2D& InStartPos, con
 		if (FMath::Abs(DotNormal) > FMath::Abs(DotTangent))
 		{
 			::Swap(NormalProjectionToRemove, InteractionScreenAxisDirection);
-			::Swap(DebugDirection, DebugNormalRemoved);
+			::Swap(DebugData.DebugDirection, DebugData.DebugNormalRemoved);
 			InteractionScreenAxisDirection *= FMath::Sign(DotTangent) * FMath::Sign(DotNormal);
 		}
 		bTrySwitchingToNormalPull = false;
@@ -2676,7 +2836,7 @@ FQuat UTransformGizmo::ComputeGimbalRotateDelta(const FVector2D& InStartPos, con
 void UTransformGizmo::OnClickReleaseRotateAxis(const FInputDeviceRay& InReleasePos)
 {
 	bInInteraction = false;
-	bDebugRotate = false;
+	DebugData.bDebugRotate = false;
 	bTrySwitchingToNormalPull = false;
 }
 
@@ -2696,7 +2856,7 @@ void UTransformGizmo::OnClickPressGimbalRotateAxis(const FInputDeviceRay& InPres
 
 	static const TArray AxisList({EAxisList::X, EAxisList::Y, EAxisList::Z});
 	
-	bDebugRotate = true;
+	DebugData.bDebugRotate = true;
 	
 	// initialize pull data
 	InteractionScreenAxisDirection = GetScreenGimbalRotateAxisDir(InPressPos);
@@ -2751,7 +2911,7 @@ void UTransformGizmo::OnClickDragGimbalRotateAxis(const FInputDeviceRay& InDragP
 		}
 	case EAxisRotateMode::Arc:
 		{
-		FVector::FReal HitDepth;
+			FVector::FReal HitDepth;
 			if (GetRayParamIntersectionWithInteractionPlane(InDragPos, HitDepth))
 			{
 				const FVector HitPoint = InDragPos.WorldRay.Origin + InDragPos.WorldRay.Direction * HitDepth;
@@ -2775,6 +2935,14 @@ void UTransformGizmo::OnClickDragGimbalRotateAxis(const FInputDeviceRay& InDragP
 
 				ApplyRotateDelta(Delta);
 				InteractionCurrAngle = HitAngle;
+				
+				// Update Debug
+				if (GizmoLocals::bDebugDraw)
+				{
+					DebugData.InteractionCurrent.SetLocation(HitPoint);
+					DebugData.InteractionRadius = FVector::Distance(DebugData.TransformStart.GetLocation(), DebugData.InteractionCurrent.GetLocation());
+					DebugData.InteractionPlaneNormal = InteractionPlanarNormal;
+				}
 			}
 			break;
 		}
@@ -2796,14 +2964,14 @@ void UTransformGizmo::OnClickPressScreenSpaceRotate(const FInputDeviceRay& InPre
 
 	bTrySwitchingToNormalPull = false;
 	bInInteraction = false;
-	bDebugRotate = false;
+	DebugData.bDebugRotate = false;
 
 	if (bIndirectManipulation)
 	{
 		InteractionScreenAxisDirection = GetScreenRotateAxisDir(InPressPos);
 		InteractionScreenStartPos = InteractionScreenCurrPos = InPressPos.ScreenPosition;
 		bTrySwitchingToNormalPull = true;
-		bDebugRotate = true;
+		DebugData.bDebugRotate = true;
 	}
 	else
 	{
@@ -2816,8 +2984,18 @@ void UTransformGizmo::OnClickPressScreenSpaceRotate(const FInputDeviceRay& InPre
 		InteractionStartAngle = GizmoMath::ComputeAngleInPlane(HitPoint,
 			InteractionPlanarOrigin, InteractionPlanarNormal, InteractionPlanarAxisX, InteractionPlanarAxisY);
 		InteractionCurrAngle = InteractionStartAngle;
+
+		// Update Debug
+		if (GizmoLocals::bDebugDraw)
+		{
+			DebugData.InteractionStart.SetLocation(HitPoint);
+			DebugData.InteractionAngleStart = InteractionStartAngle;
+			DebugData.InteractionAngleCurrent = InteractionCurrAngle;
+			DebugData.InteractionPlaneNormal = InteractionPlanarNormal;
+			DebugData.InteractionRadius = FVector::Distance(DebugData.TransformStart.GetLocation(), DebugData.InteractionCurrent.GetLocation());
+		}
 	}
-	
+
 	bInInteraction = true;
 	SetModeLastHitPart(EGizmoTransformMode::Rotate, LastHitPart);
 }
@@ -2837,7 +3015,7 @@ void UTransformGizmo::OnClickDragScreenSpaceRotate(const FInputDeviceRay& InDrag
 			if (FMath::Abs(DotNormal) > FMath::Abs(DotTangent))
 			{
 				::Swap(NormalProjectionToRemove, InteractionScreenAxisDirection);
-				::Swap(DebugDirection, DebugNormalRemoved);
+				::Swap(DebugData.DebugDirection, DebugData.DebugNormalRemoved);
 				InteractionScreenAxisDirection *= FMath::Sign(DotTangent) * FMath::Sign(DotNormal);
 			}
 			bTrySwitchingToNormalPull = false;
@@ -2876,6 +3054,15 @@ void UTransformGizmo::OnClickDragScreenSpaceRotate(const FInputDeviceRay& InDrag
 		ApplyRotateDelta(Delta);
 
 		InteractionCurrAngle = HitAngle;
+		
+		// Update Debug
+		if (GizmoLocals::bDebugDraw)
+		{
+			//DebugData.TransformStart.SetLocation(HitPoint);
+			DebugData.InteractionCurrent.SetLocation(HitPoint);
+			DebugData.InteractionAngleCurrent += FMath::FindDeltaAngleRadians(DebugData.InteractionAngleCurrent, HitAngle);
+			DebugData.InteractionRadius = FVector::Distance(DebugData.TransformStart.GetLocation(), DebugData.InteractionCurrent.GetLocation());
+		}
 	}
 }
 
@@ -2888,7 +3075,7 @@ FQuat UTransformGizmo::ComputeAngularRotateDelta(FQuat::FReal InStartAngle, FQua
 void UTransformGizmo::OnClickReleaseScreenSpaceRotate(const FInputDeviceRay& InReleasePos)
 {
 	bInInteraction = false;
-	bDebugRotate = false;
+	DebugData.bDebugRotate = false;
 	bTrySwitchingToNormalPull = false;
 }
 
@@ -2994,6 +3181,16 @@ void UTransformGizmo::OnClickPressArcBallRotate(const FInputDeviceRay& PressPos)
 		ArcBallLocals::GetSphereAndHyperbolicProjection(
 			InteractionPlanarOrigin, SphereRadius, RayOrigin, RayDir, *GizmoViewContext,
 			InteractionArcBallStartPoint);
+
+		// Update Debug
+		if (GizmoLocals::bDebugDraw)
+		{
+			DebugData.TransformStart.SetLocation(InteractionPlanarOrigin);
+			DebugData.InteractionStart.SetLocation(InteractionArcBallStartPoint);
+			DebugData.InteractionCurrent.SetLocation(InteractionArcBallStartPoint);
+			DebugData.InteractionPlaneNormal = InteractionPlanarNormal;
+			DebugData.InteractionRadius = FVector::Distance(DebugData.TransformStart.GetLocation(), DebugData.InteractionCurrent.GetLocation());
+		}
 	
 		bInInteraction = true;
 	
@@ -3034,6 +3231,14 @@ void UTransformGizmo::OnClickDragArcBallRotate(const FInputDeviceRay& DragPos)
 		ApplyRotateDelta(DeltaRot);
 		CurrentRotation = DeltaRot * CurrentRotation; 
 	}
+
+	// Update Debug
+	if (GizmoLocals::bDebugDraw)
+	{
+		DebugData.InteractionCurrent.SetLocation(InteractionArcBallCurrPoint);
+		DebugData.InteractionRadius = FVector::Distance(DebugData.TransformStart.GetLocation(), DebugData.InteractionCurrent.GetLocation());
+		DebugData.InteractionPlaneNormal = FVector::CrossProduct(Axis1, Axis0);
+	}
 }
 
 void UTransformGizmo::OnClickReleaseArcBallRotate(const FInputDeviceRay& ReleasePos)
diff --git a/Engine/Source/Editor/Experimental/EditorInteractiveToolsFramework/Public/EditorGizmos/TransformGizmo.h b/Engine/Source/Editor/Experimental/EditorInteractiveToolsFramework/Public/EditorGizmos/TransformGizmo.h
--- a/Engine/Source/Editor/Experimental/EditorInteractiveToolsFramework/Public/EditorGizmos/TransformGizmo.h
+++ b/Engine/Source/Editor/Experimental/EditorInteractiveToolsFramework/Public/EditorGizmos/TransformGizmo.h
@@ -439,6 +439,9 @@ protected:
 
 protected:
 
+	/** Renders various debug visualizations, if enabled. */
+	void RenderDebug(IToolsContextRenderAPI* RenderAPI);
+
 	/** Setup behaviors */
 	virtual void SetupBehaviors();
 
@@ -968,10 +971,27 @@ protected:
 	bool bGimbalRotationMode = false;
 	
 private:
-	/** Debug attributes to display the pull direction */
-	bool bDebugRotate = false;
-	FVector DebugDirection = FVector::ZeroVector;
-	FVector DebugClosest = FVector::ZeroVector;
-	FVector DebugNormalRemoved = FVector::ZeroVector;
-	FVector DebugNormalSkip = FVector::ZeroVector;
+
+	struct FGizmoDebugData
+	{
+		/** Determines whether certain data is displayed, ie. drag operation deltas. */
+		bool bIsEditing = false;
+
+		FTransform TransformStart;
+		FTransform TransformCurrent;
+
+		FTransform InteractionStart;
+		FTransform InteractionCurrent;
+		FVector InteractionPlaneNormal;
+
+		/** Debug attributes to display the pull direction */
+		bool bDebugRotate = false;
+		FVector DebugDirection = FVector::ZeroVector;
+		FVector DebugClosest = FVector::ZeroVector;
+		FVector DebugNormalRemoved = FVector::ZeroVector;
+		FVector DebugNormalSkip = FVector::ZeroVector;
+		double InteractionAngleStart;
+		double InteractionAngleCurrent;
+		double InteractionRadius;
+	} DebugData;
 };
