[Drag Tools] Separate ITF Drag Tools from legacy ones

Separating the the ITF drag tools from the legacy ones, and placing them in the EditorInteractiveToolsFramework, since in the future they will only be used via ITF, and the legacy ones removed.

- Adding new ITF only Drag tools base class and Frustum Select class

- Moving UDragToolsBehaviorSource to EditorInteractiveToolsFramework

diff --git a/Engine/Source/Editor/Experimental/EditorInteractiveToolsFramework/Private/EditorDragTools/DragToolsBehaviourSource.cpp b/Engine/Source/Editor/Experimental/EditorInteractiveToolsFramework/Private/EditorDragTools/DragToolsBehaviourSource.cpp
new file mode 100644
--- /dev/null
+++ b/Engine/Source/Editor/Experimental/EditorInteractiveToolsFramework/Private/EditorDragTools/DragToolsBehaviourSource.cpp
@@ -0,0 +1,156 @@
+// Copyright Epic Games, Inc. All Rights Reserved.
+
+#include "EditorDragTools/DragToolsBehaviourSource.h"
+#include "BaseBehaviors/ClickDragBehavior.h"
+#include "BaseBehaviors/KeyAsModifierInputBehavior.h"
+#include "EditorDragTools/EditorDragToolBehaviorTarget.h"
+#include "EditorDragTools/EditorFrustumSelect.h"
+#include "EditorModeManager.h"
+#include "Tools/EdModeInteractiveToolsContext.h"
+
+void UDragToolsBehaviorSource::Initialize(const UEditorInteractiveToolsContext* InInteractiveToolsContext)
+{
+	if (!InInteractiveToolsContext)
+	{
+		return;
+	}
+
+	TObjectPtr<UInteractiveToolManager> ToolManager = InInteractiveToolsContext->ToolManager;
+	if (!ToolManager)
+	{
+		return;
+	}
+
+	IToolsContextQueriesAPI* ContextQueriesAPI = ToolManager->GetContextQueriesAPI();
+	if (!ContextQueriesAPI)
+	{
+		return;
+	}
+
+	FEditorViewportClient* EditorViewportClient = nullptr;
+
+	if (const FEditorModeTools* const ModeManager = InInteractiveToolsContext->GetParentEditorModeManager())
+	{
+		EditorViewportClient = ModeManager->GetFocusedViewportClient();
+	}
+
+	if (!EditorViewportClient)
+	{
+		FViewport* Viewport = ContextQueriesAPI->GetFocusedViewport();
+		if (!Viewport)
+		{
+			return;
+		}
+
+		EditorViewportClient = static_cast<FEditorViewportClient*>(Viewport->GetClient());
+		if (!EditorViewportClient)
+		{
+			return;
+		}
+	}
+
+	BehaviorSet = NewObject<UInputBehaviorSet>();
+
+	// In case New TRS Gizmos are enabled, we need Drag Tools to be processed before the Gizmo ones, otherwise Drag Behaviors will never be triggered.
+	// See:
+	// * UTransformGizmo::SetupBehaviors()			[DEFAULT_GIZMO_PRIORITY]
+	// * UTransformGizmo::SetupIndirectBehaviors()	[DEFAULT_GIZMO_PRIORITY-1]
+	static constexpr int DEFAULT_DRAG_TOOLS_PRIORITY = FInputCapturePriority::DEFAULT_GIZMO_PRIORITY-2;
+
+	// Frustum Select
+	{
+		TSharedPtr<FEditorDragToolBehaviorTarget> FrustumSelectTool;
+
+		// Some Drag Tools require or can make use specifically of Level Viewport Client.
+		if (EditorViewportClient->IsLevelEditorClient())
+		{
+			// Initialize using Level Viewport Client if available
+			FLevelEditorViewportClient* LevelViewportClient = static_cast<FLevelEditorViewportClient*>(EditorViewportClient);
+			FrustumSelectTool = MakeShared<FEditorFrustumSelect>(LevelViewportClient);
+		}
+		else
+		{
+			FrustumSelectTool = MakeShared<FEditorFrustumSelect>(EditorViewportClient);
+		}
+
+		DragTools.Add(FrustumSelectTool);
+
+		// Add drag input behavior
+		UClickDragInputBehavior* MouseBehavior = NewObject<UClickDragInputBehavior>();
+		MouseBehavior->Initialize(FrustumSelectTool.Get());
+		MouseBehavior->SetDefaultPriority(DEFAULT_DRAG_TOOLS_PRIORITY);
+		BehaviorSet->Add(MouseBehavior);
+
+		constexpr int ShiftKeyMod = 1;
+		constexpr int AltKeyMod = 2;
+		constexpr int CtrlKeyMod = 3;
+
+		// Register modifier listeners
+		UKeyAsModifierInputBehavior* KeyModifierInputBehavior = NewObject<UKeyAsModifierInputBehavior>();
+		KeyModifierInputBehavior->Initialize(FrustumSelectTool.Get(), ShiftKeyMod, FInputDeviceState::IsShiftKeyDown);
+		KeyModifierInputBehavior->Initialize(FrustumSelectTool.Get(), CtrlKeyMod, FInputDeviceState::IsCtrlKeyDown);
+		KeyModifierInputBehavior->Initialize(FrustumSelectTool.Get(), AltKeyMod, FInputDeviceState::IsAltKeyDown);
+		KeyModifierInputBehavior->SetDefaultPriority(DEFAULT_DRAG_TOOLS_PRIORITY);
+		BehaviorSet->Add(KeyModifierInputBehavior);
+
+		FrustumSelectTool->OnActivateTool().AddUObject(
+			this, &UDragToolsBehaviorSource::OnActivateTool, EDragTool::FrustumSelect
+		);
+		FrustumSelectTool->OnDeactivateTool().AddUObject(
+			this, &UDragToolsBehaviorSource::OnDeactivateTool, EDragTool::FrustumSelect
+		);
+	}
+
+	// Add other drag tools
+}
+
+FEditorDragToolBehaviorTarget* UDragToolsBehaviorSource::GetActiveTool() const
+{
+	if (ActiveToolType > 0)
+	{
+		int32 ToolIndex = GetToolIndex(static_cast<EDragTool::Type>(ActiveToolType));
+		if (DragTools.IsValidIndex(ToolIndex))
+		{
+			if (const TSharedPtr<FEditorDragToolBehaviorTarget>& ActiveTool = DragTools[ToolIndex])
+			{
+				return ActiveTool.Get();
+			}
+		}
+	}
+
+	return nullptr;
+}
+
+int32 UDragToolsBehaviorSource::GetToolIndex(EDragTool::Type InDragToolType)
+{
+	static TArray<int32> ToolIndexes = { EDragTool::FrustumSelect };
+
+	int32 ToolIndex = ToolIndexes.Find(InDragToolType);
+	if (ToolIndex != INDEX_NONE)
+	{
+		return ToolIndex;
+	}
+
+	return -1;
+}
+
+void UDragToolsBehaviorSource::OnActivateTool(EDragTool::Type InDragToolType)
+{
+	ActiveToolType = InDragToolType;
+}
+
+void UDragToolsBehaviorSource::OnDeactivateTool(EDragTool::Type InDragToolType)
+{
+	if (ActiveToolType == InDragToolType)
+	{
+		ActiveToolType = -1;
+	}
+}
+
+void UDragToolsBehaviorSource::RenderTools(const FSceneView* InSceneView, FCanvas* InCanvas) const
+{
+	if (FEditorDragToolBehaviorTarget* ActiveDragTool = GetActiveTool())
+	{
+		ActiveDragTool->Render(InSceneView, InCanvas);
+	}
+}
diff --git a/Engine/Source/Editor/Experimental/EditorInteractiveToolsFramework/Private/EditorDragTools/EditorDragToolBehaviorTarget.cpp b/Engine/Source/Editor/Experimental/EditorInteractiveToolsFramework/Private/EditorDragTools/EditorDragToolBehaviorTarget.cpp
new file mode 100644
--- /dev/null
+++ b/Engine/Source/Editor/Experimental/EditorInteractiveToolsFramework/Private/EditorDragTools/EditorDragToolBehaviorTarget.cpp
@@ -0,0 +1,133 @@
+﻿// Copyright Epic Games, Inc. All Rights Reserved.
+
+#include "EditorDragTools/EditorDragToolBehaviorTarget.h"
+#include "InputCoreTypes.h"
+#include "EditorViewportClient.h"
+#include "Editor.h"
+#include "EditorModeManager.h"
+
+namespace UE::Editor::DragTools
+{
+
+// CVar initializer
+static int32 UseITFTools = 0;
+static FAutoConsoleVariableRef CVarEnableITFTools(
+	TEXT("DragTools.EnableITFTools"),
+	UseITFTools,
+	TEXT("Is the ITF version of Drag Tools enabled?\nCurrently supports: Frustum Select."),
+	FConsoleVariableDelegate::CreateLambda(
+		[](const IConsoleVariable* InVariable)
+		{
+			if (UseITFTools)
+			{
+				OnEditorDragToolsActivated().Broadcast();
+			}
+			else
+			{
+				OnEditorDragToolsDeactivated().Broadcast();
+			}
+		}
+	)
+);
+
+bool UseEditorDragTools()
+{
+	return UseITFTools == 1;
+}
+
+FOnEditorDragToolsToggleDelegate& OnEditorDragToolsActivated()
+{
+	static FOnEditorDragToolsToggleDelegate OnDragToolsActivated;
+	return OnDragToolsActivated;
+}
+
+FOnEditorDragToolsToggleDelegate& OnEditorDragToolsDeactivated()
+{
+	static FOnEditorDragToolsToggleDelegate OnDragToolsDeactivated;
+	return OnDragToolsDeactivated;
+}
+
+} // namespace UE::Editor::DragTools
+
+///////////////////////////////////////////////////////////////////////////////////////////////////////////////////////
+//
+// FEditorDragTools
+//
+///////////////////////////////////////////////////////////////////////////////////////////////////////////////////////
+
+FEditorDragToolBehaviorTarget::FEditorDragToolBehaviorTarget(FEditorModeTools* InModeTools)
+	: bConvertDelta(true)
+	, ModeTools(InModeTools)
+	, Start(FVector::ZeroVector)
+	, End(FVector::ZeroVector)
+	, bUseSnapping(false)
+	, bIsDragging(false)
+{
+}
+
+void FEditorDragToolBehaviorTarget::OnUpdateModifierState(int InModifierID, bool bInIsOn)
+{
+	constexpr int ShiftKeyMod = 1;
+	constexpr int AltKeyMod = 2;
+	constexpr int CtrlKeyMod = 3;
+
+	switch (InModifierID)
+	{
+	case ShiftKeyMod:
+		InputState.bShiftKeyDown = bInIsOn;
+		break;
+	case CtrlKeyMod:
+		InputState.bCtrlKeyDown = bInIsOn;
+		break;
+	case AltKeyMod:
+		InputState.bAltKeyDown = bInIsOn;
+		break;
+	default:;
+	}
+}
+
+void FEditorDragToolBehaviorTarget::OnTerminateDragSequence()
+{
+	Start = End = FVector::ZeroVector;
+	bIsDragging = false;
+
+	InputState.SetModifierKeyStates(false, false, false, false);
+}
+
+bool FEditorDragToolBehaviorTarget::IsActivationChordPressed(const FInputChord& InChord) const
+{
+	bool bIsChordPressed = true;
+
+	if (InChord.NeedsControl())
+	{
+		bIsChordPressed = bIsChordPressed && FInputDeviceState::IsCtrlKeyDown(InputState);
+	}
+
+	if (InChord.NeedsAlt())
+	{
+		bIsChordPressed = bIsChordPressed && FInputDeviceState::IsAltKeyDown(InputState);
+	}
+
+	if (InChord.NeedsShift())
+	{
+		bIsChordPressed = bIsChordPressed && FInputDeviceState::IsShiftKeyDown(InputState);
+	}
+
+	return bIsChordPressed;
+}
+
+bool FEditorDragToolBehaviorTarget::IsCurrentModeSupported()
+{
+	if (ModeTools)
+	{
+		for (const FEditorModeID UnsupportedMode : GetUnsupportedModes())
+		{
+			if (!ModeTools->EnsureNotInMode(UnsupportedMode))
+			{
+				return false;
+			}
+		}
+	}
+
+	return true;
+}
diff --git a/Engine/Source/Editor/Experimental/EditorInteractiveToolsFramework/Private/EditorDragTools/EditorFrustumSelect.cpp b/Engine/Source/Editor/Experimental/EditorInteractiveToolsFramework/Private/EditorDragTools/EditorFrustumSelect.cpp
new file mode 100644
--- /dev/null
+++ b/Engine/Source/Editor/Experimental/EditorInteractiveToolsFramework/Private/EditorDragTools/EditorFrustumSelect.cpp
@@ -0,0 +1,407 @@
+﻿// Copyright Epic Games, Inc. All Rights Reserved.
+
+#include "EditorDragTools/EditorFrustumSelect.h"
+
+#include "CanvasItem.h"
+#include "CanvasTypes.h"
+#include "Components/PrimitiveComponent.h"
+#include "Editor/UnrealEdEngine.h"
+#include "EditorModeManager.h"
+#include "EditorModes.h"
+#include "Elements/Framework/EngineElementsLibrary.h"
+#include "Elements/Framework/TypedElementRegistry.h"
+#include "Elements/Framework/TypedElementSelectionSet.h"
+#include "Elements/Interfaces/TypedElementWorldInterface.h"
+#include "EngineUtils.h"
+#include "GameFramework/Volume.h"
+#include "LevelEditorSubsystem.h"
+#include "Model.h"
+#include "SceneView.h"
+#include "ScopedTransaction.h"
+#include "Settings/LevelEditorViewportSettings.h"
+#include "SnappingUtils.h"
+#include "Subsystems/BrushEditingSubsystem.h"
+#include "UnrealEdGlobals.h"
+
+namespace UE::EditorDragTools::Private
+{
+TArray<FTypedElementHandle> GetElementsIntersectingFrustum(
+	const AActor* Actor,
+	const FConvexVolume& InFrustum,
+	const FEditorViewportClient* EditorViewport,
+	const FLevelEditorViewportClient* LevelViewport,
+	const FWorldSelectionElementArgs& SelectionArgs
+)
+{
+	if (Actor
+		&& (!EditorViewport || !Actor->IsA(AVolume::StaticClass())
+			|| (LevelViewport ? !LevelViewport->IsVolumeVisibleInViewport(*Actor) : false)))
+	{
+		if (FTypedElementHandle ActorHandle = UEngineElementsLibrary::AcquireEditorActorElementHandle(Actor))
+		{
+			if (TTypedElement<ITypedElementWorldInterface> WorldElement =
+					UTypedElementRegistry::GetInstance()->GetElement<ITypedElementWorldInterface>(ActorHandle))
+			{
+				return WorldElement.GetSelectionElementsInConvexVolume(InFrustum, SelectionArgs);
+			}
+		}
+	}
+
+	return {};
+}
+
+} // namespace UE::EditorDragTools::Private
+
+void FEditorFrustumSelect::Render(const FSceneView* View, FCanvas* Canvas)
+{
+	if (!Canvas || !View)
+	{
+		return;
+	}
+
+	FCanvasBoxItem BoxItem(
+		FVector2D(Start.X, Start.Y) / Canvas->GetDPIScale(),
+		FVector2D(End.X - Start.X, End.Y - Start.Y) / Canvas->GetDPIScale()
+	);
+
+	BoxItem.SetColor(FLinearColor::White);
+	Canvas->DrawItem(BoxItem);
+}
+
+FInputRayHit FEditorFrustumSelect::CanBeginClickDragSequence(const FInputDeviceRay& InPressPos)
+{
+	// Todo: this could be retrieved from a command for customization
+	FInputChord ActivationChord(EModifierKey::Control | EModifierKey::Alt, EKeys::LeftMouseButton);
+
+	return IsActivationChordPressed(ActivationChord)
+				&& IsCurrentModeSupported()
+			 ? FInputRayHit(TNumericLimits<float>::Max()) // bHit is true. Depth is max to lose the standard tiebreaker.
+			 : FInputRayHit();
+}
+
+void FEditorFrustumSelect::OnClickPress(const FInputDeviceRay& InPressPos)
+{
+	if (EditorViewportClient)
+	{
+		// Signal that this tool is now active
+		OnActivateTool().Broadcast();
+
+		Start = InPressPos.WorldRay.Origin;
+		bIsDragging = true;
+
+		// Snap to constraints.
+		if (bUseSnapping)
+		{
+			const float GridSize = GEditor->GetGridSize();
+			const FVector GridBase(GridSize, GridSize, GridSize);
+			FSnappingUtils::SnapPointToGrid(Start, GridBase);
+		}
+		End = Start;
+
+		// Remove any active hover objects
+		FLevelEditorViewportClient::ClearHoverFromObjects();
+
+		Start = FVector(InPressPos.ScreenPosition.X, InPressPos.ScreenPosition.Y, 0);
+		End = Start;
+	}
+}
+
+void FEditorFrustumSelect::OnClickDrag(const FInputDeviceRay& InDragPos)
+{
+	End = FVector(InDragPos.ScreenPosition.X, InDragPos.ScreenPosition.Y, 0);
+}
+
+void FEditorFrustumSelect::OnClickRelease(const FInputDeviceRay& InReleasePos)
+{
+	End = FVector(InReleasePos.ScreenPosition.X, InReleasePos.ScreenPosition.Y, 0);
+
+	if (!EditorViewportClient)
+	{
+		return;
+	}
+
+	FViewport* Viewport = EditorViewportClient->Viewport;
+	if (!Viewport)
+	{
+		return;
+	}
+
+	UBrushEditingSubsystem* BrushSubsystem = GEditor->GetEditorSubsystem<UBrushEditingSubsystem>();
+	const bool bGeometryMode = BrushSubsystem ? BrushSubsystem->IsGeometryEditorModeActive() : false;
+
+	FSceneViewFamilyContext ViewFamily(
+		FSceneViewFamily::ConstructionValues(Viewport, EditorViewportClient->GetScene(), EditorViewportClient->EngineShowFlags)
+	);
+	FSceneView* SceneView = EditorViewportClient->CalcSceneView(&ViewFamily);
+
+	// Generate a frustum out of the dragged box
+	FConvexVolume Frustum;
+	CalculateFrustum(SceneView, Frustum, true);
+
+	FScopedTransaction Transaction(NSLOCTEXT("ActorFrustumSelect", "MarqueeSelectTransaction", "Marquee Select"));
+
+	if (!FInputDeviceState::IsShiftKeyDown(InputState))
+	{
+		// If the user is selecting, but isn't hold down SHIFT, remove all current selections.
+		ModeTools->SelectNone();
+	}
+
+	// Does an actor have to be fully contained in the box to be selected
+	const bool bStrictDragSelection = GetDefault<ULevelEditorViewportSettings>()->bStrictBoxSelection;
+
+	constexpr bool bShouldSelect = true; // bLeftMouseButtonDown
+	// Let the editor mode try to handle the selection.
+	const bool bEditorModeHandledSelection = ModeTools->FrustumSelect(Frustum, EditorViewportClient, bShouldSelect);
+
+	// Let the component visualizers try to handle the selection.
+	const bool bComponentVisHandledSelection =
+		!bEditorModeHandledSelection
+		&& GUnrealEd->ComponentVisManager.HandleFrustumSelect(Frustum, EditorViewportClient, Viewport);
+
+	if (!bEditorModeHandledSelection && !bComponentVisHandledSelection)
+	{
+		UTypedElementSelectionSet* SelectionSet = GEditor->GetEditorSubsystem<ULevelEditorSubsystem>()->GetSelectionSet();
+		SelectionSet->Modify();
+
+		FWorldSelectionElementArgs SeletionArgs{ SelectionSet,
+												 ETypedElementSelectionMethod::Primary,
+												 FTypedElementSelectionOptions(),
+												 &(EditorViewportClient->EngineShowFlags),
+												 bStrictDragSelection,
+												 bGeometryMode };
+
+		const int32 ViewportSizeX = Viewport->GetSizeXY().X;
+		const int32 ViewportSizeY = Viewport->GetSizeXY().Y;
+
+		if (Start.X > End.X)
+		{
+			Swap(Start.X, End.X);
+		}
+
+		if (Start.Y > End.Y)
+		{
+			Swap(Start.Y, End.Y);
+		}
+
+		TArray<FTypedElementHandle> ElementsToSelect;
+		const bool bTransparentBoxSelection = GetDefault<ULevelEditorViewportSettings>()->bTransparentBoxSelection;
+		if (bTransparentBoxSelection)
+		{
+			// Get a list of frustum-culled actors
+			for (FActorIterator It(EditorViewportClient->GetWorld()); It; ++It)
+			{
+				AActor* Actor = *It;
+				ElementsToSelect.Append(
+					UE::EditorDragTools::Private::GetElementsIntersectingFrustum(
+						Actor, Frustum, EditorViewportClient, LevelViewportClient, SeletionArgs
+					)
+				);
+			}
+		}
+		else
+		{
+			// Extend the endpoint of the rect to get the actual line
+
+			const int32 MinX = UE::LWC::FloatToIntCastChecked<int32>(FMath::Max<double>(0.0, Start.X));
+			const int32 MinY = UE::LWC::FloatToIntCastChecked<int32>(FMath::Max<double>(0.0, Start.Y));
+			const int32 MaxX = FMath::Min(ViewportSizeX, FMath::TruncToInt32(End.X + 1.0));
+			const int32 MaxY = FMath::Min(ViewportSizeY, FMath::TruncToInt32(End.Y + 1.0));
+
+			const FIntPoint Min{ MinX, MinY };
+			const FIntPoint Max{ MaxX, MaxY };
+			const FIntRect BoxRect{ Min, Max };
+
+			// Typed Element selection
+			{
+				FTypedElementListRef ElementList = UTypedElementRegistry::GetInstance()->CreateElementList();
+				Viewport->GetElementHandlesInRect(BoxRect, ElementList);
+
+				if (bStrictDragSelection)
+				{
+					ElementList->ForEachElement<ITypedElementWorldInterface>(
+						[bStrictDragSelection, &Frustum, &SelectionSet, &ElementsToSelect](
+							const TTypedElement<ITypedElementWorldInterface>& InElement
+						)
+						{
+							if (InElement.IsElementInConvexVolume(Frustum, bStrictDragSelection))
+							{
+								ElementsToSelect.Add(
+									SelectionSet->GetSelectionElement(InElement, ETypedElementSelectionMethod::Primary)
+								);
+							}
+
+							return true;
+						}
+					);
+				}
+				else
+				{
+					// Grab only the selectable handles (this remove the components from the selection and select the actor instead)
+					ElementList->ForEachElementHandle(
+						[&SelectionSet, &ElementsToSelect](const FTypedElementHandle& InHandle)
+						{
+							ElementsToSelect.Add(
+								SelectionSet->GetSelectionElement(InHandle, ETypedElementSelectionMethod::Primary)
+							);
+							return true;
+						}
+					);
+				}
+			}
+
+			// We need this old code to support the BSP
+			TSet<AActor*> BSPActors;
+			TSet<UModel*> HitModels;
+			Viewport->GetActorsAndModelsInHitProxy(BoxRect, BSPActors, HitModels);
+			BSPActors.Empty(HitModels.Num());
+
+			if (HitModels.Num() > 0)
+			{
+				// Check every model to see if its BSP surfaces should be selected
+				for (auto It = HitModels.CreateConstIterator(); It; ++It)
+				{
+					UModel& Model = **It;
+					// Check every node in the model
+					for (int32 NodeIndex = 0; NodeIndex < Model.Nodes.Num(); NodeIndex++)
+					{
+						if (IntersectsFrustum(Model, NodeIndex, Frustum, bStrictDragSelection))
+						{
+							uint32 SurfaceIndex = Model.Nodes[NodeIndex].iSurf;
+							FBspSurf& Surf = Model.Surfs[SurfaceIndex];
+							BSPActors.Add(Surf.Actor);
+						}
+					}
+				}
+			}
+
+			if (BSPActors.Num() > 0)
+			{
+				for (auto It = BSPActors.CreateConstIterator(); It; ++It)
+				{
+					AActor* Actor = *It;
+					if (bStrictDragSelection)
+					{
+						ElementsToSelect.Append(
+							UE::EditorDragTools::Private::GetElementsIntersectingFrustum(
+								Actor, Frustum, EditorViewportClient, LevelViewportClient, SeletionArgs
+							)
+						);
+					}
+					else
+					{
+						ElementsToSelect.Add(UEngineElementsLibrary::AcquireEditorActorElementHandle(Actor));
+					}
+				}
+			}
+		}
+
+		FTypedElementSelectionOptions ElementSelectionOption;
+
+		if (!FInputDeviceState::IsShiftKeyDown(InputState))
+		{
+			// If the user is selecting, but isn't hold down SHIFT, remove the previous selections.
+			SelectionSet->SetSelection(MoveTemp(ElementsToSelect), ElementSelectionOption);
+		}
+		else
+		{
+			SelectionSet->SelectElements(MoveTemp(ElementsToSelect), ElementSelectionOption);
+		}
+	}
+
+	// Clear any hovered objects that might have been created while dragging
+	FLevelEditorViewportClient::ClearHoverFromObjects();
+
+	FEditorFrustumSelect::OnTerminateDragSequence();
+
+	// Signal that this tool is no longer active
+	OnDeactivateTool().Broadcast();
+}
+
+void FEditorFrustumSelect::OnTerminateDragSequence()
+{
+	FEditorDragToolBehaviorTarget::OnTerminateDragSequence();
+
+	Start = End = FVector::ZeroVector;
+}
+
+TArray<FEditorModeID> FEditorFrustumSelect::GetUnsupportedModes()
+{
+	const TArray<FEditorModeID> UnsupportedModes = { FBuiltinEditorModes::EM_Landscape, FBuiltinEditorModes::EM_Foliage };
+	return UnsupportedModes;
+}
+
+bool FEditorFrustumSelect::IntersectsFrustum(
+	const UModel& InModel, int32 NodeIndex, const FConvexVolume& InFrustum, bool bUseStrictSelection
+)
+{
+	FBox NodeBB;
+	// Get a bounding box of the node being checked
+	InModel.GetNodeBoundingBox(InModel.Nodes[NodeIndex], NodeBB);
+
+	bool bFullyContained = false;
+
+	// Does the box intersect the frustum
+	bool bIntersects = InFrustum.IntersectBox(NodeBB.GetCenter(), NodeBB.GetExtent(), bFullyContained);
+
+	return bIntersects && (!bUseStrictSelection || (bUseStrictSelection && bFullyContained));
+}
+
+void FEditorFrustumSelect::CalculateFrustum(const FSceneView* InView, FConvexVolume& OutFrustum, bool bUseBoxFrustum) const
+{
+	if (!InView || !EditorViewportClient)
+	{
+		return;
+	}
+
+	if (bUseBoxFrustum)
+	{
+		FVector CamPoint = EditorViewportClient->GetViewLocation();
+		FVector BoxPoint1, BoxPoint2, BoxPoint3, BoxPoint4;
+		FVector WorldDir1, WorldDir2, WorldDir3, WorldDir4;
+		// extend the 2D box of 1 pixel if needed to avoid degenerated volume
+		const double DX = Start.X == End.X ? 0.5 : 0.0;
+		const double Left = FMath::Min(Start.X, End.X) - DX;
+		const double Right = FMath::Max(Start.X, End.X) + DX;
+		const double DY = Start.Y == End.Y ? 0.5 : 0.0;
+		const double Bottom = FMath::Min(Start.Y, End.Y) - DY;
+		const double Top = FMath::Max(Start.Y, End.Y) + DY;
+		// Deproject the four corners of the selection box
+		FVector2D Point1(Left, Bottom); // Upper Left Corner
+		FVector2D Point2(Right, Bottom); // Upper Right Corner
+		FVector2D Point3(Right, Top); // Lower Right Corner
+		FVector2D Point4(Left, Top); // Lower Left Corner
+		InView->DeprojectFVector2D(Point1, BoxPoint1, WorldDir1);
+		InView->DeprojectFVector2D(Point2, BoxPoint2, WorldDir2);
+		InView->DeprojectFVector2D(Point3, BoxPoint3, WorldDir3);
+		InView->DeprojectFVector2D(Point4, BoxPoint4, WorldDir4);
+		// Use the camera position and the selection box to create the bounding planes
+		FPlane TopPlane(BoxPoint1, BoxPoint2, CamPoint); // Top Plane
+		FPlane RightPlane(BoxPoint2, BoxPoint3, CamPoint); // Right Plane
+		FPlane BottomPlane(BoxPoint3, BoxPoint4, CamPoint); // Bottom Plane
+		FPlane LeftPlane(BoxPoint4, BoxPoint1, CamPoint); // Left Plane
+
+		// Try to get all six planes to create a frustum.
+		// The frustum is built with the first four planes corresponding to the sides of the frustum.
+		FPlane NearPlane;
+		FPlane FarPlane;
+		OutFrustum.Planes.Empty();
+		OutFrustum.Planes.Add(TopPlane);
+		OutFrustum.Planes.Add(RightPlane);
+		OutFrustum.Planes.Add(BottomPlane);
+		OutFrustum.Planes.Add(LeftPlane);
+		if (InView->ViewMatrices.GetViewProjectionMatrix().GetFrustumNearPlane(NearPlane))
+		{
+			OutFrustum.Planes.Add(NearPlane);
+		}
+		if (InView->ViewMatrices.GetViewProjectionMatrix().GetFrustumFarPlane(FarPlane))
+		{
+			OutFrustum.Planes.Add(FarPlane);
+		}
+		OutFrustum.Init();
+	}
+	else
+	{
+		OutFrustum = InView->ViewFrustum;
+		OutFrustum.Init();
+	}
+}
diff --git a/Engine/Source/Editor/Experimental/EditorInteractiveToolsFramework/Public/EditorDragTools/DragToolsBehaviourSource.h b/Engine/Source/Editor/Experimental/EditorInteractiveToolsFramework/Public/EditorDragTools/DragToolsBehaviourSource.h
new file mode 100644
--- /dev/null
+++ b/Engine/Source/Editor/Experimental/EditorInteractiveToolsFramework/Public/EditorDragTools/DragToolsBehaviourSource.h
@@ -0,0 +1,67 @@
+// Copyright Epic Games, Inc. All Rights Reserved.
+
+#pragma once
+
+#include "EditorViewportClient.h"
+#include "InputBehaviorSet.h"
+#include "UObject/Object.h"
+#include "DragToolsBehaviourSource.generated.h"
+
+class FEditorDuplicateDragSelection;
+class FEditorDragToolBehaviorTarget;
+class UEditorInteractiveToolsContext;
+
+/**
+ * Hosts Drag Tools and needed behaviors.
+ * Handles an Input Behavior Set, and keeps track of currently active Drag tool.
+ */
+UCLASS(Transient, MinimalAPI)
+class UDragToolsBehaviorSource final : public UObject, public IInputBehaviorSource
+{
+	GENERATED_BODY()
+
+public:
+	//~ Begin IInputBehaviorSource
+	EDITORINTERACTIVETOOLSFRAMEWORK_API virtual const UInputBehaviorSet* GetInputBehaviors() const override { return BehaviorSet; }
+	//~ End IClickDragBehaviorTarget
+
+	/**
+	 * Creates and stores the drag tools.
+	 * Instantiates the BehaviorSet hosting all behaviors required by drag tools.
+	 *
+	 * @param InInteractiveToolsContext: used to retrieve the viewport client used by drag tools. Can be gathered e.g.
+	 * from UEdMode::GetInteractiveToolsContext.
+	 */
+	EDITORINTERACTIVETOOLSFRAMEWORK_API void Initialize(const UEditorInteractiveToolsContext* InInteractiveToolsContext);
+
+	/**
+	 * Renders the active tool on the specified View/Canvas
+	 */
+	EDITORINTERACTIVETOOLSFRAMEWORK_API void RenderTools(const FSceneView* InSceneView, FCanvas* InCanvas) const;
+private:
+	static int32 GetToolIndex(EDragTool::Type InDragToolType);
+
+	void OnActivateTool(EDragTool::Type InDragToolType);
+	void OnDeactivateTool(EDragTool::Type InDragToolType);
+
+	/**
+	 * Returns the drag tool currently being used, if any
+	 */
+	FEditorDragToolBehaviorTarget* GetActiveTool() const;
+
+	/**
+	 * Hosting drag tools behaviors
+	 */
+	UPROPERTY(Transient)
+	TObjectPtr<UInputBehaviorSet> BehaviorSet;
+
+	/**
+	 * Available Drag Tools, access by using the index provided by GetToolIndex(EDragTool::Type)
+	 */
+	TArray<TSharedPtr<FEditorDragToolBehaviorTarget>> DragTools;
+
+	/**
+	 * Which type of Drag tool is currently active
+	 */
+	int ActiveToolType = -1;
+};
diff --git a/Engine/Source/Editor/Experimental/EditorInteractiveToolsFramework/Public/EditorDragTools/EditorDragToolBehaviorTarget.h b/Engine/Source/Editor/Experimental/EditorInteractiveToolsFramework/Public/EditorDragTools/EditorDragToolBehaviorTarget.h
new file mode 100644
--- /dev/null
+++ b/Engine/Source/Editor/Experimental/EditorInteractiveToolsFramework/Public/EditorDragTools/EditorDragToolBehaviorTarget.h
@@ -0,0 +1,123 @@
+﻿// Copyright Epic Games, Inc. All Rights Reserved.
+
+#pragma once
+
+#include "BaseBehaviors/BehaviorTargetInterfaces.h"
+#include "Editor.h"
+
+class FCanvas;
+class FEditorModeTools;
+class FEditorViewportClient;
+class FPrimitiveDrawInterface;
+class FSceneView;
+class FUICommandInfo;
+
+/**
+ * ITF related helper functions
+ */
+namespace UE::Editor::DragTools
+{
+
+DECLARE_MULTICAST_DELEGATE(FOnEditorDragToolsToggleDelegate)
+
+/**
+ * Returns true if ITF-based drag tools should be used.
+ * If an ITF version of a drag tool is not available yet, legacy will be used.
+ */
+EDITORINTERACTIVETOOLSFRAMEWORK_API bool UseEditorDragTools();
+
+EDITORINTERACTIVETOOLSFRAMEWORK_API FOnEditorDragToolsToggleDelegate& OnEditorDragToolsActivated();
+EDITORINTERACTIVETOOLSFRAMEWORK_API FOnEditorDragToolsToggleDelegate& OnEditorDragToolsDeactivated();
+
+} // namespace UE::Editor::DragTools
+
+/**
+ * The base class that all drag tools inherit from.
+ * The drag tools implement special behaviors for the user clicking and dragging in a viewport.
+ *
+ * Inheriting from IClickDragBehaviorTarget so that Drag Tools logic can be shared between legacy and ITF usage of the same tools
+ */
+class FEditorDragToolBehaviorTarget : public IClickDragBehaviorTarget
+{
+public:
+	EDITORINTERACTIVETOOLSFRAMEWORK_API FEditorDragToolBehaviorTarget(FEditorModeTools* InModeTools);
+	virtual ~FEditorDragToolBehaviorTarget() override
+	{
+	}
+
+	virtual void Render(const FSceneView* InView, FCanvas* InCanvas)
+	{
+	}
+
+	/**
+	 * Rendering stub for 2D viewport drag tools.
+	 */
+	virtual void Render(FCanvas* Canvas)
+	{
+	}
+
+	/** @return true if we are dragging */
+	bool IsDragging() const
+	{
+		return bIsDragging;
+	}
+
+	/** Does this drag tool need to have the mouse movement converted to the viewport orientation? */
+	bool bConvertDelta;
+
+	//~ Begin IModifierToggleBehaviorTarget
+	virtual void OnUpdateModifierState(int InModifierID, bool bInIsOn) override;
+	//~ End IModifierToggleBehaviorTarget
+
+	//~ Begin IClickDragBehaviorTarget
+	virtual FInputRayHit CanBeginClickDragSequence(const FInputDeviceRay& InPressPos) override
+	{
+		return FInputRayHit();
+	}
+
+	virtual void OnClickPress(const FInputDeviceRay& InPressPos) override
+	{
+	}
+
+	virtual void OnClickDrag(const FInputDeviceRay& InDragPos) override
+	{
+	}
+
+	virtual void OnClickRelease(const FInputDeviceRay& InReleasePos) override
+	{
+	}
+
+	virtual void OnTerminateDragSequence() override;
+	//~ End IClickDragBehaviorTarget
+
+	bool IsActivationChordPressed(const FInputChord& InChord) const;
+
+	DECLARE_MULTICAST_DELEGATE(FOnToolStateChange);
+	FOnToolStateChange& OnActivateTool()
+	{
+		return OnToolActivatedDelegate;
+	}
+	FOnToolStateChange& OnDeactivateTool()
+	{
+		return OnToolDeactivatedDelegate;
+	}
+
+protected:
+	virtual TArray<FEditorModeID> GetUnsupportedModes() { return {}; }
+
+	bool IsCurrentModeSupported();
+
+	FEditorModeTools* ModeTools;
+
+	/** The start/end location of the current drag. */
+	FVector Start, End;
+
+	/** If true, the drag tool wants to be passed grid snapped values. */
+	bool bUseSnapping;
+
+	FInputDeviceState InputState;
+	bool bIsDragging;
+
+	FOnToolStateChange OnToolActivatedDelegate;
+	FOnToolStateChange OnToolDeactivatedDelegate;
+};
diff --git a/Engine/Source/Editor/Experimental/EditorInteractiveToolsFramework/Public/EditorDragTools/EditorFrustumSelect.h b/Engine/Source/Editor/Experimental/EditorInteractiveToolsFramework/Public/EditorDragTools/EditorFrustumSelect.h
new file mode 100644
--- /dev/null
+++ b/Engine/Source/Editor/Experimental/EditorInteractiveToolsFramework/Public/EditorDragTools/EditorFrustumSelect.h
@@ -0,0 +1,67 @@
+﻿// Copyright Epic Games, Inc. All Rights Reserved.
+
+#pragma once
+
+#include "EditorDragToolBehaviorTarget.h"
+#include "LevelEditorViewport.h"
+
+class FCanvas;
+class UModel;
+
+/**
+ * Draws a box in the current viewport and when the mouse button is released,
+ * it selects/unselects everything inside of it.
+ */
+class FEditorFrustumSelect : public FEditorDragToolBehaviorTarget
+{
+public:
+	EDITORINTERACTIVETOOLSFRAMEWORK_API explicit FEditorFrustumSelect(FLevelEditorViewportClient* InLevelViewportClient)
+		: FEditorDragToolBehaviorTarget(InLevelViewportClient->GetModeTools())
+		, LevelViewportClient(InLevelViewportClient)
+		, EditorViewportClient(InLevelViewportClient)
+	{
+	}
+
+	EDITORINTERACTIVETOOLSFRAMEWORK_API explicit FEditorFrustumSelect(FEditorViewportClient* InEditorViewportClient)
+		: FEditorDragToolBehaviorTarget(InEditorViewportClient->GetModeTools())
+		, LevelViewportClient(nullptr)
+		, EditorViewportClient(InEditorViewportClient)
+	{
+	}
+
+	virtual void Render(const FSceneView* View, FCanvas* Canvas) override;
+
+	//~ Begin IClickDragBehaviorTarget
+	virtual FInputRayHit CanBeginClickDragSequence(const FInputDeviceRay& InPressPos) override;
+	virtual void OnClickPress(const FInputDeviceRay& InPressPos) override;
+	virtual void OnClickDrag(const FInputDeviceRay& InDragPos) override;
+	virtual void OnClickRelease(const FInputDeviceRay& InReleasePos) override;
+	virtual void OnTerminateDragSequence() override;
+	//~ End IClickDragBehaviorTarget
+
+protected:
+	virtual TArray<FEditorModeID> GetUnsupportedModes() override;
+
+private:
+	/**
+	 * Returns true if the provided BSP node intersects with the provided frustum
+	 *
+	 * @param InModel				The model containing BSP nodes to check
+	 * @param NodeIndex				The index to a BSP node in the model.  This node is used for the bounds check.
+	 * @param InFrustum				The frustum to check against.
+	 * @param bUseStrictSelection	true if the node must be entirely within the frustum
+	 */
+	static bool IntersectsFrustum(const UModel& InModel, int32 NodeIndex, const FConvexVolume& InFrustum, bool bUseStrictSelection);
+
+	/**
+	 * Calculates a frustum to check actors against
+	 *
+	 * @param InView			Information about scene projection
+	 * @param OutFrustum		The created frustum
+	 * @param bUseBoxFrustum	If true a frustum out of the current dragged box will be created.  false will use the view frustum.
+	 */
+	void CalculateFrustum(const FSceneView* InView, FConvexVolume& OutFrustum, bool bUseBoxFrustum) const;
+
+	FLevelEditorViewportClient* LevelViewportClient;
+	FEditorViewportClient* EditorViewportClient;
+};
diff --git a/Engine/Source/Editor/UnrealEd/Private/DragTool_FrustumSelect.cpp b/Engine/Source/Editor/UnrealEd/Private/DragTool_FrustumSelect.cpp
--- a/Engine/Source/Editor/UnrealEd/Private/DragTool_FrustumSelect.cpp
+++ b/Engine/Source/Editor/UnrealEd/Private/DragTool_FrustumSelect.cpp
@@ -32,7 +32,7 @@
 //
 ///////////////////////////////////////////////////////////////////////////////
 
-namespace UE::DragTools::Private
+namespace UE::LevelEditor::Private
 {
 	TArray<FTypedElementHandle> GetElementsIntersectingFrustum(const AActor* Actor,
 		const FConvexVolume& InFrustum,
@@ -53,143 +53,48 @@ namespace UE::DragTools::Private
 
 		return {};
 	}
-
-	bool IsCurrentModeSupported(const FEditorViewportClient* InViewportClient)
-	{
-		static TArray<FEditorModeID> UnsupportedModes = { FBuiltinEditorModes::EM_Landscape };
-
-		if (InViewportClient)
-		{
-			if (const FEditorModeTools* ModeTools = InViewportClient->GetModeTools())
-			{
-				for (const FEditorModeID UnsupportedMode : UnsupportedModes)
-				{
-					if (!ModeTools->EnsureNotInMode(UnsupportedMode))
-					{
-						return false;
-					}
-				}
-			}
-		}
-
-		return true;
-	}
 }
 
 void FDragTool_ActorFrustumSelect::AddDelta( const FVector& InDelta )
 {
-	if (EditorViewportClient)
-	{
-		if (FViewport* const Viewport = EditorViewportClient->Viewport)
-		{
-			FIntPoint MousePos;
-			Viewport->GetMousePos(MousePos);
+	FIntPoint MousePos;
+	EditorViewportClient->Viewport->GetMousePos(MousePos);
 
-			FInputDeviceRay DragPos;
-			DragPos.ScreenPosition = MousePos;
+	EndWk = FVector(MousePos);
+	End = EndWk;
 
-			// Former AddDelta logic is now in OnClickDrag
-			OnClickDrag(DragPos);
-		}
-	}
+	const bool bUseHoverFeedback = GEditor != NULL && GetDefault<ULevelEditorViewportSettings>()->bEnableViewportHoverFeedback;
 }
 
 void FDragTool_ActorFrustumSelect::StartDrag(FEditorViewportClient* InViewportClient, const FVector& InStart, const FVector2D& InStartScreen)
 {
-	EditorViewportClient = InViewportClient;
-
-	FInputDeviceRay PressPos;
-	PressPos.WorldRay.Origin = InStart;
-	PressPos.ScreenPosition = InStartScreen;
-
-	// Former StartDrag logic is now in OnClickPress
-	OnClickPress(PressPos);
-}
-
-void FDragTool_ActorFrustumSelect::EndDrag()
-{
-	FInputDeviceRay ReleasePos;
-	ReleasePos.ScreenPosition = FVector2D(End);
-
-	// Former EndDrag logic is now in OnClickRelease
-	OnClickRelease(ReleasePos);
-}
-
-void FDragTool_ActorFrustumSelect::Render(const FSceneView* View, FCanvas* Canvas)
-{
-	if (!Canvas || !View)
-	{
-		return;
-	}
-
-	FCanvasBoxItem BoxItem(
-		FVector2D(Start.X, Start.Y) / Canvas->GetDPIScale(),
-		FVector2D(End.X - Start.X, End.Y - Start.Y) / Canvas->GetDPIScale()
-	);
-
-	BoxItem.SetColor(FLinearColor::White);
-	Canvas->DrawItem(BoxItem);
-}
-
-FInputRayHit FDragTool_ActorFrustumSelect::CanBeginClickDragSequence(const FInputDeviceRay& InPressPos)
-{
-	// Todo: this could be retrieved from a command for customization
-	FInputChord ActivationChord(EModifierKey::Control | EModifierKey::Alt, EKeys::LeftMouseButton);
+	FDragTool::StartDrag(InViewportClient, InStart, InStartScreen);
 
-	return IsActivationChordPressed(ActivationChord) && UE::DragTools::Private::IsCurrentModeSupported(LevelViewportClient)
-			 ? FInputRayHit(TNumericLimits<float>::Max()) // bHit is true. Depth is max to lose the standard tiebreaker.
-			 : FInputRayHit();
-}
+	const bool bUseHoverFeedback = GEditor != NULL && GetDefault<ULevelEditorViewportSettings>()->bEnableViewportHoverFeedback;
 
-void FDragTool_ActorFrustumSelect::OnClickPress(const FInputDeviceRay& InPressPos)
-{
-	if (EditorViewportClient)
-	{
-		// Signal that this tool is now active
-		OnActivateTool().Broadcast();
-
-		FDragTool::StartDrag(EditorViewportClient, InPressPos.WorldRay.Origin, InPressPos.ScreenPosition);
+	// Remove any active hover objects
+	FLevelEditorViewportClient::ClearHoverFromObjects();
 
-		// Remove any active hover objects
-		FLevelEditorViewportClient::ClearHoverFromObjects();
+	FIntPoint MousePos;
+	InViewportClient->Viewport->GetMousePos(MousePos);
 
-		Start = FVector(InPressPos.ScreenPosition.X, InPressPos.ScreenPosition.Y, 0);
-		End = EndWk = Start;
-	}
+	Start = FVector(InStartScreen.X, InStartScreen.Y, 0);
+	End = EndWk = Start;
 }
 
-void FDragTool_ActorFrustumSelect::OnClickDrag(const FInputDeviceRay& InDragPos)
-{
-	EndWk = FVector(InDragPos.ScreenPosition.X, InDragPos.ScreenPosition.Y, 0);
-	End = EndWk;
-}
-
-void FDragTool_ActorFrustumSelect::OnClickRelease(const FInputDeviceRay& InReleasePos)
+void FDragTool_ActorFrustumSelect::EndDrag()
 {
-	End = FVector(InReleasePos.ScreenPosition.X, InReleasePos.ScreenPosition.Y, 0);
-
-	if (!EditorViewportClient)
-	{
-		return;
-	}
-
-	FViewport* Viewport = EditorViewportClient->Viewport;
-	if (!Viewport)
-	{
-		return;
-	}
-
 	UBrushEditingSubsystem* BrushSubsystem = GEditor->GetEditorSubsystem<UBrushEditingSubsystem>();
 	const bool bGeometryMode = BrushSubsystem ? BrushSubsystem->IsGeometryEditorModeActive() : false;
 
-	FSceneViewFamilyContext ViewFamily(FSceneViewFamily::ConstructionValues(Viewport, EditorViewportClient->GetScene(), EditorViewportClient->EngineShowFlags ));
+	FSceneViewFamilyContext ViewFamily(FSceneViewFamily::ConstructionValues(EditorViewportClient->Viewport, EditorViewportClient->GetScene(), EditorViewportClient->EngineShowFlags ));
 	FSceneView* SceneView = EditorViewportClient->CalcSceneView(&ViewFamily);
 
 	// Generate a frustum out of the dragged box
 	FConvexVolume Frustum;
 	CalculateFrustum( SceneView, Frustum, true );
 
-	FScopedTransaction Transaction( NSLOCTEXT("ActorFrustumSelect", "MarqueeSelectTransaction", "Marquee Select" ) );
+	FScopedTransaction Transaction( NSLOCTEXT("ActorFrustumSelect", "MarqueeSelectTransation", "Marquee Select" ) );
 
 	bool bShouldSelect = true;
 
@@ -206,7 +111,7 @@ void FDragTool_ActorFrustumSelect::OnClickRelease(const FInputDeviceRay& InRelea
 	const bool bEditorModeHandledSelection = ModeTools->FrustumSelect(Frustum, EditorViewportClient, bLeftMouseButtonDown);
 
 	// Let the component visualizers try to handle the selection.
-	const bool bComponentVisHandledSelection = !bEditorModeHandledSelection && GUnrealEd->ComponentVisManager.HandleFrustumSelect(Frustum, EditorViewportClient, Viewport);
+	const bool bComponentVisHandledSelection = !bEditorModeHandledSelection && GUnrealEd->ComponentVisManager.HandleFrustumSelect(Frustum, EditorViewportClient, EditorViewportClient->Viewport);
 
 	if( !bEditorModeHandledSelection && !bComponentVisHandledSelection)
 	{
@@ -223,8 +128,8 @@ void FDragTool_ActorFrustumSelect::OnClickRelease(const FInputDeviceRay& InRelea
 			bGeometryMode
 		};
 
-		const int32 ViewportSizeX = Viewport->GetSizeXY().X;
-		const int32 ViewportSizeY = Viewport->GetSizeXY().Y;
+		const int32 ViewportSizeX = EditorViewportClient->Viewport->GetSizeXY().X;
+		const int32 ViewportSizeY = EditorViewportClient->Viewport->GetSizeXY().Y;
 
 		if( Start.X > End.X )
 		{
@@ -244,7 +149,7 @@ void FDragTool_ActorFrustumSelect::OnClickRelease(const FInputDeviceRay& InRelea
 			for(FActorIterator It(EditorViewportClient->GetWorld()); It; ++It)
 			{
 				AActor* Actor = *It;
-				ElementsToSelect.Append(UE::DragTools::Private::GetElementsIntersectingFrustum(Actor, Frustum, EditorViewportClient, LevelViewportClient, SeletionArgs));
+				ElementsToSelect.Append(UE::LevelEditor::Private::GetElementsIntersectingFrustum(Actor, Frustum, EditorViewportClient, LevelViewportClient, SeletionArgs));
 			}
 		}
 		else
@@ -263,7 +168,7 @@ void FDragTool_ActorFrustumSelect::OnClickRelease(const FInputDeviceRay& InRelea
 			// Typed Element selection
 			{
 				FTypedElementListRef ElementList = UTypedElementRegistry::GetInstance()->CreateElementList();
-				Viewport->GetElementHandlesInRect(BoxRect, ElementList);
+				EditorViewportClient->Viewport->GetElementHandlesInRect(BoxRect, ElementList);
 
 				if (bStrictDragSelection)
 				{
@@ -292,9 +197,11 @@ void FDragTool_ActorFrustumSelect::OnClickRelease(const FInputDeviceRay& InRelea
 			// We need this old code to support the BSP
 			TSet<AActor*> BSPActors;
 			TSet<UModel*> HitModels;
-			Viewport->GetActorsAndModelsInHitProxy( BoxRect, BSPActors, HitModels );
+			EditorViewportClient->Viewport->GetActorsAndModelsInHitProxy( BoxRect, BSPActors, HitModels );
 			BSPActors.Empty(HitModels.Num());
 
+			
+
 			if (HitModels.Num() > 0)
 			{
 				// Check every model to see if its BSP surfaces should be selected
@@ -321,7 +228,7 @@ void FDragTool_ActorFrustumSelect::OnClickRelease(const FInputDeviceRay& InRelea
 					AActor* Actor = *It;
 					if (bStrictDragSelection)
 					{
-						ElementsToSelect.Append(UE::DragTools::Private::GetElementsIntersectingFrustum(Actor, Frustum, EditorViewportClient, LevelViewportClient, SeletionArgs));
+						ElementsToSelect.Append(UE::LevelEditor::Private::GetElementsIntersectingFrustum(Actor, Frustum, EditorViewportClient, LevelViewportClient, SeletionArgs));
 					}
 					else
 					{
@@ -348,14 +255,13 @@ void FDragTool_ActorFrustumSelect::OnClickRelease(const FInputDeviceRay& InRelea
 	FLevelEditorViewportClient::ClearHoverFromObjects();
 
 	FDragTool::EndDrag();
-
-	// Signal that this tool is no longer active
-	OnDeactivateTool().Broadcast();
 }
 
-void FDragTool_ActorFrustumSelect::OnTerminateDragSequence()
+void FDragTool_ActorFrustumSelect::Render(const FSceneView* View, FCanvas* Canvas )
 {
-	Start = End = FVector::ZeroVector;
+	FCanvasBoxItem BoxItem( FVector2D(Start.X, Start.Y) / Canvas->GetDPIScale(), FVector2D(End.X-Start.X, End.Y-Start.Y) / Canvas->GetDPIScale());
+	BoxItem.SetColor( FLinearColor::White );
+	Canvas->DrawItem( BoxItem );
 }
 
 bool FDragTool_ActorFrustumSelect::IntersectsFrustum( const UModel& InModel, int32 NodeIndex, const FConvexVolume& InFrustum, bool bUseStrictSelection ) const
@@ -374,11 +280,6 @@ bool FDragTool_ActorFrustumSelect::IntersectsFrustum( const UModel& InModel, int
 
 void FDragTool_ActorFrustumSelect::CalculateFrustum( const FSceneView* InView, FConvexVolume& OutFrustum, bool bUseBoxFrustum ) const
 {
-	if (!InView || ! EditorViewportClient)
-	{
-		return;
-	}
-
 	if( bUseBoxFrustum )
 	{
 		FVector CamPoint = EditorViewportClient->GetViewLocation();
diff --git a/Engine/Source/Editor/UnrealEd/Private/EditorDragTools.cpp b/Engine/Source/Editor/UnrealEd/Private/EditorDragTools.cpp
--- a/Engine/Source/Editor/UnrealEd/Private/EditorDragTools.cpp
+++ b/Engine/Source/Editor/UnrealEd/Private/EditorDragTools.cpp
@@ -4,51 +4,7 @@
 #include "InputCoreTypes.h"
 #include "EditorViewportClient.h"
 #include "Editor.h"
-#include "EditorModeManager.h"
-#include "EditorModes.h"
 #include "SnappingUtils.h"
-#include "Tools/UEdMode.h"
-
-namespace UE::DragTools
-{
-// CVar initializer
-static TAutoConsoleVariable<int32> CVarEnableITFTools(
-	TEXT("SelectionMode.EnableITFTools"),
-	0,
-	TEXT("Is the ITF version of Selection Mode Tools enabled?\nCurrently supports: Frustum Select."),
-	FConsoleVariableDelegate::CreateLambda(
-		[](const IConsoleVariable* InVariable)
-		{
-			if (InVariable->GetBool())
-			{
-				OnITFDragToolsActivated().Broadcast();
-			}
-			else
-			{
-				OnITFDragToolsDeactivated().Broadcast();
-			}
-		}
-	)
-);
-
-bool UseITFTools()
-{
-	return CVarEnableITFTools->GetBool();
-}
-
-FOnDragToolsToggleDelegate& OnITFDragToolsActivated()
-{
-	static FOnDragToolsToggleDelegate OnDragToolsActivated;
-	return OnDragToolsActivated;
-}
-
-FOnDragToolsToggleDelegate& OnITFDragToolsDeactivated()
-{
-	static FOnDragToolsToggleDelegate OnDragToolsDeactivated;
-	return OnDragToolsDeactivated;
-}
-
-} // namespace UE::DragTools
 
 ///////////////////////////////////////////////////////////////////////////////////////////////////////////////////////
 //
@@ -110,46 +66,3 @@ void FDragTool::EndDrag()
 	Start = End = EndWk = FVector::ZeroVector;
 	bIsDragging = false;
 }
-
-void FDragTool::OnUpdateModifierState(int InModifierID, bool bInIsOn)
-{
-	constexpr int ShiftKeyMod = 1;
-	constexpr int AltKeyMod = 2;
-	constexpr int CtrlKeyMod = 3;
-
-	switch (InModifierID)
-	{
-	case ShiftKeyMod:
-		bShiftDown = bInIsOn;
-		break;
-	case CtrlKeyMod:
-		bControlDown = bInIsOn;
-		break;
-	case AltKeyMod:
-		bAltDown = bInIsOn;
-		break;
-	default:;
-	}
-}
-
-bool FDragTool::IsActivationChordPressed(const FInputChord& InChord) const
-{
-	bool bIsChordPressed = true;
-
-	if (InChord.NeedsControl())
-	{
-		bIsChordPressed = bIsChordPressed && bControlDown;
-	}
-
-	if (InChord.NeedsAlt())
-	{
-		bIsChordPressed = bIsChordPressed && bAltDown;
-	}
-
-	if (InChord.NeedsShift())
-	{
-		bIsChordPressed = bIsChordPressed && bShiftDown;
-	}
-
-	return bIsChordPressed;
-}
diff --git a/Engine/Source/Editor/UnrealEd/Private/MouseDeltaTracker.cpp b/Engine/Source/Editor/UnrealEd/Private/MouseDeltaTracker.cpp
--- a/Engine/Source/Editor/UnrealEd/Private/MouseDeltaTracker.cpp
+++ b/Engine/Source/Editor/UnrealEd/Private/MouseDeltaTracker.cpp
@@ -10,6 +10,7 @@
 #include "EditorDragTools.h"
 #include "SnappingUtils.h"
 #include "UnrealWidget.h"
+#include "EditorDragTools/EditorDragToolBehaviorTarget.h"
 
 #define LOCTEXT_NAMESPACE "MouseDeltaTracker"
 
@@ -370,7 +371,7 @@ void FMouseDeltaTracker::ConditionalBeginUsingDragTool( FEditorViewportClient* I
 				}
 				else
 				{
-					if (!UE::DragTools::UseITFTools())
+					if (!UE::Editor::DragTools::UseEditorDragTools())
 					{
 						if (LeftMouseButtonDown && bControlDown && bAltDown)
 						{
diff --git a/Engine/Source/Editor/UnrealEd/Private/Tools/DragToolsBehaviourSource.cpp b/Engine/Source/Editor/UnrealEd/Private/Tools/DragToolsBehaviourSource.cpp
deleted file mode 100644
diff --git a/Engine/Source/Editor/UnrealEd/Private/Tools/DragToolsBehaviourSource.h b/Engine/Source/Editor/UnrealEd/Private/Tools/DragToolsBehaviourSource.h
deleted file mode 100644
diff --git a/Engine/Source/Editor/UnrealEd/Private/Tools/EdModeInteractiveToolsContext.cpp b/Engine/Source/Editor/UnrealEd/Private/Tools/EdModeInteractiveToolsContext.cpp
--- a/Engine/Source/Editor/UnrealEd/Private/Tools/EdModeInteractiveToolsContext.cpp
+++ b/Engine/Source/Editor/UnrealEd/Private/Tools/EdModeInteractiveToolsContext.cpp
@@ -36,7 +36,8 @@
 
 #include "BaseGizmos/GizmoRenderingUtil.h"
 #include "UnrealClient.h"
-#include "Editor/UnrealEd/Private/Tools/DragToolsBehaviourSource.h"
+#include "EditorDragTools/DragToolsBehaviourSource.h"
+#include "EditorDragTools/EditorDragToolBehaviorTarget.h"
 #include "UObject/ObjectSaveContext.h"
 
 //#define ENABLE_DEBUG_PRINTING
@@ -1242,16 +1243,16 @@ void UModeManagerInteractiveToolsContext::Initialize(IToolsContextQueriesAPI* Qu
 	GetParentEditorModeManager()->OnEditorModeIDChanged().AddUObject(this, &UModeManagerInteractiveToolsContext::OnEditorModeChanged);
 
 	// Delegates used to handle SelectionMode.EnableITFTools CVAR changes
-	UE::DragTools::OnITFDragToolsActivated().AddUObject(this, &UModeManagerInteractiveToolsContext::RegisterDragTools);
-	UE::DragTools::OnITFDragToolsDeactivated().AddUObject(this, &UModeManagerInteractiveToolsContext::UnregisterDragTools);
+	UE::Editor::DragTools::OnEditorDragToolsActivated().AddUObject(this, &UModeManagerInteractiveToolsContext::RegisterDragTools);
+	UE::Editor::DragTools::OnEditorDragToolsDeactivated().AddUObject(this, &UModeManagerInteractiveToolsContext::UnregisterDragTools);
 }
 
 void UModeManagerInteractiveToolsContext::Shutdown()
 {
 	UnregisterDragTools();
 	GetParentEditorModeManager()->OnEditorModeIDChanged().RemoveAll(this);
-	UE::DragTools::OnITFDragToolsActivated().RemoveAll(this);
-	UE::DragTools::OnITFDragToolsDeactivated().RemoveAll(this);
+	UE::Editor::DragTools::OnEditorDragToolsActivated().RemoveAll(this);
+	UE::Editor::DragTools::OnEditorDragToolsDeactivated().RemoveAll(this);
 
 	if (ToolManager)
 	{
@@ -1364,7 +1365,7 @@ void UModeManagerInteractiveToolsContext::UnregisterDragTools()
 }
 void UModeManagerInteractiveToolsContext::OnEditorModeChanged(const FEditorModeID& InModeID, bool bInIsEnteringMode)
 {
-	if (!UE::DragTools::UseITFTools())
+	if (!UE::Editor::DragTools::UseEditorDragTools())
 	{
 		return;
 	}
diff --git a/Engine/Source/Editor/UnrealEd/Public/DragTool_FrustumSelect.h b/Engine/Source/Editor/UnrealEd/Public/DragTool_FrustumSelect.h
--- a/Engine/Source/Editor/UnrealEd/Public/DragTool_FrustumSelect.h
+++ b/Engine/Source/Editor/UnrealEd/Public/DragTool_FrustumSelect.h
@@ -1,5 +1,6 @@
 // Copyright Epic Games, Inc. All Rights Reserved.
 
+
 #pragma once
 
 #include "CoreMinimal.h"
@@ -42,21 +43,13 @@ public:
 	 * @param	InStart				Where the mouse was when the drag started.
 	 */
 	UNREALED_API virtual void StartDrag(FEditorViewportClient* InViewportClient, const FVector& InStart, const FVector2D& InStartScreen) override;
-
+	
 	/**
 	 * Ends a mouse drag behavior (the user has let go of the mouse button).
 	 */
 	UNREALED_API virtual void EndDrag() override;
 	UNREALED_API virtual void Render(const FSceneView* View,FCanvas* Canvas) override;
 
-	//~ Begin IClickDragBehaviorTarget
-	UNREALED_API virtual FInputRayHit CanBeginClickDragSequence(const FInputDeviceRay& InPressPos) override;
-	UNREALED_API virtual void OnClickPress(const FInputDeviceRay& InPressPos) override;
-	UNREALED_API virtual void OnClickDrag(const FInputDeviceRay& InDragPos) override;
-	UNREALED_API virtual void OnClickRelease(const FInputDeviceRay& InReleasePos) override;
-	UNREALED_API virtual void OnTerminateDragSequence() override;
-	//~ End IClickDragBehaviorTarget
-
 private:
 	/** 
 	 * Calculates a frustum to check actors against 
@@ -76,10 +69,10 @@ private:
 	 * @param bUseStrictSelection	true if the node must be entirely within the frustum
 	 */
 	UNREALED_API bool IntersectsFrustum( const UModel& InModel, int32 NodeIndex, const FConvexVolume& InFrustum, bool bUseStrictSelection ) const;
-
+	
 	/** Adds a hover effect to the passed in actor */
 	UNREALED_API void AddHoverEffect( AActor& InActor );
-
+	
 	/** Adds a hover effect to the passed in bsp surface */
 	UNREALED_API void AddHoverEffect( UModel& InModel, int32 SurfIndex );
 
diff --git a/Engine/Source/Editor/UnrealEd/Public/EditorDragTools.h b/Engine/Source/Editor/UnrealEd/Public/EditorDragTools.h
--- a/Engine/Source/Editor/UnrealEd/Public/EditorDragTools.h
+++ b/Engine/Source/Editor/UnrealEd/Public/EditorDragTools.h
@@ -3,42 +3,19 @@
 
 #pragma once
 
-#include "BaseBehaviors/BehaviorTargetInterfaces.h"
 #include "CoreMinimal.h"
-#include "Editor.h"
 
 class FCanvas;
 class FEditorModeTools;
 class FEditorViewportClient;
 class FPrimitiveDrawInterface;
 class FSceneView;
-class FUICommandInfo;
-
-/**
- * ITF related helper functions
- */
-namespace UE::DragTools
-{
-
-DECLARE_MULTICAST_DELEGATE(FOnDragToolsToggleDelegate)
-/**
- * Returns true if ITF-based drag tools should be used.
- * If an ITF version of a drag tool is not available yet, legacy will be used.
- */
-UNREALED_API bool UseITFTools();
-
-FOnDragToolsToggleDelegate& OnITFDragToolsActivated();
-FOnDragToolsToggleDelegate& OnITFDragToolsDeactivated();
-
-} // namespace UE::DragTools
 
 /**
  * The base class that all drag tools inherit from.
  * The drag tools implement special behaviors for the user clicking and dragging in a viewport.
- *
- * Inheriting from IClickDragBehaviorTarget so that Drag Tools logic can be shared between legacy and ITF usage of the same tools
  */
-class FDragTool : public IClickDragBehaviorTarget
+class FDragTool
 {
 public:
 	UNREALED_API FDragTool(FEditorModeTools* InModeTools);
@@ -78,39 +55,6 @@ public:
 	/** Does this drag tool need to have the mouse movement converted to the viewport orientation? */
 	bool bConvertDelta;
 
-	//~ Begin IModifierToggleBehaviorTarget
-	virtual void OnUpdateModifierState(int InModifierID, bool bInIsOn) override;
-	//~ End IModifierToggleBehaviorTarget
-
-	//~ Begin IClickDragBehaviorTarget
-	virtual FInputRayHit CanBeginClickDragSequence(const FInputDeviceRay& InPressPos) override
-	{
-		return FInputRayHit();
-	}
-
-	virtual void OnClickPress(const FInputDeviceRay& InPressPos) override
-	{
-	}
-
-	virtual void OnClickDrag(const FInputDeviceRay& InDragPos) override
-	{
-	}
-
-	virtual void OnClickRelease(const FInputDeviceRay& InReleasePos) override
-	{
-	}
-
-	virtual void OnTerminateDragSequence() override
-	{
-	}
-	//~ End IClickDragBehaviorTarget
-
-	bool IsActivationChordPressed(const FInputChord& InChord) const;
-
-	DECLARE_MULTICAST_DELEGATE(FOnToolStateChange);
-	FOnToolStateChange& OnActivateTool() { return OnToolActivatedDelegate; }
-	FOnToolStateChange& OnDeactivateTool() { return OnToolDeactivatedDelegate; }
-
 protected:
 	FEditorModeTools* ModeTools;
 
@@ -129,7 +73,4 @@ protected:
 	bool bMiddleMouseButtonDown;
 	/** True if we are dragging */
 	bool bIsDragging;
-
-	FOnToolStateChange OnToolActivatedDelegate;
-	FOnToolStateChange OnToolDeactivatedDelegate;
 };
