[EditorGizmos] Fixed rotation snapping order

diff --git a/Engine/Source/Editor/Experimental/EditorInteractiveToolsFramework/Private/EditorGizmos/TransformGizmo.cpp b/Engine/Source/Editor/Experimental/EditorInteractiveToolsFramework/Private/EditorGizmos/TransformGizmo.cpp
--- a/Engine/Source/Editor/Experimental/EditorInteractiveToolsFramework/Private/EditorGizmos/TransformGizmo.cpp
+++ b/Engine/Source/Editor/Experimental/EditorInteractiveToolsFramework/Private/EditorGizmos/TransformGizmo.cpp
@@ -2712,6 +2712,8 @@ void UTransformGizmo::OnClickDragRotateAxis(const FInputDeviceRay& DragPos)
 			const double HitAngle = InteractionCurrAngle + DeltaAngle;
 			DebugData.InteractionAngleCurrent += FMath::FindDeltaAngleRadians(DebugData.InteractionAngleCurrent, HitAngle);
 			DebugData.InteractionRadius = FVector::Distance(DebugData.TransformStart.GetLocation(), DebugData.InteractionCurrent.GetLocation());
+
+			DebugData.InteractionPlaneNormal = DeltaRot.GetRotationAxis();
 		}
 
 		break;
@@ -2904,9 +2906,24 @@ void UTransformGizmo::OnClickDragGimbalRotateAxis(const FInputDeviceRay& InDragP
 	case EAxisRotateMode::Pull:
 		{
 			FQuat DeltaRot = ComputeGimbalRotateDelta(InteractionScreenCurrPos, InDragPos.ScreenPosition);
-			ApplyRotateDelta(DeltaRot);
 			SnapRotateDelta(DeltaRot);
+			ApplyRotateDelta(DeltaRot);
 			InteractionScreenCurrPos = InDragPos.ScreenPosition;
+
+			// Update Debug
+			if (GizmoLocals::bDebugDraw)
+			{
+				FVector::FReal HitDepth;
+				if (GetRayParamIntersectionWithInteractionPlane(InDragPos, HitDepth))
+				{
+					const FVector HitPoint = InDragPos.WorldRay.Origin + InDragPos.WorldRay.Direction * HitDepth;
+					DebugData.InteractionCurrent.SetLocation(HitPoint);
+				}
+				
+				DebugData.InteractionRadius = FVector::Distance(DebugData.TransformStart.GetLocation(), DebugData.InteractionCurrent.GetLocation());
+				DebugData.InteractionPlaneNormal = InteractionPlanarNormal;
+			}
+			
 			break;
 		}
 	case EAxisRotateMode::Arc:
