[ITF Drag Tools] Add Marquee Box Select, Measure and Viewport Change tools

Adds missing ITF based Drag Tools

- Marquee Box Select for Ortho View
- Measure Tool for Ortho View
- Viewport Change Tool

Adds some feature + fix for both ITF and legacy Viewport Change Tool

- Uses DPI info, so that tool line is rendered properly even when scale is not 100%
- Fixes "misaligned" view modes enum, which seemed to be mixed up when choosing which one to choose and also show in widget
- ITF only: adds text preview to let the user know which View they will get when releasing the mouse button

Changes to allow Viewport Change tool to work with FEditorViewportClient instead of FLevelEditorViewportClient:

- Deprecated FLevelEditorViewportClient::SetViewportTypeFromTool, since it is simply calling SetViewportType.
- Removed code creating that one item radial entry menu when Viewport Change drag distance is short. Having the upcoming view name shown should make it unnecessary

FEditorDragToolBehaviorTarget now has a IEditorViewportClientProxy member which can be used to access Editor VPC or Level VPC functionality without the Drag Tools having to store separate VPCs directly.

diff --git a/Engine/Source/Editor/Experimental/EditorInteractiveToolsFramework/Private/EditorDragTools/DragToolsBehaviourSource.cpp b/Engine/Source/Editor/Experimental/EditorInteractiveToolsFramework/Private/EditorDragTools/DragToolsBehaviourSource.cpp
--- a/Engine/Source/Editor/Experimental/EditorInteractiveToolsFramework/Private/EditorDragTools/DragToolsBehaviourSource.cpp
+++ b/Engine/Source/Editor/Experimental/EditorInteractiveToolsFramework/Private/EditorDragTools/DragToolsBehaviourSource.cpp
@@ -3,8 +3,11 @@
 #include "EditorDragTools/DragToolsBehaviourSource.h"
 #include "BaseBehaviors/ClickDragBehavior.h"
 #include "BaseBehaviors/KeyAsModifierInputBehavior.h"
+#include "EditorDragTools/EditorBoxSelect.h"
 #include "EditorDragTools/EditorDragToolBehaviorTarget.h"
 #include "EditorDragTools/EditorFrustumSelect.h"
+#include "EditorDragTools/EditorMeasureTool.h"
+#include "EditorDragTools/EditorViewportChange.h"
 #include "EditorModeManager.h"
 #include "Tools/EdModeInteractiveToolsContext.h"
 
@@ -51,46 +54,61 @@ void UDragToolsBehaviorSource::Initialize(const UEditorInteractiveToolsContext*
 
 	BehaviorSet = NewObject<UInputBehaviorSet>();
 
-	// In case New TRS Gizmos are enabled, we need Drag Tools to be processed before the Gizmo ones, otherwise Drag Behaviors will never be triggered.
-	// See:
+	// In case New TRS Gizmos are enabled, we need Mouse Drag + Key modifier Drag Tools Behaviors to be processed before
+	// the Gizmo ones, otherwise Drag Behaviors will never be triggered. See:
 	// * UTransformGizmo::SetupBehaviors()			[DEFAULT_GIZMO_PRIORITY]
-	// * UTransformGizmo::SetupIndirectBehaviors()	[DEFAULT_GIZMO_PRIORITY-1]
-	static constexpr int DEFAULT_DRAG_TOOLS_PRIORITY = FInputCapturePriority::DEFAULT_GIZMO_PRIORITY-2;
+	// * UTransformGizmo::SetupIndirectBehaviors()	[DEFAULT_GIZMO_PRIORITY - 1]
+	static constexpr int HANDLED_BEFORE_GIZMO_PRIORITY = FInputCapturePriority::DEFAULT_GIZMO_PRIORITY - 3;
 
-	// Frustum Select
+	// Conversely, tools which are activated with just mouse buttons input (no key modifiers, mouse drag behavior only) need to be processed after Gizmo Behaviors
+	static constexpr int HANDLED_AFTER_GIZMO_PRIORITY = FInputCapturePriority::DEFAULT_GIZMO_PRIORITY + 1;
+
+	static constexpr int FRUSTUM_SELECT_PRIORITY = HANDLED_BEFORE_GIZMO_PRIORITY;
+	static constexpr int VIEWPORT_CHANGE_PRIORITY = HANDLED_BEFORE_GIZMO_PRIORITY;
+	static constexpr int BOX_SELECT_PRIORITY = HANDLED_AFTER_GIZMO_PRIORITY;
+	static constexpr int MEASURE_PRIORITY = VIEWPORT_CHANGE_PRIORITY + 1;
+
+	// Box Select
 	{
-		TSharedPtr<FEditorDragToolBehaviorTarget> FrustumSelectTool;
+		TSharedPtr<FEditorDragToolBehaviorTarget> BoxSelectTool = MakeShared<FEditorBoxSelect>(EditorViewportClient);
 
-		// Some Drag Tools require or can make use specifically of Level Viewport Client.
-		if (EditorViewportClient->IsLevelEditorClient())
-		{
-			// Initialize using Level Viewport Client if available
-			FLevelEditorViewportClient* LevelViewportClient = static_cast<FLevelEditorViewportClient*>(EditorViewportClient);
-			FrustumSelectTool = MakeShared<FEditorFrustumSelect>(LevelViewportClient);
-		}
-		else
-		{
-			FrustumSelectTool = MakeShared<FEditorFrustumSelect>(EditorViewportClient);
-		}
+		DragTools.Add(BoxSelectTool);
+
+		// Add drag input behavior
+		UClickDragInputBehavior* MouseBehavior = NewObject<UClickDragInputBehavior>();
+		MouseBehavior->Initialize(BoxSelectTool.Get());
+		MouseBehavior->SetDefaultPriority(BOX_SELECT_PRIORITY);
+		BehaviorSet->Add(MouseBehavior);
+
+		BoxSelectTool->OnActivateTool().AddUObject(this, &UDragToolsBehaviorSource::OnActivateTool, EDragTool::BoxSelect);
+		BoxSelectTool->OnDeactivateTool().AddUObject(this, &UDragToolsBehaviorSource::OnDeactivateTool, EDragTool::BoxSelect);
+	}
+
+	// Frustum Select
+	{
+		TSharedPtr<FEditorDragToolBehaviorTarget> FrustumSelectTool = MakeShared<FEditorFrustumSelect>(EditorViewportClient);
 
 		DragTools.Add(FrustumSelectTool);
 
 		// Add drag input behavior
 		UClickDragInputBehavior* MouseBehavior = NewObject<UClickDragInputBehavior>();
 		MouseBehavior->Initialize(FrustumSelectTool.Get());
-		MouseBehavior->SetDefaultPriority(DEFAULT_DRAG_TOOLS_PRIORITY);
+		MouseBehavior->SetDefaultPriority(FRUSTUM_SELECT_PRIORITY);
+		MouseBehavior->SetUseLeftMouseButton();
 		BehaviorSet->Add(MouseBehavior);
 
-		constexpr int ShiftKeyMod = 1;
-		constexpr int AltKeyMod = 2;
-		constexpr int CtrlKeyMod = 3;
-
 		// Register modifier listeners
 		UKeyAsModifierInputBehavior* KeyModifierInputBehavior = NewObject<UKeyAsModifierInputBehavior>();
-		KeyModifierInputBehavior->Initialize(FrustumSelectTool.Get(), ShiftKeyMod, FInputDeviceState::IsShiftKeyDown);
-		KeyModifierInputBehavior->Initialize(FrustumSelectTool.Get(), CtrlKeyMod, FInputDeviceState::IsCtrlKeyDown);
-		KeyModifierInputBehavior->Initialize(FrustumSelectTool.Get(), AltKeyMod, FInputDeviceState::IsAltKeyDown);
-		KeyModifierInputBehavior->SetDefaultPriority(DEFAULT_DRAG_TOOLS_PRIORITY);
+		KeyModifierInputBehavior->Initialize(
+			FrustumSelectTool.Get(), UE::Editor::DragTools::ShiftKeyMod, FInputDeviceState::IsShiftKeyDown
+		);
+		KeyModifierInputBehavior->Initialize(
+			FrustumSelectTool.Get(), UE::Editor::DragTools::CtrlKeyMod, FInputDeviceState::IsCtrlKeyDown
+		);
+		KeyModifierInputBehavior->Initialize(
+			FrustumSelectTool.Get(), UE::Editor::DragTools::AltKeyMod, FInputDeviceState::IsAltKeyDown
+		);
+		KeyModifierInputBehavior->SetDefaultPriority(FRUSTUM_SELECT_PRIORITY);
 		BehaviorSet->Add(KeyModifierInputBehavior);
 
 		FrustumSelectTool->OnActivateTool().AddUObject(
@@ -101,12 +119,62 @@ void UDragToolsBehaviorSource::Initialize(const UEditorInteractiveToolsContext*
 		);
 	}
 
-	// Add other drag tools
+	// Measure
+	{
+		TSharedPtr<FEditorDragToolBehaviorTarget> MeasureTool = MakeShared<FEditorMeasureTool>(EditorViewportClient);
+
+		DragTools.Add(MeasureTool);
+
+		// Add drag input behavior
+		UClickDragInputBehavior* MouseBehavior = NewObject<UClickDragInputBehavior>();
+		MouseBehavior->Initialize(MeasureTool.Get());
+		MouseBehavior->SetDefaultPriority(MEASURE_PRIORITY);
+		MouseBehavior->SetUseMiddleMouseButton();
+		BehaviorSet->Add(MouseBehavior);
+
+		MeasureTool->OnActivateTool().AddUObject(this, &UDragToolsBehaviorSource::OnActivateTool, EDragTool::Measure);
+		MeasureTool->OnDeactivateTool().AddUObject(this, &UDragToolsBehaviorSource::OnDeactivateTool, EDragTool::Measure);
+	}
+
+	// Viewport Change
+	{
+		// Requires Level Viewport Client specifically.
+		// todo: Maybe we could move some code from Level Viewport Client to Editor Viewport Client, to make this more generic?
+		if (EditorViewportClient->IsLevelEditorClient())
+		{
+			TSharedPtr<FEditorDragToolBehaviorTarget> ViewportChangeTool =
+				MakeShared<FEditorViewportChange>(EditorViewportClient);
+
+			DragTools.Add(ViewportChangeTool);
+
+			// Add drag input behavior
+			UClickDragInputBehavior* MouseBehavior = NewObject<UClickDragInputBehavior>();
+			MouseBehavior->Initialize(ViewportChangeTool.Get());
+			MouseBehavior->SetDefaultPriority(VIEWPORT_CHANGE_PRIORITY);
+			MouseBehavior->SetUseMiddleMouseButton();
+			BehaviorSet->Add(MouseBehavior);
+
+			// Register modifier listeners
+			UKeyAsModifierInputBehavior* KeyModifierInputBehavior = NewObject<UKeyAsModifierInputBehavior>();
+			KeyModifierInputBehavior->Initialize(
+				ViewportChangeTool.Get(), UE::Editor::DragTools::CtrlKeyMod, FInputDeviceState::IsCtrlKeyDown
+			);
+			KeyModifierInputBehavior->SetDefaultPriority(VIEWPORT_CHANGE_PRIORITY);
+			BehaviorSet->Add(KeyModifierInputBehavior);
+
+			ViewportChangeTool->OnActivateTool().AddUObject(
+				this, &UDragToolsBehaviorSource::OnActivateTool, EDragTool::ViewportChange
+			);
+			ViewportChangeTool->OnDeactivateTool().AddUObject(
+				this, &UDragToolsBehaviorSource::OnDeactivateTool, EDragTool::ViewportChange
+			);
+		}
+	}
 }
 
 FEditorDragToolBehaviorTarget* UDragToolsBehaviorSource::GetActiveTool() const
 {
-	if (ActiveToolType > 0)
+	if (ActiveToolType >= 0)
 	{
 		int32 ToolIndex = GetToolIndex(static_cast<EDragTool::Type>(ActiveToolType));
 		if (DragTools.IsValidIndex(ToolIndex))
@@ -123,7 +191,10 @@ FEditorDragToolBehaviorTarget* UDragToolsBehaviorSource::GetActiveTool() const
 
 int32 UDragToolsBehaviorSource::GetToolIndex(EDragTool::Type InDragToolType)
 {
-	static TArray<int32> ToolIndexes = { EDragTool::FrustumSelect };
+	// Please add tools in the same order as they appear in EDragTool declaration
+	static TArray<int32> ToolIndexes = {
+		EDragTool::BoxSelect, EDragTool::FrustumSelect, EDragTool::Measure, EDragTool::ViewportChange
+	};
 
 	int32 ToolIndex = ToolIndexes.Find(InDragToolType);
 	if (ToolIndex != INDEX_NONE)
diff --git a/Engine/Source/Editor/Experimental/EditorInteractiveToolsFramework/Private/EditorDragTools/EditorBoxSelect.cpp b/Engine/Source/Editor/Experimental/EditorInteractiveToolsFramework/Private/EditorDragTools/EditorBoxSelect.cpp
new file mode 100644
--- /dev/null
+++ b/Engine/Source/Editor/Experimental/EditorInteractiveToolsFramework/Private/EditorDragTools/EditorBoxSelect.cpp
@@ -0,0 +1,510 @@
+// Copyright Epic Games, Inc. All Rights Reserved.
+
+#include "EditorDragTools/EditorBoxSelect.h"
+#include "CanvasItem.h"
+#include "CanvasTypes.h"
+#include "Components/PrimitiveComponent.h"
+#include "Editor/UnrealEdEngine.h"
+#include "EditorDragTools/EditorViewportClientProxy.h"
+#include "EditorModeManager.h"
+#include "EditorModes.h"
+#include "Elements/Framework/EngineElementsLibrary.h"
+#include "Elements/Framework/TypedElementRegistry.h"
+#include "Elements/Framework/TypedElementSelectionSet.h"
+#include "Elements/Interfaces/TypedElementWorldInterface.h"
+#include "Engine/LevelStreaming.h"
+#include "EngineUtils.h"
+#include "GameFramework/Volume.h"
+#include "LevelEditorSubsystem.h"
+#include "LevelEditorViewport.h"
+#include "Model.h"
+#include "SceneView.h"
+#include "ScopedTransaction.h"
+#include "Settings/LevelEditorViewportSettings.h"
+#include "Subsystems/BrushEditingSubsystem.h"
+#include "UnrealEdGlobals.h"
+
+namespace UE::EditorDragTools::Private
+{
+TArray<FTypedElementHandle> GetElementsIntersectingBox(
+	const AActor* Actor,
+	const FBox& InBox,
+	const IEditorViewportClientProxy* InEditorViewportClientProxy,
+	const FWorldSelectionElementArgs& SelectionArgs
+)
+{
+	if (InEditorViewportClientProxy && InEditorViewportClientProxy->IsActorVisible(Actor))
+	{
+		if (FTypedElementHandle ActorHandle = UEngineElementsLibrary::AcquireEditorActorElementHandle(Actor))
+		{
+			if (TTypedElement<ITypedElementWorldInterface> WorldElement =
+					UTypedElementRegistry::GetInstance()->GetElement<ITypedElementWorldInterface>(ActorHandle))
+			{
+				return WorldElement.GetSelectionElementsInBox(InBox, SelectionArgs);
+			}
+		}
+	}
+
+	return {};
+}
+
+void AddHoverEffect(const FViewportHoverTarget& InHoverTarget)
+{
+	FLevelEditorViewportClient::AddHoverEffect(InHoverTarget);
+	FLevelEditorViewportClient::HoveredObjects.Add(InHoverTarget);
+}
+
+void RemoveHoverEffect(const FViewportHoverTarget& InHoverTarget)
+{
+	FSetElementId Id = FLevelEditorViewportClient::HoveredObjects.FindId(InHoverTarget);
+	if (Id.IsValidId())
+	{
+		FLevelEditorViewportClient::RemoveHoverEffect(InHoverTarget);
+		FLevelEditorViewportClient::HoveredObjects.Remove(Id);
+	}
+}
+
+} // namespace UE::EditorDragTools::Private
+
+void FEditorBoxSelect::Render(const FSceneView* View, FCanvas* Canvas)
+{
+	FCanvasBoxItem BoxItem(
+		FVector2D(Start.X, Start.Y) / Canvas->GetDPIScale(),
+		FVector2D(End.X - Start.X, End.Y - Start.Y) / Canvas->GetDPIScale()
+	);
+	BoxItem.SetColor(FLinearColor::White);
+	Canvas->DrawItem(BoxItem);
+}
+
+FInputRayHit FEditorBoxSelect::CanBeginClickDragSequence(const FInputDeviceRay& InPressPos)
+{
+	if (FEditorViewportClient* const EditorViewportClient = EditorViewportClientProxy->GetEditorViewportClient())
+	{
+		if (IsCurrentModeSupported())
+		{
+			if (EditorViewportClient->IsOrtho()
+				&& !IEditorViewportClientProxy::IsMouseOnWidgetAxis(EditorViewportClientProxy))
+			{
+				return FInputRayHit(TNumericLimits<float>::Max());
+			}
+		}
+	}
+
+	return FInputRayHit();
+}
+
+void FEditorBoxSelect::OnClickPress(const FInputDeviceRay& InPressPos)
+{
+	if (!EditorViewportClientProxy)
+	{
+		return;
+	}
+
+	FEditorViewportClient* const EditorViewportClient = EditorViewportClientProxy->GetEditorViewportClient();
+	if (!EditorViewportClient)
+	{
+		return;
+	}
+
+	FViewport* Viewport = EditorViewportClient->Viewport;
+	if (!Viewport)
+	{
+		return;
+	}
+
+	OnActivateTool().Broadcast();
+
+	FIntPoint MousePos;
+	Viewport->GetMousePos(MousePos);
+
+	Start = FVector(MousePos);
+	End = Start;
+
+	FLevelEditorViewportClient::ClearHoverFromObjects();
+
+	// Create a list of bsp models to check for intersection with the box
+	ModelsToCheck.Reset();
+	// Do not select BSP if its not visible
+	if (EditorViewportClient->EngineShowFlags.BSP)
+	{
+		UWorld* World = EditorViewportClient->GetWorld();
+		check(World);
+		// Add the persistent level always
+		ModelsToCheck.Add(World->PersistentLevel->Model);
+		// Add all streaming level models
+		for (ULevelStreaming* StreamingLevel : World->GetStreamingLevels())
+		{
+			// Only add streaming level models if the level is visible
+			if (StreamingLevel && StreamingLevel->GetShouldBeVisibleInEditor())
+			{
+				if (ULevel* Level = StreamingLevel->GetLoadedLevel())
+				{
+					ModelsToCheck.Add(Level->Model);
+				}
+			}
+		}
+	}
+}
+
+void FEditorBoxSelect::OnClickDrag(const FInputDeviceRay& InDragPos)
+{
+	if (!EditorViewportClientProxy)
+	{
+		return;
+	}
+
+	FEditorViewportClient* const EditorViewportClient = EditorViewportClientProxy->GetEditorViewportClient();
+	if (!EditorViewportClient)
+	{
+		return;
+	}
+
+	FIntPoint MousePos;
+	EditorViewportClient->Viewport->GetMousePos(MousePos);
+
+	End = FVector(MousePos);
+
+	const bool bUseHoverFeedback = GEditor && GetDefault<ULevelEditorViewportSettings>()->bEnableViewportHoverFeedback;
+
+	if (bUseHoverFeedback)
+	{
+		const bool bStrictDragSelection = GetDefault<ULevelEditorViewportSettings>()->bStrictBoxSelection;
+
+		UTypedElementSelectionSet* SelectionSet = GEditor->GetEditorSubsystem<ULevelEditorSubsystem>()->GetSelectionSet();
+
+		UBrushEditingSubsystem* BrushSubsystem = GEditor->GetEditorSubsystem<UBrushEditingSubsystem>();
+		const bool bGeometryMode = BrushSubsystem ? BrushSubsystem->IsGeometryEditorModeActive() : false;
+
+		FWorldSelectionElementArgs SelectionArgs{ SelectionSet,
+												  ETypedElementSelectionMethod::Primary,
+												  FTypedElementSelectionOptions(),
+												  &(EditorViewportClient->EngineShowFlags),
+												  bStrictDragSelection,
+												  bGeometryMode };
+
+		// If we are using over feedback calculate a new box from the one being dragged
+		FBox SelBBox;
+		CalculateBox(SelBBox);
+
+		// Check every actor to see if it intersects the frustum created by the box
+		// If it does, the actor will be selected and should be given a hover cue
+		UWorld* IteratorWorld = GWorld;
+		for (FActorIterator It(IteratorWorld); It; ++It)
+		{
+			AActor& Actor = **It;
+			const bool bActorHitByBox = !UE::EditorDragTools::Private::GetElementsIntersectingBox(
+											 &Actor, SelBBox, EditorViewportClientProxy, SelectionArgs
+			)
+											 .IsEmpty();
+
+			if (bActorHitByBox)
+			{
+				// Apply a hover effect to any actor that will be selected
+				AddHoverEffect(Actor);
+			}
+			else
+			{
+				// Remove any hover effect on this actor as it no longer will be selected by the current box
+				RemoveHoverEffect(Actor);
+			}
+		}
+
+		// Check each model to see if it will be selected
+		for (int32 ModelIndex = 0; ModelIndex < ModelsToCheck.Num(); ++ModelIndex)
+		{
+			UModel& Model = *ModelsToCheck[ModelIndex];
+			for (int32 NodeIndex = 0; NodeIndex < Model.Nodes.Num(); NodeIndex++)
+			{
+				if (IntersectsBox(Model, NodeIndex, SelBBox, bStrictDragSelection))
+				{
+					// Apply a hover effect to any bsp surface that will be selected
+					AddHoverEffect(Model, Model.Nodes[NodeIndex].iSurf);
+				}
+				else
+				{
+					// Remove any hover effect on this bsp surface as it no longer will be selected by the current box
+					RemoveHoverEffect(Model, Model.Nodes[NodeIndex].iSurf);
+				}
+			}
+		}
+	}
+}
+
+void FEditorBoxSelect::OnClickRelease(const FInputDeviceRay& InReleasePos)
+{
+	if (EditorViewportClientProxy)
+	{
+		if (FEditorViewportClient* const EditorViewportClient = EditorViewportClientProxy->GetEditorViewportClient())
+		{
+			// Workaround to use while ITF and legacy input coexist.
+			// Editor Viewport Client processes ITF input before anything else.
+			// This means that, while ITF and legacy input coexist, simple clicks will be "stolen" by ITF Box Select, which
+			// is enabled by a simple Left Mouse Button click, with no modifiers. Because IE_Released events are used as
+			// "Clicks" by the Viewport Client only while tracking, but tracking will be false after this OnClickRelease,
+			// the Viewport Client will never use this IE_Release to trigger a click.
+			// This workaround checks the distance traveled during this Drag, and if it's close to zero, forwards a IE_Released
+			// event for left mouse button to the Viewport Client, bypassing the Internal_InputKey method.
+			{
+				UE::Math::TVector<double> Distance = End - Start;
+				const bool bForwardAsSimpleClick = FMath::IsNearlyZero(Distance.Length());
+
+				if (bForwardAsSimpleClick)
+				{
+					FSceneViewFamilyContext ViewFamily(
+						FSceneViewFamily::ConstructionValues(
+							EditorViewportClient->Viewport, EditorViewportClient->GetScene(), EditorViewportClient->EngineShowFlags
+						)
+							.SetRealtimeUpdate(EditorViewportClient->IsRealtime())
+					);
+					FSceneView* View = EditorViewportClient->CalcSceneView(&ViewFamily);
+
+					FIntPoint MousePos;
+					EditorViewportClient->Viewport->GetMousePos(MousePos);
+
+					HHitProxy* const Proxy = EditorViewportClient->Viewport->GetHitProxy(MousePos.X, MousePos.Y);
+
+					EditorViewportClient->ProcessClick(
+						*View, Proxy, EKeys::LeftMouseButton, IE_Released, MousePos.X, MousePos.Y
+					);
+
+					return;
+				}
+			}
+
+			UBrushEditingSubsystem* BrushSubsystem = GEditor->GetEditorSubsystem<UBrushEditingSubsystem>();
+			const bool bGeometryMode = BrushSubsystem ? BrushSubsystem->IsGeometryEditorModeActive() : false;
+
+			FScopedTransaction Transaction(NSLOCTEXT("ActorFrustumSelect", "MarqueeSelectTransaction", "Marquee Select"));
+
+			bool bShouldSelect = true;
+			FBox SelBBox;
+			CalculateBox(SelBBox);
+
+			const bool bControlDown = FInputDeviceState::IsCtrlKeyDown(InputState);
+			const bool bShiftDown = FInputDeviceState::IsShiftKeyDown(InputState);
+
+			if (bControlDown)
+			{
+				// If control is down remove from selection
+				bShouldSelect = false;
+			}
+			else if (!bShiftDown)
+			{
+				// If the user is selecting, but isn't holding down SHIFT, give modes a chance to clear selection
+				ModeTools->SelectNone();
+			}
+
+			constexpr bool bSelect = true;
+			// Let the editor mode try to handle the box selection.
+			const bool bEditorModeHandledBoxSelection = ModeTools->BoxSelect(SelBBox, bSelect);
+
+			// Let the component visualizers try to handle the selection.
+			const bool bComponentVisHandledSelection = !bEditorModeHandledBoxSelection
+													&& GUnrealEd->ComponentVisManager.HandleBoxSelect(
+														SelBBox, EditorViewportClient, EditorViewportClient->Viewport
+													);
+
+			// If the edit mode didn't handle the selection, try normal actor box selection.
+			if (!bEditorModeHandledBoxSelection && !bComponentVisHandledSelection)
+			{
+				const bool bStrictDragSelection = GetDefault<ULevelEditorViewportSettings>()->bStrictBoxSelection;
+
+				UTypedElementSelectionSet* SelectionSet =
+					GEditor->GetEditorSubsystem<ULevelEditorSubsystem>()->GetSelectionSet();
+				FTypedElementSelectionOptions ElementSelectionOption;
+				if (!bControlDown && !bShiftDown)
+				{
+					// If the user is selecting, but isn't holding down SHIFT, remove all current selections from the selection set
+					SelectionSet->ClearSelection(ElementSelectionOption);
+				}
+
+				FWorldSelectionElementArgs SelectionArgs{
+					SelectionSet,           ETypedElementSelectionMethod::Primary,
+					ElementSelectionOption, &(EditorViewportClient->EngineShowFlags),
+					bStrictDragSelection,   bGeometryMode
+				};
+
+				// Select all element that are within the selection box area.  Be aware that certain modes do special processing below.
+				bool bSelectionChanged = false;
+				UWorld* IteratorWorld = GWorld;
+				const TArray<FName>& HiddenLayers = EditorViewportClientProxy->GetHiddenLayers();
+				TArray<FTypedElementHandle> Handles;
+
+				for (FActorIterator It(IteratorWorld); It; ++It)
+				{
+					AActor* Actor = *It;
+
+					bool bActorIsVisible = true;
+					for (auto Layer : Actor->Layers)
+					{
+						// Check the actor isn't in one of the layers hidden from this viewport.
+						if (HiddenLayers.Contains(Layer))
+						{
+							bActorIsVisible = false;
+							break;
+						}
+					}
+
+					// Select the actor or its child elements
+					if (bActorIsVisible)
+					{
+						Handles.Append(
+							UE::EditorDragTools::Private::GetElementsIntersectingBox(
+								Actor, SelBBox, EditorViewportClientProxy, SelectionArgs
+							)
+						);
+					}
+				}
+
+				if (bShouldSelect)
+				{
+					SelectionSet->SelectElements(Handles, ElementSelectionOption);
+				}
+				else
+				{
+					SelectionSet->DeselectElements(Handles, ElementSelectionOption);
+				}
+
+				// Check every model to see if its BSP surfaces should be selected
+				for (int32 ModelIndex = 0; ModelIndex < ModelsToCheck.Num(); ++ModelIndex)
+				{
+					UModel& Model = *ModelsToCheck[ModelIndex];
+					// Check every node in the model
+					for (int32 NodeIndex = 0; NodeIndex < Model.Nodes.Num(); NodeIndex++)
+					{
+						if (IntersectsBox(Model, NodeIndex, SelBBox, bStrictDragSelection))
+						{
+							// If the node intersected the frustum select the corresponding surface
+							GEditor->SelectBSPSurf(&Model, Model.Nodes[NodeIndex].iSurf, bShouldSelect, false);
+							bSelectionChanged = true;
+						}
+					}
+				}
+
+				if (bSelectionChanged)
+				{
+					// If any selections were made.  Notify that now.
+					GEditor->NoteSelectionChange();
+				}
+			}
+
+			// Clear any hovered objects that might have been created while dragging
+			FLevelEditorViewportClient::ClearHoverFromObjects();
+		}
+	}
+
+	FEditorDragToolBehaviorTarget::OnClickRelease(InReleasePos);
+}
+
+void FEditorBoxSelect::OnTerminateDragSequence()
+{
+	FEditorDragToolBehaviorTarget::OnTerminateDragSequence();
+}
+
+TArray<FEditorModeID> FEditorBoxSelect::GetUnsupportedModes()
+{
+	const TArray<FEditorModeID> UnsupportedModes = { FBuiltinEditorModes::EM_Landscape, FBuiltinEditorModes::EM_Foliage };
+	return UnsupportedModes;
+}
+
+void FEditorBoxSelect::CalculateBox(FBox& OutBox) const
+{
+	if (!EditorViewportClientProxy)
+	{
+		return;
+	}
+
+	FEditorViewportClient* const EditorViewportClient = EditorViewportClientProxy->GetEditorViewportClient();
+	if (!EditorViewportClient)
+	{
+		return;
+	}
+
+	FSceneViewFamilyContext ViewFamily(
+		FSceneViewFamily::ConstructionValues(
+			EditorViewportClient->Viewport, EditorViewportClient->GetScene(), EditorViewportClient->EngineShowFlags
+		)
+			.SetRealtimeUpdate(EditorViewportClient->IsRealtime())
+	);
+
+	FSceneView* View = EditorViewportClient->CalcSceneView(&ViewFamily);
+
+	FVector3f StartFloat{ Start };
+	FVector3f EndFloat{ End };
+
+	FVector4 StartScreenPos = View->PixelToScreen(StartFloat.X, StartFloat.Y, 0);
+	FVector4 EndScreenPos = View->PixelToScreen(EndFloat.X, EndFloat.Y, 0);
+
+	FVector TransformedStart = View->ScreenToWorld(View->PixelToScreen(StartFloat.X, StartFloat.Y, 0.5f));
+	FVector TransformedEnd = View->ScreenToWorld(View->PixelToScreen(EndFloat.X, EndFloat.Y, 0.5f));
+
+	// Create a bounding box based on the start/end points (normalizes the points).
+	OutBox.Init();
+	OutBox += TransformedStart;
+	OutBox += TransformedEnd;
+
+	switch (EditorViewportClient->ViewportType)
+	{
+	case LVT_OrthoXY:
+	case LVT_OrthoNegativeXY:
+		OutBox.Min.Z = -WORLD_MAX;
+		OutBox.Max.Z = WORLD_MAX;
+		break;
+	case LVT_OrthoXZ:
+	case LVT_OrthoNegativeXZ:
+		OutBox.Min.Y = -WORLD_MAX;
+		OutBox.Max.Y = WORLD_MAX;
+		break;
+	case LVT_OrthoYZ:
+	case LVT_OrthoNegativeYZ:
+		OutBox.Min.X = -WORLD_MAX;
+		OutBox.Max.X = WORLD_MAX;
+		break;
+	case LVT_OrthoFreelook:
+	case LVT_Perspective:
+		break;
+	}
+}
+
+bool FEditorBoxSelect::IntersectsBox(const UModel& InModel, int32 InNodeIndex, const FBox& InBox, bool bInUseStrictSelection)
+{
+	FBox NodeBB;
+	InModel.GetNodeBoundingBox(InModel.Nodes[InNodeIndex], NodeBB);
+
+	bool bFullyContained = false;
+	bool bIntersects = false;
+	if (!bInUseStrictSelection)
+	{
+		bIntersects = InBox.Intersect(NodeBB);
+	}
+	else
+	{
+		bIntersects = InBox.IsInside(NodeBB.Max) && InBox.IsInside(NodeBB.Min);
+	}
+
+	return bIntersects;
+}
+
+void FEditorBoxSelect::AddHoverEffect(AActor& InActor)
+{
+	const FViewportHoverTarget HoverTarget(&InActor);
+	UE::EditorDragTools::Private::AddHoverEffect(HoverTarget);
+}
+
+void FEditorBoxSelect::AddHoverEffect(UModel& InModel, int32 InSurfIndex)
+{
+	FViewportHoverTarget HoverTarget(&InModel, InSurfIndex);
+	UE::EditorDragTools::Private::AddHoverEffect(HoverTarget);
+}
+
+void FEditorBoxSelect::RemoveHoverEffect(AActor& InActor)
+{
+	FViewportHoverTarget HoverTarget(&InActor);
+	UE::EditorDragTools::Private::RemoveHoverEffect(HoverTarget);
+}
+
+void FEditorBoxSelect::RemoveHoverEffect(UModel& InModel, int32 InSurfIndex)
+{
+	FViewportHoverTarget HoverTarget(&InModel, InSurfIndex);
+	UE::EditorDragTools::Private::RemoveHoverEffect(HoverTarget);
+}
diff --git a/Engine/Source/Editor/Experimental/EditorInteractiveToolsFramework/Private/EditorDragTools/EditorDragToolBehaviorTarget.cpp b/Engine/Source/Editor/Experimental/EditorInteractiveToolsFramework/Private/EditorDragTools/EditorDragToolBehaviorTarget.cpp
--- a/Engine/Source/Editor/Experimental/EditorInteractiveToolsFramework/Private/EditorDragTools/EditorDragToolBehaviorTarget.cpp
+++ b/Engine/Source/Editor/Experimental/EditorInteractiveToolsFramework/Private/EditorDragTools/EditorDragToolBehaviorTarget.cpp
@@ -1,10 +1,11 @@
 ﻿// Copyright Epic Games, Inc. All Rights Reserved.
 
 #include "EditorDragTools/EditorDragToolBehaviorTarget.h"
-#include "InputCoreTypes.h"
-#include "EditorViewportClient.h"
 #include "Editor.h"
+#include "EditorDragTools/EditorViewportClientProxy.h"
 #include "EditorModeManager.h"
+#include "EditorViewportClient.h"
+#include "InputCoreTypes.h"
 
 namespace UE::Editor::DragTools
 {
@@ -55,43 +56,55 @@ FOnEditorDragToolsToggleDelegate& OnEditorDragToolsDeactivated()
 //
 ///////////////////////////////////////////////////////////////////////////////////////////////////////////////////////
 
-FEditorDragToolBehaviorTarget::FEditorDragToolBehaviorTarget(FEditorModeTools* InModeTools)
-	: bConvertDelta(true)
-	, ModeTools(InModeTools)
-	, Start(FVector::ZeroVector)
-	, End(FVector::ZeroVector)
-	, bUseSnapping(false)
-	, bIsDragging(false)
+FEditorDragToolBehaviorTarget::FEditorDragToolBehaviorTarget(FEditorViewportClient* InEditorViewportClient)
 {
+	if (InEditorViewportClient)
+	{
+		ModeTools = InEditorViewportClient->GetModeTools();
+	}
+
+	EditorViewportClientProxy = IEditorViewportClientProxy::CreateViewportClientProxy(InEditorViewportClient);
 }
 
 void FEditorDragToolBehaviorTarget::OnUpdateModifierState(int InModifierID, bool bInIsOn)
 {
-	constexpr int ShiftKeyMod = 1;
-	constexpr int AltKeyMod = 2;
-	constexpr int CtrlKeyMod = 3;
-
 	switch (InModifierID)
 	{
-	case ShiftKeyMod:
+	case UE::Editor::DragTools::ShiftKeyMod:
 		InputState.bShiftKeyDown = bInIsOn;
 		break;
-	case CtrlKeyMod:
+	case UE::Editor::DragTools::CtrlKeyMod:
 		InputState.bCtrlKeyDown = bInIsOn;
 		break;
-	case AltKeyMod:
+	case UE::Editor::DragTools::AltKeyMod:
 		InputState.bAltKeyDown = bInIsOn;
 		break;
 	default:;
 	}
 }
 
+void FEditorDragToolBehaviorTarget::OnClickRelease(const FInputDeviceRay& InReleasePos)
+{
+	Start = End = FVector::ZeroVector;
+	bIsDragging = false;
+
+	// Signal that this tool is no longer active
+	OnDeactivateTool().Broadcast();
+}
+
 void FEditorDragToolBehaviorTarget::OnTerminateDragSequence()
 {
 	Start = End = FVector::ZeroVector;
 	bIsDragging = false;
 
-	InputState.SetModifierKeyStates(false, false, false, false);
+	InputState.bShiftKeyDown = false;
+	InputState.bCtrlKeyDown = false;
+	InputState.bAltKeyDown = false;
+
+	// Signal that this tool is no longer active
+	OnDeactivateTool().Broadcast();
+
+	EditorViewportClientProxy->GetEditorViewportClient()->Invalidate(true, false);
 }
 
 bool FEditorDragToolBehaviorTarget::IsActivationChordPressed(const FInputChord& InChord) const
diff --git a/Engine/Source/Editor/Experimental/EditorInteractiveToolsFramework/Private/EditorDragTools/EditorFrustumSelect.cpp b/Engine/Source/Editor/Experimental/EditorInteractiveToolsFramework/Private/EditorDragTools/EditorFrustumSelect.cpp
--- a/Engine/Source/Editor/Experimental/EditorInteractiveToolsFramework/Private/EditorDragTools/EditorFrustumSelect.cpp
+++ b/Engine/Source/Editor/Experimental/EditorInteractiveToolsFramework/Private/EditorDragTools/EditorFrustumSelect.cpp
@@ -6,6 +6,7 @@
 #include "CanvasTypes.h"
 #include "Components/PrimitiveComponent.h"
 #include "Editor/UnrealEdEngine.h"
+#include "EditorDragTools/EditorViewportClientProxy.h"
 #include "EditorModeManager.h"
 #include "EditorModes.h"
 #include "Elements/Framework/EngineElementsLibrary.h"
@@ -15,6 +16,7 @@
 #include "EngineUtils.h"
 #include "GameFramework/Volume.h"
 #include "LevelEditorSubsystem.h"
+#include "LevelEditorViewport.h"
 #include "Model.h"
 #include "SceneView.h"
 #include "ScopedTransaction.h"
@@ -28,14 +30,11 @@ namespace UE::EditorDragTools::Private
 TArray<FTypedElementHandle> GetElementsIntersectingFrustum(
 	const AActor* Actor,
 	const FConvexVolume& InFrustum,
-	const FEditorViewportClient* EditorViewport,
-	const FLevelEditorViewportClient* LevelViewport,
+	const IEditorViewportClientProxy* InEditorViewportClientProxy,
 	const FWorldSelectionElementArgs& SelectionArgs
 )
 {
-	if (Actor
-		&& (!EditorViewport || !Actor->IsA(AVolume::StaticClass())
-			|| (LevelViewport ? !LevelViewport->IsVolumeVisibleInViewport(*Actor) : false)))
+	if (InEditorViewportClientProxy && InEditorViewportClientProxy->IsActorVisible(Actor))
 	{
 		if (FTypedElementHandle ActorHandle = UEngineElementsLibrary::AcquireEditorActorElementHandle(Actor))
 		{
@@ -81,29 +80,26 @@ FInputRayHit FEditorFrustumSelect::CanBeginClickDragSequence(const FInputDeviceR
 
 void FEditorFrustumSelect::OnClickPress(const FInputDeviceRay& InPressPos)
 {
-	if (EditorViewportClient)
-	{
-		// Signal that this tool is now active
-		OnActivateTool().Broadcast();
+	// Signal that this tool is now active
+	OnActivateTool().Broadcast();
 
-		Start = InPressPos.WorldRay.Origin;
-		bIsDragging = true;
+	Start = InPressPos.WorldRay.Origin;
+	bIsDragging = true;
 
-		// Snap to constraints.
-		if (bUseSnapping)
-		{
-			const float GridSize = GEditor->GetGridSize();
-			const FVector GridBase(GridSize, GridSize, GridSize);
-			FSnappingUtils::SnapPointToGrid(Start, GridBase);
-		}
-		End = Start;
+	// Snap to constraints.
+	if (bUseSnapping)
+	{
+		const float GridSize = GEditor->GetGridSize();
+		const FVector GridBase(GridSize, GridSize, GridSize);
+		FSnappingUtils::SnapPointToGrid(Start, GridBase);
+	}
+	End = Start;
 
-		// Remove any active hover objects
-		FLevelEditorViewportClient::ClearHoverFromObjects();
+	// Remove any active hover objects
+	FLevelEditorViewportClient::ClearHoverFromObjects();
 
-		Start = FVector(InPressPos.ScreenPosition.X, InPressPos.ScreenPosition.Y, 0);
-		End = Start;
-	}
+	Start = FVector(InPressPos.ScreenPosition.X, InPressPos.ScreenPosition.Y, 0);
+	End = Start;
 }
 
 void FEditorFrustumSelect::OnClickDrag(const FInputDeviceRay& InDragPos)
@@ -115,6 +111,12 @@ void FEditorFrustumSelect::OnClickRelease(const FInputDeviceRay& InReleasePos)
 {
 	End = FVector(InReleasePos.ScreenPosition.X, InReleasePos.ScreenPosition.Y, 0);
 
+	if (!EditorViewportClientProxy)
+	{
+		return;
+	}
+
+	FEditorViewportClient* const EditorViewportClient = EditorViewportClientProxy->GetEditorViewportClient();
 	if (!EditorViewportClient)
 	{
 		return;
@@ -193,7 +195,7 @@ void FEditorFrustumSelect::OnClickRelease(const FInputDeviceRay& InReleasePos)
 				AActor* Actor = *It;
 				ElementsToSelect.Append(
 					UE::EditorDragTools::Private::GetElementsIntersectingFrustum(
-						Actor, Frustum, EditorViewportClient, LevelViewportClient, SeletionArgs
+						Actor, Frustum, EditorViewportClientProxy, SeletionArgs
 					)
 				);
 			}
@@ -283,7 +285,7 @@ void FEditorFrustumSelect::OnClickRelease(const FInputDeviceRay& InReleasePos)
 					{
 						ElementsToSelect.Append(
 							UE::EditorDragTools::Private::GetElementsIntersectingFrustum(
-								Actor, Frustum, EditorViewportClient, LevelViewportClient, SeletionArgs
+								Actor, Frustum, EditorViewportClientProxy, SeletionArgs
 							)
 						);
 					}
@@ -311,10 +313,7 @@ void FEditorFrustumSelect::OnClickRelease(const FInputDeviceRay& InReleasePos)
 	// Clear any hovered objects that might have been created while dragging
 	FLevelEditorViewportClient::ClearHoverFromObjects();
 
-	FEditorFrustumSelect::OnTerminateDragSequence();
-
-	// Signal that this tool is no longer active
-	OnDeactivateTool().Broadcast();
+	FEditorDragToolBehaviorTarget::OnClickRelease(InReleasePos);
 }
 
 void FEditorFrustumSelect::OnTerminateDragSequence()
@@ -348,7 +347,13 @@ bool FEditorFrustumSelect::IntersectsFrustum(
 
 void FEditorFrustumSelect::CalculateFrustum(const FSceneView* InView, FConvexVolume& OutFrustum, bool bUseBoxFrustum) const
 {
-	if (!InView || !EditorViewportClient)
+	if (!InView || !EditorViewportClientProxy)
+	{
+		return;
+	}
+
+	FEditorViewportClient* const EditorViewportClient = EditorViewportClientProxy->GetEditorViewportClient();
+	if (!EditorViewportClient)
 	{
 		return;
 	}
diff --git a/Engine/Source/Editor/Experimental/EditorInteractiveToolsFramework/Private/EditorDragTools/EditorMeasureTool.cpp b/Engine/Source/Editor/Experimental/EditorInteractiveToolsFramework/Private/EditorDragTools/EditorMeasureTool.cpp
new file mode 100644
--- /dev/null
+++ b/Engine/Source/Editor/Experimental/EditorInteractiveToolsFramework/Private/EditorDragTools/EditorMeasureTool.cpp
@@ -0,0 +1,187 @@
+// Copyright Epic Games, Inc. All Rights Reserved.
+
+#include "EditorDragTools/EditorMeasureTool.h"
+#include "CanvasItem.h"
+#include "CanvasTypes.h"
+#include "Editor.h"
+#include "EditorDragTools/EditorViewportClientProxy.h"
+#include "EditorViewportClient.h"
+#include "EngineGlobals.h"
+#include "SceneView.h"
+#include "Settings/EditorStyleSettings.h"
+#include "Settings/LevelEditorViewportSettings.h"
+#include "SnappingUtils.h"
+
+FEditorMeasureTool::FEditorMeasureTool(FEditorViewportClient* InViewportClient)
+	: FEditorDragToolBehaviorTarget(InViewportClient)
+	, PixelStart()
+	, PixelEnd()
+{
+	bUseSnapping = true;
+	bConvertDelta = false;
+}
+
+void FEditorMeasureTool::Render(const FSceneView* InView, FCanvas* InCanvas)
+{
+	FEditorViewportClient* const EditorViewportClient = EditorViewportClientProxy->GetEditorViewportClient();
+	if (!EditorViewportClient)
+	{
+		return;
+	}
+
+	if (!EditorViewportClient->Viewport)
+	{
+		return;
+	}
+
+	const float OrthoUnitsPerPixel = EditorViewportClient->GetOrthoUnitsPerPixel(EditorViewportClient->Viewport);
+	const float Length =
+		FMath::RoundToFloat(FVector2f{ PixelEnd - PixelStart }.Size() * OrthoUnitsPerPixel * EditorViewportClient->GetDPIScale());
+
+	if (InCanvas && Length >= 1.0f)
+	{
+		const FLinearColor ToolColor = GetDefault<UEditorStyleSettings>()->ViewportToolOverlayColor;
+		FCanvasLineItem LineItem(PixelStart, PixelEnd);
+		LineItem.SetColor(ToolColor);
+		InCanvas->DrawItem(LineItem);
+
+		const FVector2D PixelMid = FVector2D(PixelStart + ((PixelEnd - PixelStart) / 2));
+
+		// Calculate number of decimal places to display, based on the current viewport zoom
+		float Divisor = 1.0f;
+		int32 DecimalPlaces = 0;
+		const float OrderOfMagnitude = FMath::LogX(10.0f, OrthoUnitsPerPixel);
+
+		switch (GetDefault<ULevelEditorViewportSettings>()->MeasuringToolUnits)
+		{
+		case MeasureUnits_Meters:
+			Divisor = 100.0f;
+			// Max one decimal place allowed for meters.
+			DecimalPlaces = FMath::Clamp(FMath::FloorToInt(1.5f - OrderOfMagnitude), 0, 1);
+			break;
+
+		case MeasureUnits_Kilometers:
+			Divisor = 100000.0f;
+			// Max two decimal places allowed for kilometers.
+			DecimalPlaces = FMath::Clamp(FMath::FloorToInt(4.5f - OrderOfMagnitude), 0, 2);
+			break;
+		}
+
+		FNumberFormattingOptions Options;
+		Options.UseGrouping = false;
+		Options.MinimumFractionalDigits = DecimalPlaces;
+		Options.MaximumFractionalDigits = DecimalPlaces;
+
+		const FText LengthStr = FText::AsNumber(Length / Divisor, &Options);
+
+		FCanvasTextItem TextItem(
+			FVector2D(FMath::FloorToFloat(PixelMid.X), FMath::FloorToFloat(PixelMid.Y)), LengthStr, GEngine->GetSmallFont(), ToolColor
+		);
+		TextItem.bCentreX = true;
+		InCanvas->DrawItem(TextItem);
+	}
+}
+
+FInputRayHit FEditorMeasureTool::CanBeginClickDragSequence(const FInputDeviceRay& InPressPos)
+{
+	if (FEditorViewportClient* const EditorViewportClient = EditorViewportClientProxy->GetEditorViewportClient())
+	{
+		if (IsCurrentModeSupported())
+		{
+			// This Box Selection is for Orthographic viewports only
+			if (EditorViewportClient->IsOrtho() && ! IEditorViewportClientProxy::IsMouseOnWidgetAxis(EditorViewportClientProxy))
+			{
+				return FInputRayHit(TNumericLimits<float>::Max());
+			}
+		}
+	}
+
+	return FInputRayHit();
+}
+
+void FEditorMeasureTool::OnClickPress(const FInputDeviceRay& InPressPos)
+{
+	FEditorViewportClient* const EditorViewportClient = EditorViewportClientProxy->GetEditorViewportClient();
+	if (EditorViewportClient && EditorViewportClient->Viewport)
+	{
+		OnActivateTool().Broadcast();
+
+		bIsDragging = true;
+
+		FIntPoint MousePos;
+		EditorViewportClient->Viewport->GetMousePos(MousePos);
+
+		// Take into account DPI scale for drawing lines properly when scale is not 1.0
+		Start = FVector(MousePos.X, MousePos.Y, 0) / EditorViewportClient->GetDPIScale();
+
+		// Snap to constraints.
+		if (bUseSnapping)
+		{
+			const float GridSize = GEditor->GetGridSize();
+			const FVector GridBase(GridSize, GridSize, GridSize);
+			FSnappingUtils::SnapPointToGrid(Start, GridBase);
+		}
+
+		End = Start;
+
+		PixelStart = GetSnappedPixelPos(FVector2D(InPressPos.ScreenPosition));
+		PixelEnd = PixelStart;
+	}
+}
+
+void FEditorMeasureTool::OnClickDrag(const FInputDeviceRay& InDragPos)
+{
+	if (FEditorViewportClient* const EditorViewportClient = EditorViewportClientProxy->GetEditorViewportClient())
+	{
+		if (EditorViewportClient->Viewport)
+		{
+			FIntPoint MousePos;
+			EditorViewportClient->Viewport->GetMousePos(MousePos);
+			PixelEnd = GetSnappedPixelPos(FVector2D(MousePos));
+		}
+	}
+}
+
+void FEditorMeasureTool::OnClickRelease(const FInputDeviceRay& InReleasePos)
+{
+	FEditorDragToolBehaviorTarget::OnClickRelease(InReleasePos);
+}
+
+void FEditorMeasureTool::OnTerminateDragSequence()
+{
+	FEditorDragToolBehaviorTarget::OnTerminateDragSequence();
+}
+
+FVector2D FEditorMeasureTool::GetSnappedPixelPos(FVector2D InPixelPos) const
+{
+	FEditorViewportClient* const EditorViewportClient = EditorViewportClientProxy->GetEditorViewportClient();
+	if (!EditorViewportClient)
+	{
+		return FVector2D::Zero();
+	}
+
+	FSceneViewFamilyContext ViewFamily(
+		FSceneViewFamily::ConstructionValues(EditorViewportClient->Viewport, EditorViewportClient->GetScene(), EditorViewportClient->EngineShowFlags)
+			.SetRealtimeUpdate(EditorViewportClient->IsRealtime())
+	);
+
+	FSceneView* View = EditorViewportClient->CalcSceneView(&ViewFamily);
+
+	// Put the mouse pos in world space
+	FVector2f PixelPosFloat{ InPixelPos };
+	FVector WorldPos = View->ScreenToWorld(View->PixelToScreen(PixelPosFloat.X, PixelPosFloat.Y, 0.5f));
+
+	// Snap the world position
+	const float GridSize = GEditor->GetGridSize();
+	const FVector GridBase(GridSize, GridSize, GridSize);
+	FSnappingUtils::SnapPointToGrid(WorldPos, GridBase);
+
+	// And back into pixel-space (might fail, in which case we return the original)
+	View->WorldToPixel(WorldPos, InPixelPos);
+
+	// The canvas we're going to render to factors in dpi scale to final position.
+	// Since we're basing our position off mouse coordinates it will already be pixel accurate and therefore we must back out the scale the canvas will apply
+	InPixelPos /= EditorViewportClient->GetDPIScale();
+
+	return InPixelPos;
+}
diff --git a/Engine/Source/Editor/Experimental/EditorInteractiveToolsFramework/Private/EditorDragTools/EditorViewportChange.cpp b/Engine/Source/Editor/Experimental/EditorInteractiveToolsFramework/Private/EditorDragTools/EditorViewportChange.cpp
new file mode 100644
--- /dev/null
+++ b/Engine/Source/Editor/Experimental/EditorInteractiveToolsFramework/Private/EditorDragTools/EditorViewportChange.cpp
@@ -0,0 +1,202 @@
+﻿// Copyright Epic Games, Inc. All Rights Reserved.
+
+#include "EditorDragTools/EditorViewportChange.h"
+#include "CanvasItem.h"
+#include "CanvasTypes.h"
+#include "EditorDragTools/EditorViewportClientProxy.h"
+#include "EditorViewportClient.h"
+#include "Framework/Commands/InputChord.h"
+#include "Settings/EditorStyleSettings.h"
+#include "SnappingUtils.h"
+
+FEditorViewportChange::FEditorViewportChange(FEditorViewportClient* InEditorViewportClient)
+	: FEditorDragToolBehaviorTarget(InEditorViewportClient)
+	, ViewOptionOffset(FVector2D(0.0f, 0.0f))
+{
+	bUseSnapping = true;
+	bConvertDelta = false;
+}
+
+void FEditorViewportChange::Render(const FSceneView* View, FCanvas* Canvas)
+{
+	FCanvasLineItem LineItem(Start, End);
+	Canvas->DrawItem(LineItem);
+
+	const FLinearColor ToolColor = GetDefault<UEditorStyleSettings>()->ViewportToolOverlayColor;
+	FCanvasTextItem TextItem(
+		FVector2D(FMath::FloorToFloat(End.X), FMath::FloorToFloat(End.Y) + 20),
+		GetDesiredViewportTypeText(),
+		GEngine->GetMediumFont(),
+		ToolColor
+	);
+	TextItem.bCentreX = true;
+	Canvas->DrawItem(TextItem);
+}
+
+FInputRayHit FEditorViewportChange::CanBeginClickDragSequence(const FInputDeviceRay& InPressPos)
+{
+	// Todo: this could be retrieved from a command for customization
+	FInputChord ActivationChord(EModifierKey::Control, EKeys::MiddleMouseButton);
+
+	return IsActivationChordPressed(ActivationChord)
+			 ? FInputRayHit(TNumericLimits<float>::Max()) // bHit is true. Depth is max to lose the standard tiebreaker.
+			 : FInputRayHit();
+}
+
+void FEditorViewportChange::OnClickPress(const FInputDeviceRay& InPressPos)
+{
+	FEditorViewportClient* const EditorViewportClient = EditorViewportClientProxy->GetEditorViewportClient();
+	if (EditorViewportClient && EditorViewportClient->Viewport)
+	{
+		OnActivateTool().Broadcast();
+
+		bIsDragging = true;
+
+		FIntPoint MousePos;
+		EditorViewportClient->Viewport->GetMousePos(MousePos);
+
+		// Take into account DPI scale for drawing lines properly when scale is not 1.0
+		Start = FVector(MousePos.X, MousePos.Y, 0) / EditorViewportClient->GetDPIScale();
+
+		// Snap to constraints.
+		if (bUseSnapping)
+		{
+			const float GridSize = GEditor->GetGridSize();
+			const FVector GridBase(GridSize, GridSize, GridSize);
+			FSnappingUtils::SnapPointToGrid(Start, GridBase);
+		}
+
+		End = Start;
+	}
+}
+
+void FEditorViewportChange::OnClickDrag(const FInputDeviceRay& InDragPos)
+{
+	if (FEditorViewportClient* const EditorViewportClient = EditorViewportClientProxy->GetEditorViewportClient())
+	{
+		FIntPoint MousePos;
+		EditorViewportClient->Viewport->GetMousePos(MousePos);
+
+		// Take into account DPI scale for drawing lines properly when scale is not 1.0
+		End = FVector(MousePos) / EditorViewportClient->GetDPIScale();
+
+		ViewOptionOffset.X = End.X - Start.X;
+		ViewOptionOffset.Y = End.Y - Start.Y;
+	}
+}
+
+void FEditorViewportChange::OnClickRelease(const FInputDeviceRay& InReleasePos)
+{
+	if (FEditorViewportClient* const EditorViewportClient = EditorViewportClientProxy->GetEditorViewportClient())
+	{
+		EditorViewportClient->SetViewportType(GetDesiredViewportType());
+	}
+
+	FEditorDragToolBehaviorTarget::OnClickRelease(InReleasePos);
+}
+
+void FEditorViewportChange::OnTerminateDragSequence()
+{
+	FEditorDragToolBehaviorTarget::OnTerminateDragSequence();
+}
+
+ELevelViewportType FEditorViewportChange::GetDesiredViewportType() const
+{
+	ELevelViewportType ViewOption = LVT_Perspective;
+
+	if (ViewOptionOffset.Y == 0)
+	{
+		if (ViewOptionOffset.X == 0)
+		{
+			ViewOption = LVT_Perspective;
+		}
+		else if (ViewOptionOffset.X > 0)
+		{
+			ViewOption = LVT_OrthoRight; // Right
+		}
+		else
+		{
+			ViewOption = LVT_OrthoLeft; // Left
+		}
+	}
+	else
+	{
+		double OffsetRatio = ViewOptionOffset.X / ViewOptionOffset.Y;
+		double DragAngle = FMath::RadiansToDegrees(FMath::Atan(OffsetRatio));
+
+		if (ViewOptionOffset.Y >= 0)
+		{
+			if (DragAngle >= -15.f && DragAngle <= 15.f)
+			{
+				ViewOption = LVT_OrthoBottom; // Bottom
+			}
+			else if (DragAngle > 75.f)
+			{
+				ViewOption = LVT_OrthoRight; // Right
+			}
+			else if (DragAngle < -75.f)
+			{
+				ViewOption = LVT_OrthoLeft; // Left
+			}
+		}
+		else
+		{
+			if (DragAngle >= -15.f && DragAngle < 15.f)
+			{
+				ViewOption = LVT_OrthoTop; // Top
+			}
+			else if (DragAngle >= 15.f && DragAngle < 75.f)
+			{
+				ViewOption = LVT_OrthoFront; // Front
+			}
+			else if (DragAngle >= -75.f && DragAngle < -15.f)
+			{
+				ViewOption = LVT_OrthoBack; // Back
+			}
+			else if (DragAngle >= 75.f)
+			{
+				ViewOption = LVT_OrthoLeft; // Left
+			}
+			else if (DragAngle <= -75.f)
+			{
+				ViewOption = LVT_OrthoRight; // Right
+			}
+		}
+	}
+
+	return ViewOption;
+}
+
+FText FEditorViewportChange::GetDesiredViewportTypeText() const
+{
+	switch (GetDesiredViewportType())
+	{
+	case LVT_Perspective:
+		return FText::FromString("Perspective");
+		break;
+	case LVT_OrthoFreelook:
+		return FText::FromString("Free Look");
+		break;
+	case LVT_OrthoTop:
+		return FText::FromString("Top");
+		break;
+	case LVT_OrthoLeft:
+		return FText::FromString("Left");
+		break;
+	case LVT_OrthoFront:
+		return FText::FromString("Front");
+		break;
+	case LVT_OrthoBack:
+		return FText::FromString("Back");
+		break;
+	case LVT_OrthoBottom:
+		return FText::FromString("Bottom");
+		break;
+	case LVT_OrthoRight:
+		return FText::FromString("Right");
+		break;
+	default:;
+	}
+
+	return FText();
+}
diff --git a/Engine/Source/Editor/Experimental/EditorInteractiveToolsFramework/Private/EditorDragTools/EditorViewportClientProxy.cpp b/Engine/Source/Editor/Experimental/EditorInteractiveToolsFramework/Private/EditorDragTools/EditorViewportClientProxy.cpp
new file mode 100644
--- /dev/null
+++ b/Engine/Source/Editor/Experimental/EditorInteractiveToolsFramework/Private/EditorDragTools/EditorViewportClientProxy.cpp
@@ -0,0 +1,103 @@
+﻿// Copyright Epic Games, Inc. All Rights Reserved.
+
+#include "EditorDragTools/EditorViewportClientProxy.h"
+#include "GameFramework/Volume.h"
+#include "LevelEditorViewport.h"
+#include "UnrealWidget.h"
+
+IEditorViewportClientProxy* IEditorViewportClientProxy::CreateViewportClientProxy(FEditorViewportClient* InViewportClient)
+{
+	if (InViewportClient->IsLevelEditorClient())
+	{
+		if (FLevelEditorViewportClient* LevelEditorViewportClient = static_cast<FLevelEditorViewportClient*>(InViewportClient))
+		{
+			return new FLevelEditorViewportClientProxy(LevelEditorViewportClient);
+		}
+	}
+
+	return new FEditorViewportClientProxy(InViewportClient);
+}
+
+bool IEditorViewportClientProxy::IsMouseOnWidgetAxis(IEditorViewportClientProxy* InEditorViewportClientProxy)
+{
+	if (InEditorViewportClientProxy)
+	{
+		if (FEditorViewportClient* const EditorViewportClient = InEditorViewportClientProxy->GetEditorViewportClient())
+		{
+			if (FViewport* const Viewport = EditorViewportClient->Viewport)
+			{
+				FIntPoint MousePos;
+				EditorViewportClient->Viewport->GetMousePos(MousePos);
+
+				// Since some drag tool does not involve any keyboard modifier (Alt, Shift, Ctrl),
+				// in some cases we need to make sure the user is not hovering over widget axis proxies.
+				// Intercepting input in that case prevents TRS gizmos from working.
+				if (const HHitProxy* const Proxy = Viewport->GetHitProxy(MousePos.X, MousePos.Y))
+				{
+					// We hit a widget axis, so we don't start the drag sequence
+					if (Proxy->IsA(HWidgetAxis::StaticGetType()))
+					{
+						return true;
+					}
+				}
+			}
+		}
+	}
+
+	return false;
+}
+
+FEditorViewportClientProxy::FEditorViewportClientProxy(FEditorViewportClient* InEditorViewportClient)
+	: EditorViewportClient(InEditorViewportClient)
+{
+}
+
+bool FEditorViewportClientProxy::IsActorVisible(const AActor* InActor) const
+{
+	if (InActor)
+	{
+		return !InActor->IsA(AVolume::StaticClass());
+	}
+
+	return false;
+}
+
+const TArray<FName> FEditorViewportClientProxy::GetHiddenLayers() const
+{
+	return TArray<FName>();
+}
+
+FEditorViewportClient* FEditorViewportClientProxy::GetEditorViewportClient()
+{
+	return EditorViewportClient.Get();
+}
+
+FLevelEditorViewportClientProxy::FLevelEditorViewportClientProxy(FLevelEditorViewportClient* InLevelEditorViewportClient)
+	:LevelEditorViewportClient(InLevelEditorViewportClient)
+{
+}
+
+bool FLevelEditorViewportClientProxy::IsActorVisible(const AActor* InActor) const
+{
+	if (InActor && LevelEditorViewportClient)
+	{
+		return !InActor->IsA(AVolume::StaticClass()) || !LevelEditorViewportClient->IsVolumeVisibleInViewport(*InActor);
+	}
+
+	return false;
+}
+
+const TArray<FName> FLevelEditorViewportClientProxy::GetHiddenLayers() const
+{
+	if (LevelEditorViewportClient)
+	{
+		return LevelEditorViewportClient->ViewHiddenLayers;
+	}
+
+	return TArray<FName>();
+}
+
+FEditorViewportClient* FLevelEditorViewportClientProxy::GetEditorViewportClient()
+{
+	return LevelEditorViewportClient.Get();
+}
diff --git a/Engine/Source/Editor/Experimental/EditorInteractiveToolsFramework/Public/EditorDragTools/EditorBoxSelect.h b/Engine/Source/Editor/Experimental/EditorInteractiveToolsFramework/Public/EditorDragTools/EditorBoxSelect.h
new file mode 100644
--- /dev/null
+++ b/Engine/Source/Editor/Experimental/EditorInteractiveToolsFramework/Public/EditorDragTools/EditorBoxSelect.h
@@ -0,0 +1,63 @@
+// Copyright Epic Games, Inc. All Rights Reserved.
+
+#pragma once
+
+#include "EditorDragToolBehaviorTarget.h"
+
+class FCanvas;
+class UModel;
+
+class FEditorBoxSelect : public FEditorDragToolBehaviorTarget
+{
+public:
+	explicit FEditorBoxSelect(FEditorViewportClient* const InEditorViewportClient)
+		: FEditorDragToolBehaviorTarget(InEditorViewportClient)
+	{
+	}
+
+	virtual void Render(const FSceneView* View, FCanvas* Canvas) override;
+
+	//~ Begin IClickDragBehaviorTarget
+	virtual FInputRayHit CanBeginClickDragSequence(const FInputDeviceRay& InPressPos) override;
+	virtual void OnClickPress(const FInputDeviceRay& InPressPos) override;
+	virtual void OnClickDrag(const FInputDeviceRay& InDragPos) override;
+	virtual void OnClickRelease(const FInputDeviceRay& InReleasePos) override;
+	virtual void OnTerminateDragSequence() override;
+	//~ End IClickDragBehaviorTarget
+
+protected:
+	virtual TArray<FEditorModeID> GetUnsupportedModes() override;
+
+private:
+	/**
+	 * Returns true if the provided BSP node intersects with the provided frustum
+	 *
+	 * @param InModel				The model containing BSP nodes to check
+	 * @param InNodeIndex				The index to a BSP node in the model.  This node is used for the bounds check.
+	 * @param InBox			The frustum to check against.
+	 * @param bInUseStrictSelection	true if the node must be entirely within the frustum
+	 */
+	static bool IntersectsBox(const UModel& InModel, int32 InNodeIndex, const FBox& InBox, bool bInUseStrictSelection);
+
+	/** Adds a hover effect to the passed in actor */
+	static void AddHoverEffect(AActor& InActor);
+
+	/** Adds a hover effect to the passed in bsp surface */
+	static void AddHoverEffect(UModel& InModel, int32 InSurfIndex);
+
+	/** Removes a hover effect from the passed in actor */
+	static void RemoveHoverEffect(AActor& InActor);
+
+	/** Removes a hover effect from the passed in bsp surface */
+	static void RemoveHoverEffect(UModel& InModel, int32 InSurfIndex);
+
+	/**
+	 * Calculates a box to check actors against
+	 *
+	 * @param OutBox	The created box.
+	 */
+	void CalculateBox(FBox& OutBox) const;
+
+	/** List of BSP models to check for selection */
+	TArray<UModel*> ModelsToCheck;
+};
diff --git a/Engine/Source/Editor/Experimental/EditorInteractiveToolsFramework/Public/EditorDragTools/EditorDragToolBehaviorTarget.h b/Engine/Source/Editor/Experimental/EditorInteractiveToolsFramework/Public/EditorDragTools/EditorDragToolBehaviorTarget.h
--- a/Engine/Source/Editor/Experimental/EditorInteractiveToolsFramework/Public/EditorDragTools/EditorDragToolBehaviorTarget.h
+++ b/Engine/Source/Editor/Experimental/EditorInteractiveToolsFramework/Public/EditorDragTools/EditorDragToolBehaviorTarget.h
@@ -5,6 +5,8 @@
 #include "BaseBehaviors/BehaviorTargetInterfaces.h"
 #include "Editor.h"
 
+class IEditorViewportClientProxy;
+struct FWorldSelectionElementArgs;
 class FCanvas;
 class FEditorModeTools;
 class FEditorViewportClient;
@@ -29,6 +31,10 @@ EDITORINTERACTIVETOOLSFRAMEWORK_API bool UseEditorDragTools();
 EDITORINTERACTIVETOOLSFRAMEWORK_API FOnEditorDragToolsToggleDelegate& OnEditorDragToolsActivated();
 EDITORINTERACTIVETOOLSFRAMEWORK_API FOnEditorDragToolsToggleDelegate& OnEditorDragToolsDeactivated();
 
+constexpr int ShiftKeyMod = 1;
+constexpr int AltKeyMod = 2;
+constexpr int CtrlKeyMod = 3;
+
 } // namespace UE::Editor::DragTools
 
 /**
@@ -40,10 +46,7 @@ EDITORINTERACTIVETOOLSFRAMEWORK_API FOnEditorDragToolsToggleDelegate& OnEditorDr
 class FEditorDragToolBehaviorTarget : public IClickDragBehaviorTarget
 {
 public:
-	EDITORINTERACTIVETOOLSFRAMEWORK_API FEditorDragToolBehaviorTarget(FEditorModeTools* InModeTools);
-	virtual ~FEditorDragToolBehaviorTarget() override
-	{
-	}
+	FEditorDragToolBehaviorTarget(FEditorViewportClient* InEditorViewportClient);
 
 	virtual void Render(const FSceneView* InView, FCanvas* InCanvas)
 	{
@@ -83,9 +86,7 @@ public:
 	{
 	}
 
-	virtual void OnClickRelease(const FInputDeviceRay& InReleasePos) override
-	{
-	}
+	virtual void OnClickRelease(const FInputDeviceRay& InReleasePos) override;
 
 	virtual void OnTerminateDragSequence() override;
 	//~ End IClickDragBehaviorTarget
@@ -120,4 +121,6 @@ protected:
 
 	FOnToolStateChange OnToolActivatedDelegate;
 	FOnToolStateChange OnToolDeactivatedDelegate;
+
+	IEditorViewportClientProxy* EditorViewportClientProxy;
 };
diff --git a/Engine/Source/Editor/Experimental/EditorInteractiveToolsFramework/Public/EditorDragTools/EditorFrustumSelect.h b/Engine/Source/Editor/Experimental/EditorInteractiveToolsFramework/Public/EditorDragTools/EditorFrustumSelect.h
--- a/Engine/Source/Editor/Experimental/EditorInteractiveToolsFramework/Public/EditorDragTools/EditorFrustumSelect.h
+++ b/Engine/Source/Editor/Experimental/EditorInteractiveToolsFramework/Public/EditorDragTools/EditorFrustumSelect.h
@@ -3,10 +3,13 @@
 #pragma once
 
 #include "EditorDragToolBehaviorTarget.h"
-#include "LevelEditorViewport.h"
+
 
 class FCanvas;
+class FEditorViewportClient;
+class FSceneView;
 class UModel;
+struct FConvexVolume;
 
 /**
  * Draws a box in the current viewport and when the mouse button is released,
@@ -15,17 +18,8 @@ class UModel;
 class FEditorFrustumSelect : public FEditorDragToolBehaviorTarget
 {
 public:
-	EDITORINTERACTIVETOOLSFRAMEWORK_API explicit FEditorFrustumSelect(FLevelEditorViewportClient* InLevelViewportClient)
-		: FEditorDragToolBehaviorTarget(InLevelViewportClient->GetModeTools())
-		, LevelViewportClient(InLevelViewportClient)
-		, EditorViewportClient(InLevelViewportClient)
-	{
-	}
-
-	EDITORINTERACTIVETOOLSFRAMEWORK_API explicit FEditorFrustumSelect(FEditorViewportClient* InEditorViewportClient)
-		: FEditorDragToolBehaviorTarget(InEditorViewportClient->GetModeTools())
-		, LevelViewportClient(nullptr)
-		, EditorViewportClient(InEditorViewportClient)
+	explicit FEditorFrustumSelect(FEditorViewportClient* const InEditorViewportClient)
+		: FEditorDragToolBehaviorTarget(InEditorViewportClient)
 	{
 	}
 
@@ -61,7 +55,4 @@ private:
 	 * @param bUseBoxFrustum	If true a frustum out of the current dragged box will be created.  false will use the view frustum.
 	 */
 	void CalculateFrustum(const FSceneView* InView, FConvexVolume& OutFrustum, bool bUseBoxFrustum) const;
-
-	FLevelEditorViewportClient* LevelViewportClient;
-	FEditorViewportClient* EditorViewportClient;
 };
diff --git a/Engine/Source/Editor/Experimental/EditorInteractiveToolsFramework/Public/EditorDragTools/EditorMeasureTool.h b/Engine/Source/Editor/Experimental/EditorInteractiveToolsFramework/Public/EditorDragTools/EditorMeasureTool.h
new file mode 100644
--- /dev/null
+++ b/Engine/Source/Editor/Experimental/EditorInteractiveToolsFramework/Public/EditorDragTools/EditorMeasureTool.h
@@ -0,0 +1,32 @@
+// Copyright Epic Games, Inc. All Rights Reserved.
+
+#pragma once
+
+#include "EditorDragToolBehaviorTarget.h"
+
+class FCanvas;
+
+class FEditorMeasureTool : public FEditorDragToolBehaviorTarget
+{
+public:
+	explicit FEditorMeasureTool(FEditorViewportClient* InViewportClient);
+
+	virtual void Render(const FSceneView* InView,FCanvas* InCanvas) override;
+
+	//~ Begin IClickDragBehaviorTarget
+	virtual FInputRayHit CanBeginClickDragSequence(const FInputDeviceRay& InPressPos) override;
+	virtual void OnClickPress(const FInputDeviceRay& InPressPos) override;
+	virtual void OnClickDrag(const FInputDeviceRay& InDragPos) override;
+	virtual void OnClickRelease(const FInputDeviceRay& InReleasePos) override;
+	virtual void OnTerminateDragSequence() override;
+	//~ End IClickDragBehaviorTarget
+
+private:
+	/**
+	 * Gets the world-space snapped position of the specified pixel position
+	 */
+	FVector2D GetSnappedPixelPos(FVector2D InPixelPos) const;
+
+	/** Pixel-space positions for start and end */
+	FVector2D PixelStart, PixelEnd;
+};
diff --git a/Engine/Source/Editor/Experimental/EditorInteractiveToolsFramework/Public/EditorDragTools/EditorViewportChange.h b/Engine/Source/Editor/Experimental/EditorInteractiveToolsFramework/Public/EditorDragTools/EditorViewportChange.h
new file mode 100644
--- /dev/null
+++ b/Engine/Source/Editor/Experimental/EditorInteractiveToolsFramework/Public/EditorDragTools/EditorViewportChange.h
@@ -0,0 +1,32 @@
+// Copyright Epic Games, Inc. All Rights Reserved.
+
+#pragma once
+
+#include "Editor/UnrealEdTypes.h"
+#include "EditorDragToolBehaviorTarget.h"
+
+class FCanvas;
+class FEditorViewportClient;
+class FLevelEditorViewportClient;
+class FSceneView;
+
+class FEditorViewportChange : public FEditorDragToolBehaviorTarget
+{
+public:
+	explicit FEditorViewportChange(FEditorViewportClient* InEditorViewportClient);
+	virtual void Render(const FSceneView* View, FCanvas* Canvas) override;
+
+	//~ Begin IClickDragBehaviorTarget
+	virtual FInputRayHit CanBeginClickDragSequence(const FInputDeviceRay& InPressPos) override;
+	virtual void OnClickPress(const FInputDeviceRay& InPressPos) override;
+	virtual void OnClickDrag(const FInputDeviceRay& InDragPos) override;
+	virtual void OnClickRelease(const FInputDeviceRay& InReleasePos) override;
+	virtual void OnTerminateDragSequence() override;
+	//~ End IClickDragBehaviorTarget
+
+private:
+	ELevelViewportType GetDesiredViewportType() const;
+	FText GetDesiredViewportTypeText() const;
+
+	FVector2D ViewOptionOffset;
+};
diff --git a/Engine/Source/Editor/Experimental/EditorInteractiveToolsFramework/Public/EditorDragTools/EditorViewportClientProxy.h b/Engine/Source/Editor/Experimental/EditorInteractiveToolsFramework/Public/EditorDragTools/EditorViewportClientProxy.h
new file mode 100644
--- /dev/null
+++ b/Engine/Source/Editor/Experimental/EditorInteractiveToolsFramework/Public/EditorDragTools/EditorViewportClientProxy.h
@@ -0,0 +1,52 @@
+﻿// Copyright Epic Games, Inc. All Rights Reserved.
+
+#pragma once
+
+#include "Templates/SharedPointer.h"
+
+class AActor;
+class FEditorViewportClient;
+class FLevelEditorViewportClient;
+struct FWorldSelectionElementArgs;
+
+class IEditorViewportClientProxy : public TSharedFromThis<IEditorViewportClientProxy>
+{
+public:
+	virtual ~IEditorViewportClientProxy() = default;
+	static IEditorViewportClientProxy* CreateViewportClientProxy(FEditorViewportClient* InViewportClient);
+	static bool IsMouseOnWidgetAxis(IEditorViewportClientProxy* InEditorViewportClientProxy);
+
+
+	virtual bool IsActorVisible(const AActor* InActor) const = 0;
+
+	virtual const TArray<FName> GetHiddenLayers() const = 0;
+
+	virtual FEditorViewportClient* GetEditorViewportClient() = 0;
+};
+
+class FEditorViewportClientProxy : public IEditorViewportClientProxy
+{
+public:
+	FEditorViewportClientProxy(FEditorViewportClient* InEditorViewportClient);
+	
+	virtual bool IsActorVisible(const AActor* InActor) const override;
+	virtual const TArray<FName> GetHiddenLayers() const override;
+	virtual FEditorViewportClient* GetEditorViewportClient() override;
+
+private:
+	TSharedPtr<FEditorViewportClient> EditorViewportClient;
+};
+
+class FLevelEditorViewportClientProxy : public IEditorViewportClientProxy
+{
+public:
+	FLevelEditorViewportClientProxy(FLevelEditorViewportClient *InLevelEditorViewportClient);
+
+	virtual bool IsActorVisible(const AActor* InActor) const override;
+	virtual const TArray<FName> GetHiddenLayers() const override;
+	virtual FEditorViewportClient* GetEditorViewportClient() override;
+
+private:
+	TSharedPtr<FLevelEditorViewportClient> LevelEditorViewportClient;
+};
+
diff --git a/Engine/Source/Editor/LevelEditor/Private/LevelEditorContextMenu.cpp b/Engine/Source/Editor/LevelEditor/Private/LevelEditorContextMenu.cpp
--- a/Engine/Source/Editor/LevelEditor/Private/LevelEditorContextMenu.cpp
+++ b/Engine/Source/Editor/LevelEditor/Private/LevelEditorContextMenu.cpp
@@ -985,22 +985,22 @@ void FLevelEditorContextMenu::SummonViewOptionMenu( const TSharedRef< SLevelEdit
 
 	switch (ViewOption)
 	{
-		case LVT_OrthoNegativeXY:
+		case LVT_OrthoBottom:
 			ViewOptionType = EViewOptionType::Bottom;
 			break;
-		case LVT_OrthoNegativeXZ:
+		case LVT_OrthoBack:
 			ViewOptionType = EViewOptionType::Back;
 			break;
-		case LVT_OrthoNegativeYZ:
+		case LVT_OrthoRight:
 			ViewOptionType = EViewOptionType::Right;
 			break;
-		case LVT_OrthoXY:
+		case LVT_OrthoTop:
 			ViewOptionType = EViewOptionType::Top;
 			break;
-		case LVT_OrthoXZ:
+		case LVT_OrthoFront:
 			ViewOptionType = EViewOptionType::Front;
 			break;
-		case LVT_OrthoYZ:
+		case LVT_OrthoLeft:
 			ViewOptionType = EViewOptionType::Left;
 			break;
 		case LVT_Perspective:
diff --git a/Engine/Source/Editor/UnrealEd/Private/DragTool_ViewportChange.cpp b/Engine/Source/Editor/UnrealEd/Private/DragTool_ViewportChange.cpp
--- a/Engine/Source/Editor/UnrealEd/Private/DragTool_ViewportChange.cpp
+++ b/Engine/Source/Editor/UnrealEd/Private/DragTool_ViewportChange.cpp
@@ -30,7 +30,7 @@ void FDragTool_ViewportChange::StartDrag(FEditorViewportClient* InViewportClient
 	FIntPoint MousePos;
 	InViewportClient->Viewport->GetMousePos(MousePos);
 
-	Start = FVector(InStartScreen.X, InStartScreen.Y, 0);
+	Start = FVector(InStartScreen.X, InStartScreen.Y, 0) / LevelViewportClient->GetDPIScale();
 	End = EndWk = Start;
 }
 
@@ -47,11 +47,11 @@ void FDragTool_ViewportChange::EndDrag()
 		}
 		else if (ViewOptionOffset.X > 0)
 		{
-			ViewOption = LVT_OrthoNegativeYZ; //Right
+			ViewOption = LVT_OrthoRight; // Right
 		}
 		else
 		{
-			ViewOption = LVT_OrthoYZ; //Left
+			ViewOption = LVT_OrthoLeft; // Left
 		}
 	}
 	else
@@ -63,38 +63,38 @@ void FDragTool_ViewportChange::EndDrag()
 		{
 			if (DragAngle >= -15.f && DragAngle <= 15.f)
 			{
-				ViewOption = LVT_OrthoNegativeXY; //Bottom
+				ViewOption = LVT_OrthoBottom; // Bottom
 			}
 			else if (DragAngle > 75.f)
 			{
-				ViewOption = LVT_OrthoNegativeYZ; //Right
+				ViewOption = LVT_OrthoRight; // Right
 			}
 			else if (DragAngle < -75.f)
 			{
-				ViewOption = LVT_OrthoYZ; //Left
+				ViewOption = LVT_OrthoLeft; // Left
 			}
 		}
 		else
 		{
 			if (DragAngle >= -15.f && DragAngle < 15.f)
 			{
-				ViewOption = LVT_OrthoXY; //Top
+				ViewOption = LVT_OrthoTop; // Top
 			}
 			else if (DragAngle >= 15.f && DragAngle < 75.f)
 			{
-				ViewOption = LVT_OrthoXZ; //Front
+				ViewOption = LVT_OrthoFront; // Front
 			}
 			else if (DragAngle >= -75.f && DragAngle < -15.f)
 			{
-				ViewOption = LVT_OrthoNegativeXZ; //Back
+				ViewOption = LVT_OrthoBack; // Back
 			}
 			else if (DragAngle >= 75.f)
 			{
-				ViewOption = LVT_OrthoYZ; //Left
+				ViewOption = LVT_OrthoLeft; // Left
 			}
 			else if (DragAngle <= -75.f)
 			{
-				ViewOption = LVT_OrthoNegativeYZ; //Right
+				ViewOption = LVT_OrthoRight; // Right
 			}
 		}
 	}
@@ -102,7 +102,7 @@ void FDragTool_ViewportChange::EndDrag()
 	double OffsetLength = FMath::RoundToFloat((End - Start).Size());
 	if (OffsetLength >= 125.15f)
 	{
-		LevelViewportClient->SetViewportTypeFromTool(ViewOption);
+		LevelViewportClient->SetViewportType(ViewOption);
 		return;
 	}
 
@@ -119,7 +119,7 @@ void FDragTool_ViewportChange::AddDelta(const FVector& InDelta)
 	FIntPoint MousePos;
 	LevelViewportClient->Viewport->GetMousePos(MousePos);
 
-	EndWk = FVector(MousePos);
+	EndWk = FVector(MousePos) / LevelViewportClient->GetDPIScale();
 	End = EndWk;
 
 	ViewOptionOffset.X = End.X - Start.X;
diff --git a/Engine/Source/Editor/UnrealEd/Private/MouseDeltaTracker.cpp b/Engine/Source/Editor/UnrealEd/Private/MouseDeltaTracker.cpp
--- a/Engine/Source/Editor/UnrealEd/Private/MouseDeltaTracker.cpp
+++ b/Engine/Source/Editor/UnrealEd/Private/MouseDeltaTracker.cpp
@@ -334,6 +334,11 @@ bool FMouseDeltaTracker::EndTracking(FEditorViewportClient* InViewportClient)
 
 void FMouseDeltaTracker::ConditionalBeginUsingDragTool( FEditorViewportClient* InViewportClient )
 {
+	if (UE::Editor::DragTools::UseEditorDragTools())
+	{
+		return;
+	}
+
 	const bool LeftMouseButtonDown = InViewportClient->Viewport->KeyState(EKeys::LeftMouseButton);
 	const bool RightMouseButtonDown = InViewportClient->Viewport->KeyState(EKeys::RightMouseButton);
 	const bool MiddleMouseButtonDown = InViewportClient->Viewport->KeyState(EKeys::MiddleMouseButton);
@@ -371,12 +376,9 @@ void FMouseDeltaTracker::ConditionalBeginUsingDragTool( FEditorViewportClient* I
 				}
 				else
 				{
-					if (!UE::Editor::DragTools::UseEditorDragTools())
+					if (LeftMouseButtonDown && bControlDown && bAltDown)
 					{
-						if (LeftMouseButtonDown && bControlDown && bAltDown)
-						{
-							DragTool = InViewportClient->MakeDragTool(EDragTool::FrustumSelect);
-						}
+						DragTool = InViewportClient->MakeDragTool(EDragTool::FrustumSelect);
 					}
 				}
 			}
diff --git a/Engine/Source/Editor/UnrealEd/Public/LevelEditorViewport.h b/Engine/Source/Editor/UnrealEd/Public/LevelEditorViewport.h
--- a/Engine/Source/Editor/UnrealEd/Public/LevelEditorViewport.h
+++ b/Engine/Source/Editor/UnrealEd/Public/LevelEditorViewport.h
@@ -630,6 +630,7 @@ public:
 	/**
 	 * Calling SetViewportType from Dragtool_ViewportChange
 	 */
+	UE_DEPRECATED(5.6, "Call FLevelEditorViewportClient::SetViewportType instead")
 	UNREALED_API void SetViewportTypeFromTool(ELevelViewportType InViewportType);
 
 	/**
