[ITF Box Marquee Select] Ensure Shift modifier works to add elements to current selection

Register missing Key Modifier Behavior.

diff --git a/Engine/Source/Editor/Experimental/EditorInteractiveToolsFramework/Private/EditorDragTools/DragToolsBehaviourSource.cpp b/Engine/Source/Editor/Experimental/EditorInteractiveToolsFramework/Private/EditorDragTools/DragToolsBehaviourSource.cpp
--- a/Engine/Source/Editor/Experimental/EditorInteractiveToolsFramework/Private/EditorDragTools/DragToolsBehaviourSource.cpp
+++ b/Engine/Source/Editor/Experimental/EditorInteractiveToolsFramework/Private/EditorDragTools/DragToolsBehaviourSource.cpp
@@ -80,6 +80,20 @@ void UDragToolsBehaviorSource::Initialize(const UEditorInteractiveToolsContext*
 		MouseBehavior->SetDefaultPriority(BOX_SELECT_PRIORITY);
 		BehaviorSet->Add(MouseBehavior);
 
+		// Register modifier listeners
+		UKeyAsModifierInputBehavior* KeyModifierInputBehavior = NewObject<UKeyAsModifierInputBehavior>();
+		KeyModifierInputBehavior->Initialize(
+			BoxSelectTool.Get(), UE::Editor::DragTools::ShiftKeyMod, FInputDeviceState::IsShiftKeyDown
+		);
+		KeyModifierInputBehavior->Initialize(
+			BoxSelectTool.Get(), UE::Editor::DragTools::CtrlKeyMod, FInputDeviceState::IsCtrlKeyDown
+		);
+		KeyModifierInputBehavior->Initialize(
+			BoxSelectTool.Get(), UE::Editor::DragTools::AltKeyMod, FInputDeviceState::IsAltKeyDown
+		);
+		KeyModifierInputBehavior->SetDefaultPriority(BOX_SELECT_PRIORITY);
+		BehaviorSet->Add(KeyModifierInputBehavior);
+
 		BoxSelectTool->OnActivateTool().AddUObject(this, &UDragToolsBehaviorSource::OnActivateTool, EDragTool::BoxSelect);
 		BoxSelectTool->OnDeactivateTool().AddUObject(this, &UDragToolsBehaviorSource::OnDeactivateTool, EDragTool::BoxSelect);
 	}
