[ITF Drag Tools] Alt Duplicate

- Add FEditorDuplicateDragSelection class, which is a IModifierToggleBehaviorTarget, currently working with one arbitrary modifier key.
This class waits for new TRS Gizmo drag info (drag begin/end) and combines that info with the current state of the chosen modifier key, in order to duplicate the dragged selection.

- FEditorDuplicateDragSelection added to DragToolBehaviorSource, and set to respond to the Alt mod key.

- Exports UEditorTransformGizmoContextObject::OnGizmoCreatedDelegate

- Swapping BeginUpdate and BeginPivotEdit/BeginTransformEdit order in UTransformGizmo::BeginTransformEditSequence() so that it is possible to perform actions before the transform sequence is actually started. In this case, duplicating the selection before moving and modifying it.

diff --git a/Engine/Source/Editor/Experimental/EditorInteractiveToolsFramework/Private/EditorDragTools/DragToolsBehaviourSource.cpp b/Engine/Source/Editor/Experimental/EditorInteractiveToolsFramework/Private/EditorDragTools/DragToolsBehaviourSource.cpp
--- a/Engine/Source/Editor/Experimental/EditorInteractiveToolsFramework/Private/EditorDragTools/DragToolsBehaviourSource.cpp
+++ b/Engine/Source/Editor/Experimental/EditorInteractiveToolsFramework/Private/EditorDragTools/DragToolsBehaviourSource.cpp
@@ -9,6 +9,7 @@
 #include "EditorDragTools/EditorMeasureTool.h"
 #include "EditorDragTools/EditorViewportChange.h"
 #include "EditorModeManager.h"
+#include "EditorDragTools/EditorDuplicateDragSelection.h"
 #include "Tools/EdModeInteractiveToolsContext.h"
 
 void UDragToolsBehaviorSource::Initialize(const UEditorInteractiveToolsContext* InInteractiveToolsContext)
@@ -100,7 +101,8 @@ void UDragToolsBehaviorSource::Initialize(const UEditorInteractiveToolsContext*
 
 	// Frustum Select
 	{
-		TSharedPtr<FEditorDragToolBehaviorTarget> FrustumSelectTool = MakeShared<FEditorFrustumSelect>(EditorViewportClient);
+		TSharedPtr<FEditorDragToolBehaviorTarget> FrustumSelectTool =
+			MakeShared<FEditorFrustumSelect>(EditorViewportClient);
 
 		DragTools.Add(FrustumSelectTool);
 
@@ -184,6 +186,18 @@ void UDragToolsBehaviorSource::Initialize(const UEditorInteractiveToolsContext*
 			);
 		}
 	}
+
+	// Duplicate Drag Tool
+	{
+		DuplicateDragSelection = MakeShared<FEditorDuplicateDragSelection>(ToolManager);
+		UKeyAsModifierInputBehavior* KeyModifierInputBehavior = NewObject<UKeyAsModifierInputBehavior>();
+		KeyModifierInputBehavior->Initialize(
+			DuplicateDragSelection.Get(), UE::Editor::DragTools::AltKeyMod, FInputDeviceState::IsAltKeyDown
+		);
+		KeyModifierInputBehavior->SetDefaultPriority(HANDLED_BEFORE_GIZMO_PRIORITY);
+
+		BehaviorSet->Add(KeyModifierInputBehavior);
+	}
 }
 
 FEditorDragToolBehaviorTarget* UDragToolsBehaviorSource::GetActiveTool() const
diff --git a/Engine/Source/Editor/Experimental/EditorInteractiveToolsFramework/Private/EditorDragTools/EditorDuplicateDragSelection.cpp b/Engine/Source/Editor/Experimental/EditorInteractiveToolsFramework/Private/EditorDragTools/EditorDuplicateDragSelection.cpp
new file mode 100644
--- /dev/null
+++ b/Engine/Source/Editor/Experimental/EditorInteractiveToolsFramework/Private/EditorDragTools/EditorDuplicateDragSelection.cpp
@@ -0,0 +1,213 @@
+// Copyright Epic Games, Inc. All Rights Reserved.
+
+#include "EditorDragTools/EditorDuplicateDragSelection.h"
+
+#include "ContextObjectStore.h"
+#include "EditorGizmos/EditorTransformGizmo.h"
+#include "EditorGizmos/EditorTransformGizmoUtil.h"
+#include "EditorInteractiveGizmoManager.h"
+#include "EditorModeManager.h"
+#include "Elements/Framework/TypedElementCommonActions.h"
+#include "ILevelEditor.h"
+#include "InteractiveGizmo.h"
+#include "InteractiveGizmoManager.h"
+#include "InteractiveToolManager.h"
+#include "LevelEditorViewport.h"
+#include "EditorDragTools/EditorDragToolBehaviorTarget.h"
+#include "Tools/AssetEditorContextInterface.h"
+
+FEditorDuplicateDragSelection::FEditorDuplicateDragSelection(UInteractiveToolManager* InToolManager)
+	: ToolManager(InToolManager)
+{
+	// New TRS Gizmos are already enabled - let's retrieve the Transform Gizmo, from which we register to drag Begin and End delegates
+	if (UEditorInteractiveGizmoManager::UsesNewTRSGizmos())
+	{
+		const TArray<UInteractiveGizmo*> Gizmos =
+		ToolManager->GetPairedGizmoManager()->FindAllGizmosOfType("EditorTransformGizmoBuilder");
+
+		if (!Gizmos.IsEmpty())
+		{
+			TransformGizmo = Cast<UEditorTransformGizmo>(Gizmos[0]);
+		}
+
+		Initialize();
+	}
+
+	// If New TRS Gizmos are not enabled, we need to be able to know as soon as they are enabled
+	// This allows to retrieve the Transform Gizmo, from which we register to drag Begin and End delegates
+	if (UContextObjectStore* ContextStore = ToolManager->GetContextObjectStore())
+	{
+		if (UEditorTransformGizmoContextObject* ContextObject =
+				ContextStore->FindContext<UEditorTransformGizmoContextObject>())
+		{
+			ContextObject->OnGizmoCreatedDelegate().AddRaw(this, &FEditorDuplicateDragSelection::OnGizmoCreatedDelegate);
+		}
+	}
+
+	// In case new TRS Gizmos are disabled, we want to know that, so we can stop listening to Drag delegates
+	UEditorInteractiveGizmoManager::OnUsesNewTRSGizmosChangedDelegate().AddRaw(
+		this, &FEditorDuplicateDragSelection::OnUsesNewTRSGizmosChanged
+	);
+}
+
+void FEditorDuplicateDragSelection::Initialize()
+{
+	if (TransformGizmo)
+	{
+		// Register Proxy Delegates
+		if (TObjectPtr<UTransformProxy> TransformProxy = TransformGizmo->ActiveTarget)
+		{
+			TransformProxy->OnBeginTransformEdit.AddRaw(
+				this, &FEditorDuplicateDragSelection::OnGizmoMovementBegin
+			);
+
+			TransformProxy->OnEndTransformEdit.AddRaw(
+				this, &FEditorDuplicateDragSelection::OnGizmoMovementEnd
+			);
+		}
+	}
+}
+
+void FEditorDuplicateDragSelection::Reset()
+{
+	if (TransformGizmo)
+	{
+		// Unregister Proxy delegates
+		if (TObjectPtr<UTransformProxy> TransformProxy = TransformGizmo->ActiveTarget)
+		{
+			TransformProxy->OnBeginTransformEdit.RemoveAll(this);
+			TransformProxy->OnEndTransformEdit.RemoveAll(this);
+		}
+
+		TransformGizmo = nullptr;
+	}
+}
+
+void FEditorDuplicateDragSelection::OnGizmoCreatedDelegate(UTransformGizmo* InTransformGizmo)
+{
+	if (!TransformGizmo)
+	{
+		TransformGizmo = Cast<UEditorTransformGizmo>(InTransformGizmo);
+
+		Initialize();
+	}
+}
+
+void FEditorDuplicateDragSelection::OnUsesNewTRSGizmosChanged(bool bInUseNewTRSGizmos)
+{
+	if (!bInUseNewTRSGizmos)
+	{
+		Reset();
+	}
+}
+
+FEditorDuplicateDragSelection::~FEditorDuplicateDragSelection()
+{
+	UEditorInteractiveGizmoManager::OnUsesNewTRSGizmosChangedDelegate().RemoveAll(this);
+	Reset();
+}
+
+void FEditorDuplicateDragSelection::AddReferencedObjects(FReferenceCollector& InCollector)
+{
+	if (ToolManager)
+	{
+		InCollector.AddReferencedObject(ToolManager);
+	}
+
+	if (TransformGizmo)
+	{
+		InCollector.AddReferencedObject(TransformGizmo);
+	}
+}
+
+void FEditorDuplicateDragSelection::OnUpdateModifierState(int InModifierID, bool bInIsOn)
+{
+	// Currently only supports one modifier
+	constexpr int KeyModifierID = UE::Editor::DragTools::AltKeyMod;
+
+	if (InModifierID == KeyModifierID)
+	{
+		bModifierKeyIsPressed = bInIsOn;
+
+		// If already dragging, let's duplicate without selecting new objects
+		// This will "drop" a copy of selection matching current drag transform
+		if (bGizmoIsDragged && bModifierKeyIsPressed)
+		{
+			constexpr bool bSelectNewElements = false;
+			DuplicateSelection(bSelectNewElements);
+		}
+	}
+}
+
+void FEditorDuplicateDragSelection::OnForceEndCapture()
+{
+	bGizmoIsDragged = false;
+	bModifierKeyIsPressed = false;
+}
+
+void FEditorDuplicateDragSelection::OnGizmoMovementBegin(UTransformProxy* InTransformProxy)
+{
+	if (!bGizmoIsDragged)
+	{
+		bGizmoIsDragged = true;
+
+		if (bModifierKeyIsPressed)
+		{
+			// Duplicate selection and select the newly created objects
+			constexpr bool bSelectNewElements = true;
+			DuplicateSelection(bSelectNewElements);
+		}
+	}
+}
+
+void FEditorDuplicateDragSelection::OnGizmoMovementEnd(UTransformProxy* InTransformProxy)
+{
+	bGizmoIsDragged = false;
+}
+
+void FEditorDuplicateDragSelection::DuplicateSelection(bool bInSelectNewElements) const
+{
+	if (GCurrentLevelEditingViewportClient)
+	{
+		if (TSharedPtr<ILevelEditor> LevelEditor = GCurrentLevelEditingViewportClient->ParentLevelEditor.Pin())
+		{
+			if (UTypedElementCommonActions* CommonActions = LevelEditor->GetCommonActions())
+			{
+				FToolBuilderState State;
+				if (IAssetEditorContextInterface* AssetEditorContext =
+						ToolManager->GetContextObjectStore()->FindContext<IAssetEditorContextInterface>())
+				{
+					if (UTypedElementSelectionSet* MutableSelectionSet = AssetEditorContext->GetMutableSelectionSet())
+					{
+						const TArray<FTypedElementHandle> DuplicatedElements = CommonActions->DuplicateSelectedElements(
+							MutableSelectionSet, LevelEditor->GetWorld(), FVector::ZeroVector
+						);
+						if (DuplicatedElements.Num() > 0)
+						{
+							// (assuming chosen modifier key is Alt)
+							// Alt Duplicate has 2 modes:
+							// 1 - with Alt pressed, dragging will create and select a copy, then drag it
+							// 2 - while dragging, hitting Alt will crate a copy, NOT selecting it
+
+							// Select newly created elements: Alt + Drag only
+							if (bInSelectNewElements)
+							{
+								// Select the newly created elements
+								MutableSelectionSet->SetSelection(DuplicatedElements, FTypedElementSelectionOptions());
+								MutableSelectionSet->NotifyPendingChanges();
+							}
+
+							// Notify the global mode tools, the selection set should be identical to the new actors at this point
+							TArray<AActor*> SelectedActors = MutableSelectionSet->GetSelectedObjects<AActor>();
+							constexpr bool bDidOffsetDuplicate = false;
+							GLevelEditorModeTools().ActorsDuplicatedNotify(SelectedActors, SelectedActors, bDidOffsetDuplicate);
+						}
+
+						// Invalidate all viewports, so the new gizmo is rendered in each one
+						GEditor->RedrawLevelEditingViewports();
+					}
+				}
+			}
+		}
+	}
+}
diff --git a/Engine/Source/Editor/Experimental/EditorInteractiveToolsFramework/Public/EditorDragTools/DragToolsBehaviourSource.h b/Engine/Source/Editor/Experimental/EditorInteractiveToolsFramework/Public/EditorDragTools/DragToolsBehaviourSource.h
--- a/Engine/Source/Editor/Experimental/EditorInteractiveToolsFramework/Public/EditorDragTools/DragToolsBehaviourSource.h
+++ b/Engine/Source/Editor/Experimental/EditorInteractiveToolsFramework/Public/EditorDragTools/DragToolsBehaviourSource.h
@@ -38,6 +38,7 @@ public:
 	 * Renders the active tool on the specified View/Canvas
 	 */
 	EDITORINTERACTIVETOOLSFRAMEWORK_API void RenderTools(const FSceneView* InSceneView, FCanvas* InCanvas) const;
+
 private:
 	static int32 GetToolIndex(EDragTool::Type InDragToolType);
 
@@ -60,6 +61,11 @@ private:
 	 */
 	TArray<TSharedPtr<FEditorDragToolBehaviorTarget>> DragTools;
 
+	/**
+	 * Duplicate dragged selection tool
+	 */
+	TSharedPtr<FEditorDuplicateDragSelection> DuplicateDragSelection;
+
 	/**
 	 * Which type of Drag tool is currently active
 	 */
diff --git a/Engine/Source/Editor/Experimental/EditorInteractiveToolsFramework/Public/EditorDragTools/EditorDuplicateDragSelection.h b/Engine/Source/Editor/Experimental/EditorInteractiveToolsFramework/Public/EditorDragTools/EditorDuplicateDragSelection.h
new file mode 100644
--- /dev/null
+++ b/Engine/Source/Editor/Experimental/EditorInteractiveToolsFramework/Public/EditorDragTools/EditorDuplicateDragSelection.h
@@ -0,0 +1,61 @@
+// Copyright Epic Games, Inc. All Rights Reserved.
+
+#pragma once
+
+#include "BaseGizmos/TransformProxy.h"
+#include "EditorGizmos/TransformGizmo.h"
+
+class FLevelEditorViewportClient;
+class UEditorTransformGizmo;
+class UInteractiveToolManager;
+class UTransformProxy;
+
+/**
+ * Listens for key press of a specified modifier key, and if Level Viewport Gizmo is dragging, it duplicates current selection.
+ * Currently used with Alt (see UDragToolsBehaviorSource)
+ */
+class FEditorDuplicateDragSelection
+	: public IModifierToggleBehaviorTarget
+	, public FGCObject
+{
+public:
+	FEditorDuplicateDragSelection(UInteractiveToolManager* InToolManager);
+
+	virtual ~FEditorDuplicateDragSelection() override;
+
+protected:
+	//~Begin IModifierToggleBehaviorTarget
+	virtual void OnUpdateModifierState(int InModifierID, bool bInIsOn) override;
+	virtual void OnForceEndCapture() override;
+	//~End IModifierToggleBehaviorTarget
+
+	//~Begin FGCObject
+	virtual void AddReferencedObjects(FReferenceCollector& InCollector) override;
+	virtual FString GetReferencerName() const override
+	{
+		return TEXT("FEditorDuplicateDragSelection");
+	}
+	//~End FGCObject
+
+	void OnGizmoMovementBegin(UTransformProxy* InTransformProxy);
+	void OnGizmoMovementEnd(UTransformProxy* InTransformProxy);
+
+	void Initialize();
+	void Reset();
+
+	void OnGizmoCreatedDelegate(UTransformGizmo* InTransformGizmo);
+	void OnUsesNewTRSGizmosChanged(bool bInUseNewTRSGizmos);
+
+	/**
+	 * Duplicate the current selection
+	 * @param bInSelectNewElements if newly created elements should be selected right after duplication
+	 */
+	void DuplicateSelection(bool bInSelectNewElements) const;
+
+private:
+	TObjectPtr<UInteractiveToolManager> ToolManager;
+	TObjectPtr<UEditorTransformGizmo> TransformGizmo;
+
+	bool bGizmoIsDragged = false;
+	bool bModifierKeyIsPressed = false;
+};
diff --git a/Engine/Source/Editor/Experimental/EditorInteractiveToolsFramework/Public/EditorGizmos/EditorTransformGizmoUtil.h b/Engine/Source/Editor/Experimental/EditorInteractiveToolsFramework/Public/EditorGizmos/EditorTransformGizmoUtil.h
--- a/Engine/Source/Editor/Experimental/EditorInteractiveToolsFramework/Public/EditorGizmos/EditorTransformGizmoUtil.h
+++ b/Engine/Source/Editor/Experimental/EditorInteractiveToolsFramework/Public/EditorGizmos/EditorTransformGizmoUtil.h
@@ -1,4 +1,4 @@
-﻿// Copyright Epic Games, Inc. All Rights Reserved.
+// Copyright Epic Games, Inc. All Rights Reserved.
 
 #pragma once
 
@@ -90,7 +90,7 @@ public:
 	EDITORINTERACTIVETOOLSFRAMEWORK_API FEditorModeTools* GetModeTools() const;
 
 	DECLARE_MULTICAST_DELEGATE_OneParam(FOnGizmoCreated, UTransformGizmo*);
-	FOnGizmoCreated& OnGizmoCreatedDelegate();
+	EDITORINTERACTIVETOOLSFRAMEWORK_API FOnGizmoCreated& OnGizmoCreatedDelegate();
 
 	UPROPERTY(Transient)
 	FRotationContext RotationContext;
