[EditorGizmos] TRS gizmo restores after switching editor modes

diff --git a/Engine/Source/Editor/Experimental/EditorInteractiveToolsFramework/Private/EditorInteractiveGizmoManager.cpp b/Engine/Source/Editor/Experimental/EditorInteractiveToolsFramework/Private/EditorInteractiveGizmoManager.cpp
--- a/Engine/Source/Editor/Experimental/EditorInteractiveToolsFramework/Private/EditorInteractiveGizmoManager.cpp
+++ b/Engine/Source/Editor/Experimental/EditorInteractiveToolsFramework/Private/EditorInteractiveGizmoManager.cpp
@@ -93,6 +93,8 @@ void UEditorInteractiveGizmoManager::InitializeWithEditorModeManager(IToolsConte
 	UModeManagerInteractiveToolsContext* InteractiveToolsContext = EditorModeManager->GetInteractiveToolsContext();
 	check(InteractiveToolsContext);
 
+	EditorModeManager->OnEditorModeIDChanged().AddUObject(this, &UEditorInteractiveGizmoManager::OnEditorModeChanged);
+
 	if (InteractiveToolsContext)
 	{
 		UE::Editor::Gizmos::RegisterSceneSnappingManager(InteractiveToolsContext);
@@ -103,10 +105,7 @@ void UEditorInteractiveGizmoManager::Shutdown()
 {
 	if (EditorModeManager)
 	{
-		if (const UModeManagerInteractiveToolsContext* InteractiveToolsContext = EditorModeManager->GetInteractiveToolsContext())
-		{
-			UE::Editor::Gizmos::DeregisterSceneSnappingManager(InteractiveToolsContext);
-		}
+		EditorModeManager->OnEditorModeIDChanged().RemoveAll(this);
 	}
 
 	Registry->Shutdown();
@@ -147,6 +146,22 @@ UTransformGizmo* UEditorInteractiveGizmoManager::FindDefaultTransformGizmo() con
 	return Cast<UTransformGizmo>( FindGizmoByInstanceIdentifier(TransformInstanceIdentifier()) );
 }
 
+void UEditorInteractiveGizmoManager::OnEditorModeChanged(const FName& InModeID, bool bInIsEnteringMode)
+{
+	// When a mode is switched, shutdown is called after initializing the other mode.
+	// When the system is shared between modes (ie. snapping), this can de-register the snapping AFTER being registered in init
+	// This callback allows us to de-register things before the other mode registers
+
+	// Only execute on exit (not enter)
+	if (!bInIsEnteringMode && EditorModeManager)
+	{
+		if (const UModeManagerInteractiveToolsContext* InteractiveToolsContext = EditorModeManager->GetInteractiveToolsContext())
+		{
+			UE::Editor::Gizmos::DeregisterSceneSnappingManager(InteractiveToolsContext);
+		}
+	}
+}
+
 bool UEditorInteractiveGizmoManager::DestroyEditorGizmo(UInteractiveGizmo* Gizmo)
 {
 	auto Pred = [Gizmo](const FActiveEditorGizmo& ActiveEditorGizmo) {return ActiveEditorGizmo.Gizmo == Gizmo; };
diff --git a/Engine/Source/Editor/Experimental/EditorInteractiveToolsFramework/Public/EditorInteractiveGizmoManager.h b/Engine/Source/Editor/Experimental/EditorInteractiveToolsFramework/Public/EditorInteractiveGizmoManager.h
--- a/Engine/Source/Editor/Experimental/EditorInteractiveToolsFramework/Public/EditorInteractiveGizmoManager.h
+++ b/Engine/Source/Editor/Experimental/EditorInteractiveToolsFramework/Public/EditorInteractiveGizmoManager.h
@@ -215,4 +215,6 @@ private:
 
 	/** Returns the existing default Gizmo instance if any. */
 	UTransformGizmo* FindDefaultTransformGizmo() const;
+
+	void OnEditorModeChanged(const FName& InModeID, bool bInIsEnteringMode);
 };
