Enable New TRS Gizmo to use LeftUpForward coordinate system

diff --git a/Engine/Source/Editor/Experimental/EditorInteractiveToolsFramework/Private/EditorGizmos/TransformGizmo.cpp b/Engine/Source/Editor/Experimental/EditorInteractiveToolsFramework/Private/EditorGizmos/TransformGizmo.cpp
--- a/Engine/Source/Editor/Experimental/EditorInteractiveToolsFramework/Private/EditorGizmos/TransformGizmo.cpp
+++ b/Engine/Source/Editor/Experimental/EditorInteractiveToolsFramework/Private/EditorGizmos/TransformGizmo.cpp
@@ -1209,12 +1209,21 @@ void UTransformGizmo::OnParametersChanged(const FGizmosParameters& InParameters)
 UGizmoElementArrow* UTransformGizmo::MakeTranslateAxis(ETransformGizmoPartIdentifier InPartId, const FVector& InAxisDir, const FVector& InSideDir, UMaterialInterface* InMaterial)
 {
 	const float SizeCoeff = GetSizeCoefficient();
+
+	FVector AxisDir = InAxisDir;
+	if (AxisDisplayInfo::GetAxisDisplayCoordinateSystem() == EAxisList::LeftUpForward)
+	{
+		if (InPartId == ETransformGizmoPartIdentifier::TranslateYAxis)
+		{
+			AxisDir = -AxisDir;
+		}
+	}
 	
 	UGizmoElementArrow* ArrowElement = NewObject<UGizmoElementArrow>();
 	ArrowElement->SetPartIdentifier(static_cast<uint32>(InPartId));
 	ArrowElement->SetHeadType(EGizmoElementArrowHeadType::Cone);
-	ArrowElement->SetBase(InAxisDir * AxisLengthOffset);
-	ArrowElement->SetDirection(InAxisDir);
+	ArrowElement->SetBase(AxisDir * AxisLengthOffset);
+	ArrowElement->SetDirection(AxisDir);
 	ArrowElement->SetSideDirection(InSideDir);
 	ArrowElement->SetBodyLength(TranslateAxisLength * SizeCoeff);
 	ArrowElement->SetBodyRadius(AxisRadius);
@@ -1231,11 +1240,20 @@ UGizmoElementArrow* UTransformGizmo::MakeScaleAxis(ETransformGizmoPartIdentifier
 {
 	const float SizeCoeff = GetSizeCoefficient();
 
+	FVector AxisDir = InAxisDir;
+	if (AxisDisplayInfo::GetAxisDisplayCoordinateSystem() == EAxisList::LeftUpForward)
+	{
+		if (InPartId == ETransformGizmoPartIdentifier::ScaleYAxis)
+		{
+			AxisDir = -AxisDir;
+		}
+	}
+
 	UGizmoElementArrow* ArrowElement = NewObject<UGizmoElementArrow>();
 	ArrowElement->SetPartIdentifier(static_cast<uint32>(InPartId));
 	ArrowElement->SetHeadType(EGizmoElementArrowHeadType::Cube);
-	ArrowElement->SetBase(InAxisDir * AxisLengthOffset);
-	ArrowElement->SetDirection(InAxisDir);
+	ArrowElement->SetBase(AxisDir * AxisLengthOffset);
+	ArrowElement->SetDirection(AxisDir);
 	ArrowElement->SetSideDirection(InSideDir);
 	ArrowElement->SetBodyLength(ScaleAxisLength * SizeCoeff);
 	ArrowElement->SetBodyRadius(AxisRadius);
@@ -1265,8 +1283,22 @@ UGizmoElementRectangle* UTransformGizmo::MakePlanarHandle(ETransformGizmoPartIde
 	UMaterialInterface* InMaterial, const FLinearColor& InVertexColor)
 {
 	const float SizeCoeff = GetSizeCoefficient();
+
+	FVector UpDirection = InUpDirection;
+	FVector SideDirection = InSideDirection;
+	if (AxisDisplayInfo::GetAxisDisplayCoordinateSystem() == EAxisList::LeftUpForward)
+	{
+		if (InPartId == ETransformGizmoPartIdentifier::TranslateXYPlanar || InPartId == ETransformGizmoPartIdentifier::ScaleXYPlanar)
+		{
+			SideDirection = -SideDirection;
+		}
+		if (InPartId == ETransformGizmoPartIdentifier::TranslateYZPlanar || InPartId == ETransformGizmoPartIdentifier::ScaleYZPlanar)
+		{
+			UpDirection = -UpDirection;
+		}
+	}
 	
-	FVector PlanarHandleCenter = (InUpDirection + InSideDirection) * PlanarHandleOffset * SizeCoeff;
+	FVector PlanarHandleCenter = (UpDirection + SideDirection) * PlanarHandleOffset * SizeCoeff;
 
 	FLinearColor LineColor = InVertexColor;
 	FLinearColor VertexColor = LineColor;
@@ -1274,8 +1306,8 @@ UGizmoElementRectangle* UTransformGizmo::MakePlanarHandle(ETransformGizmoPartIde
 
 	UGizmoElementRectangle* RectangleElement = NewObject<UGizmoElementRectangle>();
 	RectangleElement->SetPartIdentifier(static_cast<uint32>(InPartId));
-	RectangleElement->SetUpDirection(InUpDirection);
-	RectangleElement->SetSideDirection(InSideDirection);
+	RectangleElement->SetUpDirection(UpDirection);
+	RectangleElement->SetSideDirection(SideDirection);
 	RectangleElement->SetCenter(PlanarHandleCenter);
 	RectangleElement->SetHeight(PlanarHandleSize * SizeCoeff);
 	RectangleElement->SetWidth(PlanarHandleSize * SizeCoeff);
