[ITF Alt Duplicate] Address duplication happening even when drag distance is zero.

diff --git a/Engine/Source/Editor/Experimental/EditorInteractiveToolsFramework/Private/EditorDragTools/EditorDuplicateDragSelection.cpp b/Engine/Source/Editor/Experimental/EditorInteractiveToolsFramework/Private/EditorDragTools/EditorDuplicateDragSelection.cpp
--- a/Engine/Source/Editor/Experimental/EditorInteractiveToolsFramework/Private/EditorDragTools/EditorDuplicateDragSelection.cpp
+++ b/Engine/Source/Editor/Experimental/EditorInteractiveToolsFramework/Private/EditorDragTools/EditorDuplicateDragSelection.cpp
@@ -61,6 +61,8 @@ void FEditorDuplicateDragSelection::Initialize()
 				this, &FEditorDuplicateDragSelection::OnGizmoMovementBegin
 			);
 
+			TransformProxy->OnTransformChanged.AddRaw(this, &FEditorDuplicateDragSelection::OnGizmoTransformChanged);
+
 			TransformProxy->OnEndTransformEdit.AddRaw(
 				this, &FEditorDuplicateDragSelection::OnGizmoMovementEnd
 			);
@@ -77,6 +79,7 @@ void FEditorDuplicateDragSelection::Reset()
 		{
 			TransformProxy->OnBeginTransformEdit.RemoveAll(this);
 			TransformProxy->OnEndTransformEdit.RemoveAll(this);
+			TransformProxy->OnTransformChanged.RemoveAll(this);
 		}
 
 		TransformGizmo = nullptr;
@@ -150,19 +153,29 @@ void FEditorDuplicateDragSelection::OnGizmoMovementBegin(UTransformProxy* InTran
 	if (!bGizmoIsDragged)
 	{
 		bGizmoIsDragged = true;
-
-		if (bModifierKeyIsPressed)
-		{
-			// Duplicate selection and select the newly created objects
-			constexpr bool bSelectNewElements = true;
-			DuplicateSelection(bSelectNewElements);
-		}
 	}
 }
 
 void FEditorDuplicateDragSelection::OnGizmoMovementEnd(UTransformProxy* InTransformProxy)
 {
 	bGizmoIsDragged = false;
+	bGizmoTransformHasChanged = false;
+}
+
+void FEditorDuplicateDragSelection::OnGizmoTransformChanged(UTransformProxy* InTransformProxy, UE::Math::TTransform<double> InTransform)
+{
+	if (bGizmoTransformHasChanged || InTransform.GetTranslation().IsNearlyZero())
+	{
+		return;
+	}
+
+	bGizmoTransformHasChanged = true;
+	if (bModifierKeyIsPressed)
+	{
+		// Duplicate selection and select the newly created objects
+		constexpr bool bSelectNewElements = true;
+		DuplicateSelection(bSelectNewElements);
+	}
 }
 
 void FEditorDuplicateDragSelection::DuplicateSelection(bool bInSelectNewElements) const
diff --git a/Engine/Source/Editor/Experimental/EditorInteractiveToolsFramework/Private/EditorGizmos/EditorTransformProxy.cpp b/Engine/Source/Editor/Experimental/EditorInteractiveToolsFramework/Private/EditorGizmos/EditorTransformProxy.cpp
--- a/Engine/Source/Editor/Experimental/EditorInteractiveToolsFramework/Private/EditorGizmos/EditorTransformProxy.cpp
+++ b/Engine/Source/Editor/Experimental/EditorInteractiveToolsFramework/Private/EditorGizmos/EditorTransformProxy.cpp
@@ -37,6 +37,7 @@ void UEditorTransformProxy::InputTranslateDelta(const FVector& InDeltaTranslate,
 		// Set legacy widget axis temporarily because InputWidgetDelta branches/overrides may expect it
 		ViewportClient->SetCurrentWidgetAxis(InAxisList);
 		ViewportClient->InputWidgetDelta(ViewportClient->Viewport, InAxisList, Translate, Rot, Scale);
+		OnTransformChanged.Broadcast(this, FTransform(Rot, Translate, Scale));
 		ViewportClient->SetCurrentWidgetAxis(EAxisList::None);
 	}
 }
@@ -52,6 +53,7 @@ void UEditorTransformProxy::InputScaleDelta(const FVector& InDeltaScale, EAxisLi
 		// Set legacy widget axis temporarily because InputWidgetDelta validates the axis in some crashes and crashes if it is not set
 		ViewportClient->SetCurrentWidgetAxis(InAxisList);
 		ViewportClient->InputWidgetDelta(ViewportClient->Viewport, InAxisList, Translate, Rot, Scale);
+		OnTransformChanged.Broadcast(this, FTransform(Rot, Translate, Scale));
 		ViewportClient->SetCurrentWidgetAxis(EAxisList::None);
 	}
 }
@@ -67,6 +69,7 @@ void UEditorTransformProxy::InputRotateDelta(const FRotator& InDeltaRotate, EAxi
 		// Set legacy widget axis temporarily because InputWidgetDelta branches/overrides may expect it
 		ViewportClient->SetCurrentWidgetAxis(InAxisList);
 		ViewportClient->InputWidgetDelta(ViewportClient->Viewport, InAxisList, Translate, Rot, Scale);
+		OnTransformChanged.Broadcast(this, FTransform(Rot, Translate, Scale));
 		ViewportClient->SetCurrentWidgetAxis(EAxisList::None);
 	}
 }
diff --git a/Engine/Source/Editor/Experimental/EditorInteractiveToolsFramework/Public/EditorDragTools/EditorDuplicateDragSelection.h b/Engine/Source/Editor/Experimental/EditorInteractiveToolsFramework/Public/EditorDragTools/EditorDuplicateDragSelection.h
--- a/Engine/Source/Editor/Experimental/EditorInteractiveToolsFramework/Public/EditorDragTools/EditorDuplicateDragSelection.h
+++ b/Engine/Source/Editor/Experimental/EditorInteractiveToolsFramework/Public/EditorDragTools/EditorDuplicateDragSelection.h
@@ -39,6 +39,7 @@ protected:
 
 	void OnGizmoMovementBegin(UTransformProxy* InTransformProxy);
 	void OnGizmoMovementEnd(UTransformProxy* InTransformProxy);
+	void OnGizmoTransformChanged(UTransformProxy* InTransformProxy, UE::Math::TTransform<double> InTransform);
 
 	void Initialize();
 	void Reset();
@@ -58,4 +59,5 @@ private:
 
 	bool bGizmoIsDragged = false;
 	bool bModifierKeyIsPressed = false;
+	bool bGizmoTransformHasChanged = false;
 };
