[ITF Drag Tools] Shift-Drag handle behavior

- Add FEditorMoveCameraWithObject class, which is an IModifierToggleBehaviorTarget

- FEditorMoveCameraWithObject added to DragToolBehaviorSource, and set to respond to the Shift mod key.

- Bypassing legacy Shift Drag in LevelEditorViewport when new gizmos are enabled

Note: this new tool replicates the current Shift Drag behavior happening with New Gizmos, which presents an accelerated camera movement, which will be fixed separately.

diff --git a/Engine/Source/Editor/Experimental/EditorInteractiveToolsFramework/Private/EditorDragTools/DragToolsBehaviourSource.cpp b/Engine/Source/Editor/Experimental/EditorInteractiveToolsFramework/Private/EditorDragTools/DragToolsBehaviourSource.cpp
--- a/Engine/Source/Editor/Experimental/EditorInteractiveToolsFramework/Private/EditorDragTools/DragToolsBehaviourSource.cpp
+++ b/Engine/Source/Editor/Experimental/EditorInteractiveToolsFramework/Private/EditorDragTools/DragToolsBehaviourSource.cpp
@@ -10,6 +10,7 @@
 #include "EditorDragTools/EditorViewportChange.h"
 #include "EditorModeManager.h"
 #include "EditorDragTools/EditorDuplicateDragSelection.h"
+#include "EditorDragTools/EditorMoveCameraWithObject.h"
 #include "Tools/EdModeInteractiveToolsContext.h"
 
 void UDragToolsBehaviorSource::Initialize(const UEditorInteractiveToolsContext* InInteractiveToolsContext)
@@ -198,6 +199,17 @@ void UDragToolsBehaviorSource::Initialize(const UEditorInteractiveToolsContext*
 
 		BehaviorSet->Add(KeyModifierInputBehavior);
 	}
+
+	// Move Camera with selection
+	{
+		MoveCameraWithObject = MakeShared<FEditorMoveCameraWithObject>(ToolManager);
+
+		UKeyAsModifierInputBehavior* KeyModifierInputBehavior = NewObject<UKeyAsModifierInputBehavior>();
+		KeyModifierInputBehavior->Initialize(MoveCameraWithObject.Get(), UE::Editor::DragTools::ShiftKeyMod, FInputDeviceState::IsShiftKeyDown);
+		KeyModifierInputBehavior->SetDefaultPriority(HANDLED_BEFORE_GIZMO_PRIORITY);
+
+		BehaviorSet->Add(KeyModifierInputBehavior);
+	}
 }
 
 FEditorDragToolBehaviorTarget* UDragToolsBehaviorSource::GetActiveTool() const
diff --git a/Engine/Source/Editor/Experimental/EditorInteractiveToolsFramework/Private/EditorDragTools/EditorMoveCameraWithObject.cpp b/Engine/Source/Editor/Experimental/EditorInteractiveToolsFramework/Private/EditorDragTools/EditorMoveCameraWithObject.cpp
new file mode 100644
--- /dev/null
+++ b/Engine/Source/Editor/Experimental/EditorInteractiveToolsFramework/Private/EditorDragTools/EditorMoveCameraWithObject.cpp
@@ -0,0 +1,163 @@
+﻿// Copyright Epic Games, Inc. All Rights Reserved.
+
+#include "EditorDragTools/EditorMoveCameraWithObject.h"
+
+#include "ContextObjectStore.h"
+#include "EditorDragTools/EditorDragToolBehaviorTarget.h"
+#include "EditorGizmos/EditorTransformGizmo.h"
+#include "EditorGizmos/EditorTransformGizmoUtil.h"
+#include "EditorInteractiveGizmoManager.h"
+#include "ILevelEditor.h"
+#include "InteractiveGizmo.h"
+#include "InteractiveGizmoManager.h"
+#include "InteractiveToolManager.h"
+#include "LevelEditorViewport.h"
+#include "SLevelViewport.h"
+#include "Tools/EdModeInteractiveToolsContext.h"
+
+FEditorMoveCameraWithObject::FEditorMoveCameraWithObject(UInteractiveToolManager* InToolManager)
+	: ToolManager(InToolManager)
+{
+	// New TRS Gizmos are already enabled - let's retrieve the Transform Gizmo, from which we register to drag Begin and End delegates
+	if (UEditorInteractiveGizmoManager::UsesNewTRSGizmos())
+	{
+		const TArray<UInteractiveGizmo*> Gizmos =
+			ToolManager->GetPairedGizmoManager()->FindAllGizmosOfType("EditorTransformGizmoBuilder");
+
+		if (!Gizmos.IsEmpty())
+		{
+			TransformGizmo = Cast<UEditorTransformGizmo>(Gizmos[0]);
+		}
+
+		Initialize();
+	}
+
+	// If New TRS Gizmos are not enabled, we need to be able to know as soon as they are enabled
+	// This allows to retrieve the Transform Gizmo, from which we register to drag Begin and End delegates
+	if (UContextObjectStore* ContextStore = ToolManager->GetContextObjectStore())
+	{
+		if (UEditorTransformGizmoContextObject* ContextObject =
+				ContextStore->FindContext<UEditorTransformGizmoContextObject>())
+		{
+			ContextObject->OnGizmoCreatedDelegate().AddRaw(this, &FEditorMoveCameraWithObject::OnGizmoCreatedDelegate);
+		}
+	}
+
+	// In case new TRS Gizmos are disabled, we want to know that, so we can stop listening to Drag delegates
+	UEditorInteractiveGizmoManager::OnUsesNewTRSGizmosChangedDelegate().AddRaw(
+		this, &FEditorMoveCameraWithObject::OnUsesNewTRSGizmosChanged
+	);
+}
+
+void FEditorMoveCameraWithObject::Initialize()
+{
+	if (TransformGizmo)
+	{
+		// Register Proxy Delegates
+		if (TObjectPtr<UTransformProxy> TransformProxy = TransformGizmo->ActiveTarget)
+		{
+			OnBeginPivotEditDelegate =
+				TransformProxy->OnBeginTransformEdit.AddRaw(this, &FEditorMoveCameraWithObject::OnGizmoMovementBegin);
+
+			OnEndPivotEditDelegate =
+				TransformProxy->OnEndTransformEdit.AddRaw(this, &FEditorMoveCameraWithObject::OnGizmoMovementEnd);
+
+			OnTransformChangedDelegate =
+				TransformProxy->OnTransformChanged.AddRaw(this, &FEditorMoveCameraWithObject::OnGizmoTransformChanged);
+		}
+	}
+}
+
+void FEditorMoveCameraWithObject::Reset()
+{
+	if (TransformGizmo)
+	{
+		// Unregister Proxy delegates
+		if (TObjectPtr<UTransformProxy> TransformProxy = TransformGizmo->ActiveTarget)
+		{
+			TransformProxy->OnBeginTransformEdit.Remove(OnBeginPivotEditDelegate);
+			TransformProxy->OnEndTransformEdit.Remove(OnEndPivotEditDelegate);
+			TransformProxy->OnTransformChanged.Remove(OnTransformChangedDelegate);
+		}
+
+		TransformGizmo = nullptr;
+	}
+}
+
+void FEditorMoveCameraWithObject::OnGizmoCreatedDelegate(UTransformGizmo* InTransformGizmo)
+{
+	if (!TransformGizmo)
+	{
+		TransformGizmo = Cast<UEditorTransformGizmo>(InTransformGizmo);
+
+		Initialize();
+	}
+}
+
+void FEditorMoveCameraWithObject::OnUsesNewTRSGizmosChanged(bool bInUseNewTRSGizmos)
+{
+	if (!bInUseNewTRSGizmos)
+	{
+		Reset();
+	}
+}
+
+FEditorMoveCameraWithObject::~FEditorMoveCameraWithObject()
+{
+	UEditorInteractiveGizmoManager::OnUsesNewTRSGizmosChangedDelegate().RemoveAll(this);
+	Reset();
+}
+
+void FEditorMoveCameraWithObject::AddReferencedObjects(FReferenceCollector& InCollector)
+{
+	if (ToolManager)
+	{
+		InCollector.AddReferencedObject(ToolManager);
+	}
+
+	if (TransformGizmo)
+	{
+		InCollector.AddReferencedObject(TransformGizmo);
+	}
+}
+
+void FEditorMoveCameraWithObject::OnUpdateModifierState(int InModifierID, bool bInIsOn)
+{
+	// Currently only supports one modifier
+	constexpr int KeyModifierID = UE::Editor::DragTools::ShiftKeyMod;
+
+	if (InModifierID == KeyModifierID)
+	{
+		bModifierKeyIsPressed = bInIsOn;
+	}
+}
+
+void FEditorMoveCameraWithObject::OnGizmoMovementBegin(UTransformProxy* InTransformProxy)
+{
+	bGizmoIsDragged = true;
+}
+
+void FEditorMoveCameraWithObject::OnGizmoMovementEnd(UTransformProxy* InTransformProxy)
+{
+	bGizmoIsDragged = false;
+}
+
+void FEditorMoveCameraWithObject::OnGizmoTransformChanged(UTransformProxy* InTransformProxy, UE::Math::TTransform<double> InTransform)
+{
+	if (bGizmoIsDragged && bModifierKeyIsPressed)
+	{
+		if (GCurrentLevelEditingViewportClient)
+		{
+			const FVector DragDelta = InTransform.GetLocation();
+
+			if (TSharedPtr<ILevelEditor> LevelEditor = GCurrentLevelEditingViewportClient->ParentLevelEditor.Pin())
+			{
+				if (TSharedPtr<SLevelViewport> ViewportInterface = LevelEditor->GetActiveViewportInterface())
+				{
+					FLevelEditorViewportClient& LevelViewportClient = ViewportInterface->GetLevelViewportClient();
+					LevelViewportClient.MoveViewportCamera(DragDelta, FRotator::ZeroRotator);
+				}
+			}
+		}
+	}
+}
diff --git a/Engine/Source/Editor/Experimental/EditorInteractiveToolsFramework/Public/EditorDragTools/DragToolsBehaviourSource.h b/Engine/Source/Editor/Experimental/EditorInteractiveToolsFramework/Public/EditorDragTools/DragToolsBehaviourSource.h
--- a/Engine/Source/Editor/Experimental/EditorInteractiveToolsFramework/Public/EditorDragTools/DragToolsBehaviourSource.h
+++ b/Engine/Source/Editor/Experimental/EditorInteractiveToolsFramework/Public/EditorDragTools/DragToolsBehaviourSource.h
@@ -9,6 +9,7 @@
 
 class FEditorDuplicateDragSelection;
 class FEditorDragToolBehaviorTarget;
+class FEditorMoveCameraWithObject;
 class UEditorInteractiveToolsContext;
 
 /**
@@ -66,6 +67,11 @@ private:
 	 */
 	TSharedPtr<FEditorDuplicateDragSelection> DuplicateDragSelection;
 
+	/**
+	 * Move camera together with dragged selection
+	 */
+	TSharedPtr<FEditorMoveCameraWithObject> MoveCameraWithObject;
+
 	/**
 	 * Which type of Drag tool is currently active
 	 */
diff --git a/Engine/Source/Editor/Experimental/EditorInteractiveToolsFramework/Public/EditorDragTools/EditorMoveCameraWithObject.h b/Engine/Source/Editor/Experimental/EditorInteractiveToolsFramework/Public/EditorDragTools/EditorMoveCameraWithObject.h
new file mode 100644
--- /dev/null
+++ b/Engine/Source/Editor/Experimental/EditorInteractiveToolsFramework/Public/EditorDragTools/EditorMoveCameraWithObject.h
@@ -0,0 +1,59 @@
+﻿// Copyright Epic Games, Inc. All Rights Reserved.
+
+#pragma once
+
+#include "BaseGizmos/TransformProxy.h"
+#include "EditorGizmos/TransformGizmo.h"
+
+class FLevelEditorViewportClient;
+class UEditorTransformGizmo;
+class UInteractiveToolManager;
+class UTransformProxy;
+
+/**
+ * Listens for key press of a specified modifier key, and if Level Viewport Gizmo is dragging, it duplicates current selection.
+ * Currently used with Shift (see UDragToolsBehaviorSource)
+ */
+class EDITORINTERACTIVETOOLSFRAMEWORK_API FEditorMoveCameraWithObject
+	: public IModifierToggleBehaviorTarget
+	, public FGCObject
+{
+public:
+	FEditorMoveCameraWithObject(UInteractiveToolManager* InToolManager);
+
+	virtual ~FEditorMoveCameraWithObject() override;
+
+protected:
+	//~Begin IModifierToggleBehaviorTarget
+	virtual void OnUpdateModifierState(int InModifierID, bool bInIsOn) override;
+	//~End IModifierToggleBehaviorTarget
+
+	//~Begin FGCObject
+	virtual void AddReferencedObjects(FReferenceCollector& InCollector) override;
+	virtual FString GetReferencerName() const override
+	{
+		return TEXT("FEditorMoveCameraWithObject");
+	}
+	//~End FGCObject
+
+	void OnGizmoMovementBegin(UTransformProxy* InTransformProxy);
+	void OnGizmoMovementEnd(UTransformProxy* InTransformProxy);
+	void OnGizmoTransformChanged(UTransformProxy* InTransformProxy, UE::Math::TTransform<double> InTransform);
+
+	void Initialize();
+	void Reset();
+
+	void OnGizmoCreatedDelegate(UTransformGizmo* InTransformGizmo);
+	void OnUsesNewTRSGizmosChanged(bool bInUseNewTRSGizmos);
+
+private:
+	TObjectPtr<UInteractiveToolManager> ToolManager;
+	TObjectPtr<UEditorTransformGizmo> TransformGizmo;
+
+	bool bGizmoIsDragged = false;
+	bool bModifierKeyIsPressed = false;
+
+	FDelegateHandle OnBeginPivotEditDelegate;
+	FDelegateHandle OnEndPivotEditDelegate;
+	FDelegateHandle OnTransformChangedDelegate;
+};
diff --git a/Engine/Source/Editor/UnrealEd/Private/LevelEditorViewport.cpp b/Engine/Source/Editor/UnrealEd/Private/LevelEditorViewport.cpp
--- a/Engine/Source/Editor/UnrealEd/Private/LevelEditorViewport.cpp
+++ b/Engine/Source/Editor/UnrealEd/Private/LevelEditorViewport.cpp
@@ -113,6 +113,7 @@
 #include "GenericPlatform/GenericPlatformInputDeviceMapper.h"
 #include "Rendering/StaticLightingSystemInterface.h"
 #include "TextureResource.h"
+#include "EditorDragTools/EditorDragToolBehaviorTarget.h"
 #include "Subsystems/PlacementSubsystem.h"
 #include "Subsystems/EditorElementSubsystem.h"
 
@@ -3263,8 +3264,11 @@ bool FLevelEditorViewportClient::InputWidgetDelta(FViewport* InViewport, EAxisLi
 				const bool RightMouseButtonDown = InViewport->KeyState(EKeys::RightMouseButton);
 				const bool MiddleMouseButtonDown = InViewport->KeyState(EKeys::MiddleMouseButton);
 
+				// Some functionality can be skipped if new gizmos and ITF drag tools are active
+				const bool bUseLegacyGizmoDragTools = !(UEditorInteractiveGizmoManager::UsesNewTRSGizmos() && UE::Editor::DragTools::UseEditorDragTools());
+
 				// If duplicate dragging . . .
-				if ( IsAltPressed() && (LeftMouseButtonDown || RightMouseButtonDown) )
+				if (IsAltPressed() && (LeftMouseButtonDown || RightMouseButtonDown) && bUseLegacyGizmoDragTools)
 				{
 					// The widget has been offset, so check if we should duplicate the selection.
 					if ( bDuplicateOnNextDrag )
@@ -3372,7 +3376,7 @@ bool FLevelEditorViewportClient::InputWidgetDelta(FViewport* InViewport, EAxisLi
 				ModeTools->PivotLocation += Drag;
 				ModeTools->SnappedLocation += Drag;
 
-				if( IsShiftPressed() )
+				if (IsShiftPressed() && !UEditorInteractiveGizmoManager::UsesNewTRSGizmos())
 				{
 					bApplyCameraSpeedScaleByDistance = false;
 					FVector CameraDelta( Drag );
