[ITF Drag Tools] Disable Viewport Change

ITF based Viewport Change tool is now disabled by default.
Can be enabled with CVar DragTools.EnableViewportChangeTool 1.

diff --git a/Engine/Source/Editor/Experimental/EditorInteractiveToolsFramework/Private/EditorDragTools/DragToolsBehaviourSource.cpp b/Engine/Source/Editor/Experimental/EditorInteractiveToolsFramework/Private/EditorDragTools/DragToolsBehaviourSource.cpp
--- a/Engine/Source/Editor/Experimental/EditorInteractiveToolsFramework/Private/EditorDragTools/DragToolsBehaviourSource.cpp
+++ b/Engine/Source/Editor/Experimental/EditorInteractiveToolsFramework/Private/EditorDragTools/DragToolsBehaviourSource.cpp
@@ -13,6 +13,47 @@
 #include "EditorDragTools/EditorMoveCameraWithObject.h"
 #include "Tools/EdModeInteractiveToolsContext.h"
 
+namespace UE::Editor::DragTools
+{
+	// In case New TRS Gizmos are enabled, we need Mouse Drag + Key modifier Drag Tools Behaviors to be processed before
+	// the Gizmo ones, otherwise Drag Behaviors will never be triggered. See:
+	// * UTransformGizmo::SetupBehaviors()			[DEFAULT_GIZMO_PRIORITY]
+	// * UTransformGizmo::SetupIndirectBehaviors()	[DEFAULT_GIZMO_PRIORITY - 1]
+	static constexpr int HANDLED_BEFORE_GIZMO_PRIORITY = FInputCapturePriority::DEFAULT_GIZMO_PRIORITY - 3;
+
+	// Conversely, tools which are activated with just mouse buttons input (no key modifiers, mouse drag behavior only) need to be processed after Gizmo Behaviors
+	static constexpr int HANDLED_AFTER_GIZMO_PRIORITY = FInputCapturePriority::DEFAULT_GIZMO_PRIORITY + 1;
+
+	static constexpr int FRUSTUM_SELECT_PRIORITY = HANDLED_BEFORE_GIZMO_PRIORITY;
+	static constexpr int VIEWPORT_CHANGE_PRIORITY = HANDLED_BEFORE_GIZMO_PRIORITY;
+	static constexpr int BOX_SELECT_PRIORITY = HANDLED_AFTER_GIZMO_PRIORITY;
+	static constexpr int MEASURE_PRIORITY = VIEWPORT_CHANGE_PRIORITY + 1;
+
+	// Used to group Viewport Change behaviors inside the BehaviorSet and remove them if needed
+	static const FString ViewportChangeBehaviorGroup = TEXT("ViewportChange");
+
+	// CVar initializer
+	static int32 UseViewportChangeTool = 0;
+	static FAutoConsoleVariableRef CVarEnableViewportChangeTool(
+		TEXT("DragTools.EnableViewportChangeTool"),
+		UseViewportChangeTool,
+		TEXT("Is the ITF version of the viewport change tool enabled?"),
+		FConsoleVariableDelegate::CreateLambda(
+			[](const IConsoleVariable* InVariable)
+			{
+				if (UseViewportChangeTool)
+				{
+					UDragToolsBehaviorSource::OnViewportChangeToolActivated().Broadcast();
+				}
+				else
+				{
+					UDragToolsBehaviorSource::OnViewportChangeToolDeactivated().Broadcast();
+				}
+			}
+		)
+	);
+}
+
 void UDragToolsBehaviorSource::Initialize(const UEditorInteractiveToolsContext* InInteractiveToolsContext)
 {
 	if (!InInteractiveToolsContext)
@@ -20,66 +61,28 @@
 		return;
 	}
 
+	EditorInteractiveToolsContextWeak = InInteractiveToolsContext;
+
 	TObjectPtr<UInteractiveToolManager> ToolManager = InInteractiveToolsContext->ToolManager;
 	if (!ToolManager)
 	{
 		return;
 	}
 
-	IToolsContextQueriesAPI* ContextQueriesAPI = ToolManager->GetContextQueriesAPI();
-	if (!ContextQueriesAPI)
-	{
-		return;
-	}
+	FEditorViewportClient* EditorViewportClient = GetEditorViewportClient();
 
-	FEditorViewportClient* EditorViewportClient = nullptr;
-
-	if (const FEditorModeTools* const ModeManager = InInteractiveToolsContext->GetParentEditorModeManager())
-	{
-		EditorViewportClient = ModeManager->GetFocusedViewportClient();
-	}
-
-	if (!EditorViewportClient)
-	{
-		FViewport* Viewport = ContextQueriesAPI->GetFocusedViewport();
-		if (!Viewport)
-		{
-			return;
-		}
-
-		EditorViewportClient = static_cast<FEditorViewportClient*>(Viewport->GetClient());
-		if (!EditorViewportClient)
-		{
-			return;
-		}
-	}
-
 	BehaviorSet = NewObject<UInputBehaviorSet>();
 
-	// In case New TRS Gizmos are enabled, we need Mouse Drag + Key modifier Drag Tools Behaviors to be processed before
-	// the Gizmo ones, otherwise Drag Behaviors will never be triggered. See:
-	// * UTransformGizmo::SetupBehaviors()			[DEFAULT_GIZMO_PRIORITY]
-	// * UTransformGizmo::SetupIndirectBehaviors()	[DEFAULT_GIZMO_PRIORITY - 1]
-	static constexpr int HANDLED_BEFORE_GIZMO_PRIORITY = FInputCapturePriority::DEFAULT_GIZMO_PRIORITY - 3;
-
-	// Conversely, tools which are activated with just mouse buttons input (no key modifiers, mouse drag behavior only) need to be processed after Gizmo Behaviors
-	static constexpr int HANDLED_AFTER_GIZMO_PRIORITY = FInputCapturePriority::DEFAULT_GIZMO_PRIORITY + 1;
-
-	static constexpr int FRUSTUM_SELECT_PRIORITY = HANDLED_BEFORE_GIZMO_PRIORITY;
-	static constexpr int VIEWPORT_CHANGE_PRIORITY = HANDLED_BEFORE_GIZMO_PRIORITY;
-	static constexpr int BOX_SELECT_PRIORITY = HANDLED_AFTER_GIZMO_PRIORITY;
-	static constexpr int MEASURE_PRIORITY = VIEWPORT_CHANGE_PRIORITY + 1;
-
 	// Box Select
 	{
 		TSharedPtr<FEditorDragToolBehaviorTarget> BoxSelectTool = MakeShared<FEditorBoxSelect>(EditorViewportClient);
 
-		DragTools.Add(BoxSelectTool);
+		DragTools.Add(EDragTool::BoxSelect, BoxSelectTool);
 
 		// Add drag input behavior
 		UClickDragInputBehavior* MouseBehavior = NewObject<UClickDragInputBehavior>();
 		MouseBehavior->Initialize(BoxSelectTool.Get());
-		MouseBehavior->SetDefaultPriority(BOX_SELECT_PRIORITY);
+		MouseBehavior->SetDefaultPriority(UE::Editor::DragTools::BOX_SELECT_PRIORITY);
 		BehaviorSet->Add(MouseBehavior);
 
 		// Register modifier listeners
@@ -93,7 +96,7 @@
 		KeyModifierInputBehavior->Initialize(
 			BoxSelectTool.Get(), UE::Editor::DragTools::AltKeyMod, FInputDeviceState::IsAltKeyDown
 		);
-		KeyModifierInputBehavior->SetDefaultPriority(BOX_SELECT_PRIORITY);
+		KeyModifierInputBehavior->SetDefaultPriority(UE::Editor::DragTools::BOX_SELECT_PRIORITY);
 		BehaviorSet->Add(KeyModifierInputBehavior);
 
 		BoxSelectTool->OnActivateTool().AddUObject(this, &UDragToolsBehaviorSource::OnActivateTool, EDragTool::BoxSelect);
@@ -105,12 +108,12 @@
 		TSharedPtr<FEditorDragToolBehaviorTarget> FrustumSelectTool =
 			MakeShared<FEditorFrustumSelect>(EditorViewportClient);
 
-		DragTools.Add(FrustumSelectTool);
+		DragTools.Add(EDragTool::FrustumSelect, FrustumSelectTool);
 
 		// Add drag input behavior
 		UClickDragInputBehavior* MouseBehavior = NewObject<UClickDragInputBehavior>();
 		MouseBehavior->Initialize(FrustumSelectTool.Get());
-		MouseBehavior->SetDefaultPriority(FRUSTUM_SELECT_PRIORITY);
+		MouseBehavior->SetDefaultPriority(UE::Editor::DragTools::FRUSTUM_SELECT_PRIORITY);
 		MouseBehavior->SetUseLeftMouseButton();
 		BehaviorSet->Add(MouseBehavior);
 
@@ -125,7 +128,7 @@
 		KeyModifierInputBehavior->Initialize(
 			FrustumSelectTool.Get(), UE::Editor::DragTools::AltKeyMod, FInputDeviceState::IsAltKeyDown
 		);
-		KeyModifierInputBehavior->SetDefaultPriority(FRUSTUM_SELECT_PRIORITY);
+		KeyModifierInputBehavior->SetDefaultPriority(UE::Editor::DragTools::FRUSTUM_SELECT_PRIORITY);
 		BehaviorSet->Add(KeyModifierInputBehavior);
 
 		FrustumSelectTool->OnActivateTool().AddUObject(
@@ -140,12 +143,12 @@
 	{
 		TSharedPtr<FEditorDragToolBehaviorTarget> MeasureTool = MakeShared<FEditorMeasureTool>(EditorViewportClient);
 
-		DragTools.Add(MeasureTool);
+		DragTools.Add(EDragTool::Measure, MeasureTool);
 
 		// Add drag input behavior
 		UClickDragInputBehavior* MouseBehavior = NewObject<UClickDragInputBehavior>();
 		MouseBehavior->Initialize(MeasureTool.Get());
-		MouseBehavior->SetDefaultPriority(MEASURE_PRIORITY);
+		MouseBehavior->SetDefaultPriority(UE::Editor::DragTools::MEASURE_PRIORITY);
 		MouseBehavior->SetUseMiddleMouseButton();
 		BehaviorSet->Add(MouseBehavior);
 
@@ -155,37 +158,13 @@
 
 	// Viewport Change
 	{
-		// Requires Level Viewport Client specifically.
-		// todo: Maybe we could move some code from Level Viewport Client to Editor Viewport Client, to make this more generic?
-		if (EditorViewportClient->IsLevelEditorClient())
+		if (IsViewportChangeToolEnabled())
 		{
-			TSharedPtr<FEditorDragToolBehaviorTarget> ViewportChangeTool =
-				MakeShared<FEditorViewportChange>(EditorViewportClient);
-
-			DragTools.Add(ViewportChangeTool);
-
-			// Add drag input behavior
-			UClickDragInputBehavior* MouseBehavior = NewObject<UClickDragInputBehavior>();
-			MouseBehavior->Initialize(ViewportChangeTool.Get());
-			MouseBehavior->SetDefaultPriority(VIEWPORT_CHANGE_PRIORITY);
-			MouseBehavior->SetUseMiddleMouseButton();
-			BehaviorSet->Add(MouseBehavior);
-
-			// Register modifier listeners
-			UKeyAsModifierInputBehavior* KeyModifierInputBehavior = NewObject<UKeyAsModifierInputBehavior>();
-			KeyModifierInputBehavior->Initialize(
-				ViewportChangeTool.Get(), UE::Editor::DragTools::CtrlKeyMod, FInputDeviceState::IsCtrlKeyDown
-			);
-			KeyModifierInputBehavior->SetDefaultPriority(VIEWPORT_CHANGE_PRIORITY);
-			BehaviorSet->Add(KeyModifierInputBehavior);
-
-			ViewportChangeTool->OnActivateTool().AddUObject(
-				this, &UDragToolsBehaviorSource::OnActivateTool, EDragTool::ViewportChange
-			);
-			ViewportChangeTool->OnDeactivateTool().AddUObject(
-				this, &UDragToolsBehaviorSource::OnDeactivateTool, EDragTool::ViewportChange
-			);
+			ActivateViewportChangeTool();
 		}
+
+		OnViewportChangeToolActivated().AddUObject(this, &UDragToolsBehaviorSource::ActivateViewportChangeTool);
+		OnViewportChangeToolDeactivated().AddUObject(this, &UDragToolsBehaviorSource::DeactivateViewportChangeTool);
 	}
 
 	// Duplicate Drag Tool
@@ -195,7 +174,7 @@
 		KeyModifierInputBehavior->Initialize(
 			DuplicateDragSelection.Get(), UE::Editor::DragTools::AltKeyMod, FInputDeviceState::IsAltKeyDown
 		);
-		KeyModifierInputBehavior->SetDefaultPriority(HANDLED_BEFORE_GIZMO_PRIORITY);
+		KeyModifierInputBehavior->SetDefaultPriority(UE::Editor::DragTools::HANDLED_BEFORE_GIZMO_PRIORITY);
 
 		BehaviorSet->Add(KeyModifierInputBehavior);
 	}
@@ -206,7 +185,7 @@
 
 		UKeyAsModifierInputBehavior* KeyModifierInputBehavior = NewObject<UKeyAsModifierInputBehavior>();
 		KeyModifierInputBehavior->Initialize(MoveCameraWithObject.Get(), UE::Editor::DragTools::ShiftKeyMod, FInputDeviceState::IsShiftKeyDown);
-		KeyModifierInputBehavior->SetDefaultPriority(HANDLED_BEFORE_GIZMO_PRIORITY);
+		KeyModifierInputBehavior->SetDefaultPriority(UE::Editor::DragTools::HANDLED_BEFORE_GIZMO_PRIORITY);
 
 		BehaviorSet->Add(KeyModifierInputBehavior);
 	}
@@ -216,10 +195,10 @@
 {
 	if (ActiveToolType >= 0)
 	{
-		int32 ToolIndex = GetToolIndex(static_cast<EDragTool::Type>(ActiveToolType));
-		if (DragTools.IsValidIndex(ToolIndex))
+		EDragTool::Type ToolType = static_cast<EDragTool::Type>(ActiveToolType);
+		if (DragTools.Contains(ToolType))
 		{
-			if (const TSharedPtr<FEditorDragToolBehaviorTarget>& ActiveTool = DragTools[ToolIndex])
+			if (const TSharedPtr<FEditorDragToolBehaviorTarget>& ActiveTool = DragTools[ToolType])
 			{
 				return ActiveTool.Get();
 			}
@@ -229,25 +208,80 @@
 	return nullptr;
 }
 
-int32 UDragToolsBehaviorSource::GetToolIndex(EDragTool::Type InDragToolType)
+FEditorViewportClient* UDragToolsBehaviorSource::GetEditorViewportClient() const
 {
-	// Please add tools in the same order as they appear in EDragTool declaration
-	static TArray<int32> ToolIndexes = {
-		EDragTool::BoxSelect, EDragTool::FrustumSelect, EDragTool::Measure, EDragTool::ViewportChange
-	};
+	FEditorViewportClient* EditorViewportClient = nullptr;
 
-	int32 ToolIndex = ToolIndexes.Find(InDragToolType);
-	if (ToolIndex != INDEX_NONE)
+	if (const UEditorInteractiveToolsContext* InteractiveToolsContext = EditorInteractiveToolsContextWeak.Get())
 	{
-		return ToolIndex;
+		if (const FEditorModeTools* const ModeManager = InteractiveToolsContext->GetParentEditorModeManager())
+		{
+			EditorViewportClient = ModeManager->GetFocusedViewportClient();
+		}
+
+		if (!EditorViewportClient)
+		{
+			if (UInteractiveToolManager* ToolManager = InteractiveToolsContext->ToolManager)
+			{
+				if (IToolsContextQueriesAPI* ContextQueriesAPI = ToolManager->GetContextQueriesAPI())
+				{
+					if (FViewport* Viewport = ContextQueriesAPI->GetFocusedViewport())
+					{
+						EditorViewportClient = static_cast<FEditorViewportClient*>(Viewport->GetClient());
+					}
+				}
+			}
+		}
 	}
 
-	return -1;
+	return EditorViewportClient;
 }
 
+bool UDragToolsBehaviorSource::IsViewportChangeToolEnabled()
+{
+	return UE::Editor::DragTools::UseViewportChangeTool == 1;
+}
+
+UDragToolsBehaviorSource::FOnOnViewportChangeToolToggleDelegate& UDragToolsBehaviorSource::OnViewportChangeToolActivated()
+{
+	static FOnOnViewportChangeToolToggleDelegate OnViewportChangedToolActivated;
+	return OnViewportChangedToolActivated;
+}
+
+UDragToolsBehaviorSource::FOnOnViewportChangeToolToggleDelegate& UDragToolsBehaviorSource::OnViewportChangeToolDeactivated()
+{
+	static FOnOnViewportChangeToolToggleDelegate OnViewportChangedToolDeactivated;
+	return OnViewportChangedToolDeactivated;
+}
+
+void UDragToolsBehaviorSource::RegisterSource()
+{
+	if (const UEditorInteractiveToolsContext* InteractiveToolsContext = EditorInteractiveToolsContextWeak.Get())
+	{
+		if (UInputRouter* InputRouter = InteractiveToolsContext->InputRouter)
+		{
+			InputRouter->RegisterSource(this);
+		}
+	}
+}
+
+void UDragToolsBehaviorSource::DeregisterSource()
+{
+	if (const UEditorInteractiveToolsContext* InteractiveToolsContext = EditorInteractiveToolsContextWeak.Get())
+	{
+		if (UInputRouter* InputRouter = InteractiveToolsContext->InputRouter)
+		{
+			InputRouter->DeregisterSource(this);
+		}
+	}
+}
+
 void UDragToolsBehaviorSource::OnActivateTool(EDragTool::Type InDragToolType)
 {
-	ActiveToolType = InDragToolType;
+	if (DragTools.Contains(InDragToolType))
+	{
+		ActiveToolType = InDragToolType;
+	}
 }
 
 void UDragToolsBehaviorSource::OnDeactivateTool(EDragTool::Type InDragToolType)
@@ -265,3 +299,82 @@
 		ActiveDragTool->Render(InSceneView, InCanvas);
 	}
 }
+
+void UDragToolsBehaviorSource::ActivateViewportChangeTool()
+{
+	// If tool is already active, ignore
+	if (DragTools.Contains(EDragTool::ViewportChange))
+	{
+		return;
+	}
+
+	FEditorViewportClient* EditorViewportClient = GetEditorViewportClient();
+	if (!EditorViewportClient)
+	{
+		return;
+	}
+
+	// Make sure we temporarily deregister this source while we change its behaviors list
+	DeregisterSource();
+
+	// Setup Drag Tool and Behaviors for Viewport Change
+	{
+		// Requires Level Viewport Client specifically.
+		// todo: Maybe we could move some code from Level Viewport Client to Editor Viewport Client, to make this more generic?
+		if (EditorViewportClient->IsLevelEditorClient())
+		{
+			TSharedPtr<FEditorDragToolBehaviorTarget> ViewportChangeTool =
+				MakeShared<FEditorViewportChange>(EditorViewportClient);
+
+			DragTools.Emplace(EDragTool::ViewportChange, ViewportChangeTool);
+
+			// Add drag input behavior
+			UClickDragInputBehavior* MouseBehavior = NewObject<UClickDragInputBehavior>();
+			MouseBehavior->Initialize(ViewportChangeTool.Get());
+			MouseBehavior->SetDefaultPriority(UE::Editor::DragTools::VIEWPORT_CHANGE_PRIORITY);
+			MouseBehavior->SetUseMiddleMouseButton();
+			BehaviorSet->Add(MouseBehavior, this, UE::Editor::DragTools::ViewportChangeBehaviorGroup);
+
+			// Register modifier listeners
+			UKeyAsModifierInputBehavior* KeyModifierInputBehavior = NewObject<UKeyAsModifierInputBehavior>();
+			KeyModifierInputBehavior->Initialize(
+				ViewportChangeTool.Get(), UE::Editor::DragTools::CtrlKeyMod, FInputDeviceState::IsCtrlKeyDown
+			);
+			KeyModifierInputBehavior->SetDefaultPriority(UE::Editor::DragTools::VIEWPORT_CHANGE_PRIORITY);
+			BehaviorSet->Add(KeyModifierInputBehavior, this, UE::Editor::DragTools::ViewportChangeBehaviorGroup);
+
+			ViewportChangeTool->OnActivateTool().AddUObject(
+				this, &UDragToolsBehaviorSource::OnActivateTool, EDragTool::ViewportChange
+			);
+			ViewportChangeTool->OnDeactivateTool().AddUObject(
+				this, &UDragToolsBehaviorSource::OnDeactivateTool, EDragTool::ViewportChange
+			);
+		}
+	}
+
+	// Register again
+	RegisterSource();
+}
+
+void UDragToolsBehaviorSource::DeactivateViewportChangeTool()
+{
+	// Setup Drag Tool and Behaviors for Viewport Change
+	DeregisterSource();
+
+	// Remove Behaviors and DragTool for Viewport Change
+	{
+		if (BehaviorSet)
+		{
+			BehaviorSet->RemoveByGroup(UE::Editor::DragTools::ViewportChangeBehaviorGroup);
+		}
+
+		OnDeactivateTool(EDragTool::ViewportChange);
+		if (DragTools.Contains(EDragTool::ViewportChange))
+		{
+			DragTools.FindAndRemoveChecked(EDragTool::ViewportChange)->OnForceEndCapture();
+		}
+	}
+
+	// Register again
+	RegisterSource();
+}
diff --git a/Engine/Source/Editor/Experimental/EditorInteractiveToolsFramework/Private/EditorDragTools/EditorDragToolBehaviorTarget.cpp b/Engine/Source/Editor/Experimental/EditorInteractiveToolsFramework/Private/EditorDragTools/EditorDragToolBehaviorTarget.cpp
--- a/Engine/Source/Editor/Experimental/EditorInteractiveToolsFramework/Private/EditorDragTools/EditorDragToolBehaviorTarget.cpp
+++ b/Engine/Source/Editor/Experimental/EditorInteractiveToolsFramework/Private/EditorDragTools/EditorDragToolBehaviorTarget.cpp
@@ -15,7 +15,7 @@ static int32 UseITFTools = 0;
 static FAutoConsoleVariableRef CVarEnableITFTools(
 	TEXT("DragTools.EnableITFTools"),
 	UseITFTools,
-	TEXT("Is the ITF version of Drag Tools enabled?\nCurrently supports: Frustum Select."),
+	TEXT("Is the ITF version of Drag Tools enabled?"),
 	FConsoleVariableDelegate::CreateLambda(
 		[](const IConsoleVariable* InVariable)
 		{
diff --git a/Engine/Source/Editor/Experimental/EditorInteractiveToolsFramework/Public/EditorDragTools/DragToolsBehaviourSource.h b/Engine/Source/Editor/Experimental/EditorInteractiveToolsFramework/Public/EditorDragTools/DragToolsBehaviourSource.h
--- a/Engine/Source/Editor/Experimental/EditorInteractiveToolsFramework/Public/EditorDragTools/DragToolsBehaviourSource.h
+++ b/Engine/Source/Editor/Experimental/EditorInteractiveToolsFramework/Public/EditorDragTools/DragToolsBehaviourSource.h
@@ -22,10 +22,26 @@ class UDragToolsBehaviorSource final : public UObject, public IInputBehaviorSour
 	GENERATED_BODY()
 
 public:
+	static bool IsViewportChangeToolEnabled();
+
+	DECLARE_MULTICAST_DELEGATE(FOnOnViewportChangeToolToggleDelegate)
+	static FOnOnViewportChangeToolToggleDelegate& OnViewportChangeToolActivated();
+	static FOnOnViewportChangeToolToggleDelegate& OnViewportChangeToolDeactivated();
+
 	//~ Begin IInputBehaviorSource
 	EDITORINTERACTIVETOOLSFRAMEWORK_API virtual const UInputBehaviorSet* GetInputBehaviors() const override { return BehaviorSet; }
 	//~ End IClickDragBehaviorTarget
 
+	/**
+	 * Register this Input Behavior Source to the InputRouter
+	 */
+	EDITORINTERACTIVETOOLSFRAMEWORK_API void RegisterSource();
+
+	/**
+	 * Deregister this Input Behavior Source from the InputRouter
+	 */
+	EDITORINTERACTIVETOOLSFRAMEWORK_API void DeregisterSource();
+
 	/**
 	 * Creates and stores the drag tools.
 	 * Instantiates the BehaviorSet hosting all behaviors required by drag tools.
@@ -41,26 +57,31 @@ public:
 	EDITORINTERACTIVETOOLSFRAMEWORK_API void RenderTools(const FSceneView* InSceneView, FCanvas* InCanvas) const;
 
 private:
-	static int32 GetToolIndex(EDragTool::Type InDragToolType);
-
 	void OnActivateTool(EDragTool::Type InDragToolType);
 	void OnDeactivateTool(EDragTool::Type InDragToolType);
 
+	void ActivateViewportChangeTool();
+	void DeactivateViewportChangeTool();
+
 	/**
 	 * Returns the drag tool currently being used, if any
 	 */
 	FEditorDragToolBehaviorTarget* GetActiveTool() const;
 
+	FEditorViewportClient* GetEditorViewportClient() const;
+
 	/**
 	 * Hosting drag tools behaviors
 	 */
 	UPROPERTY(Transient)
 	TObjectPtr<UInputBehaviorSet> BehaviorSet;
 
+	TWeakObjectPtr<const UEditorInteractiveToolsContext> EditorInteractiveToolsContextWeak;
+
 	/**
-	 * Available Drag Tools, access by using the index provided by GetToolIndex(EDragTool::Type)
+	 * Available Drag Tools, access by using their type (EDragTool::Type)
 	 */
-	TArray<TSharedPtr<FEditorDragToolBehaviorTarget>> DragTools;
+	TMap<EDragTool::Type, TSharedPtr<FEditorDragToolBehaviorTarget>> DragTools;
 
 	/**
 	 * Duplicate dragged selection tool
diff --git a/Engine/Source/Editor/UnrealEd/Private/Tools/EdModeInteractiveToolsContext.cpp b/Engine/Source/Editor/UnrealEd/Private/Tools/EdModeInteractiveToolsContext.cpp
--- a/Engine/Source/Editor/UnrealEd/Private/Tools/EdModeInteractiveToolsContext.cpp
+++ b/Engine/Source/Editor/UnrealEd/Private/Tools/EdModeInteractiveToolsContext.cpp
@@ -1352,23 +1352,21 @@ void UModeManagerInteractiveToolsContext::RegisterDragTools()
 
 		DragToolsBehaviorSource = NewObject<UDragToolsBehaviorSource>();
 		DragToolsBehaviorSource->Initialize(this);
-
-		if (DragToolsBehaviorSource->GetInputBehaviors())
-		{
-			InputRouter->RegisterSource(DragToolsBehaviorSource);
-		}
+		DragToolsBehaviorSource->RegisterSource();
 	}
 }
+
 void UModeManagerInteractiveToolsContext::UnregisterDragTools()
 {
 	if (DragToolsBehaviorSource)
 	{
 		// Revert to previous state, with Alt Modifier not being processed
 		SetBypassAltModifier(true);
-		InputRouter->DeregisterSource(DragToolsBehaviorSource);
+		DragToolsBehaviorSource->DeregisterSource();
 		DragToolsBehaviorSource = nullptr;
 	}
 }
+
 void UModeManagerInteractiveToolsContext::OnEditorModeChanged(const FEditorModeID& InModeID, bool bInIsEnteringMode)
 {
 	if (!UE::Editor::DragTools::UseEditorDragTools())
