[New TRS Gizmos] Update Settings properties names and menu to match "Use Experimental Gizmos" label

diff --git a/Engine/Plugins/Experimental/GizmoFramework/Source/GizmoSettings/GizmoSettings.Build.cs b/Engine/Plugins/Experimental/GizmoFramework/Source/GizmoSettings/GizmoSettings.Build.cs
--- a/Engine/Plugins/Experimental/GizmoFramework/Source/GizmoSettings/GizmoSettings.Build.cs
+++ b/Engine/Plugins/Experimental/GizmoFramework/Source/GizmoSettings/GizmoSettings.Build.cs
@@ -12,11 +12,7 @@ namespace UnrealBuildTool.Rules
 			PrivateDependencyModuleNames.AddRange(
 				new string[]
 				{
-					"Core",
-					"CoreUObject",
-					"DeveloperSettings",
-					"Engine",
-					"EditorInteractiveToolsFramework"
+					"Core"
 				}
 			);
 		}
diff --git a/Engine/Plugins/Experimental/GizmoFramework/Source/GizmoSettings/Private/GizmoSettings.cpp b/Engine/Plugins/Experimental/GizmoFramework/Source/GizmoSettings/Private/GizmoSettings.cpp
deleted file mode 100644
diff --git a/Engine/Plugins/Experimental/GizmoFramework/Source/GizmoSettings/Private/GizmoSettingsModule.cpp b/Engine/Plugins/Experimental/GizmoFramework/Source/GizmoSettings/Private/GizmoSettingsModule.cpp
--- a/Engine/Plugins/Experimental/GizmoFramework/Source/GizmoSettings/Private/GizmoSettingsModule.cpp
+++ b/Engine/Plugins/Experimental/GizmoFramework/Source/GizmoSettings/Private/GizmoSettingsModule.cpp
@@ -1,30 +1,10 @@
 // Copyright Epic Games, Inc. All Rights Reserved.
 
-#include "GizmoSettingsModule.h"
 #include "Modules/ModuleManager.h"
 
-#include "GizmoSettings.h"
-#include "EditorInteractiveGizmoManager.h"
-
-void FGizmoSettingsModule::StartupModule()
-{
-	UGizmoSettings::OnSettingsChange.AddRaw(this, &FGizmoSettingsModule::UpdateGizmos);
-	UpdateGizmos(GetDefault<UGizmoSettings>());
-}
-
-void FGizmoSettingsModule::ShutdownModule()
+// Dummy module, currently only used by Viewport Toolbar to show/hide the "Use Experimental Gizmos" entry.
+class FGizmoSettingsModule : public IModuleInterface
 {
-	UGizmoSettings::OnSettingsChange.RemoveAll(this);
-}
-
-void FGizmoSettingsModule::UpdateGizmos(const UGizmoSettings* InSettings)
-{
-	if (InSettings)
-	{
-		UEditorInteractiveGizmoManager::SetUsesNewTRSGizmos(InSettings->bEnableNewGizmos);
-		UEditorInteractiveGizmoManager::SetGizmosParameters(InSettings->GizmoParameters);
-	}
-}
-
+};
 
 IMPLEMENT_MODULE(FGizmoSettingsModule, GizmoSettings)
diff --git a/Engine/Plugins/Experimental/GizmoFramework/Source/GizmoSettings/Public/GizmoSettings.h b/Engine/Plugins/Experimental/GizmoFramework/Source/GizmoSettings/Public/GizmoSettings.h
deleted file mode 100644
diff --git a/Engine/Plugins/Experimental/GizmoFramework/Source/GizmoSettings/Public/GizmoSettingsModule.h b/Engine/Plugins/Experimental/GizmoFramework/Source/GizmoSettings/Public/GizmoSettingsModule.h
deleted file mode 100644
diff --git a/Engine/Source/Editor/Experimental/EditorInteractiveToolsFramework/Private/EditorInteractiveGizmoManager.cpp b/Engine/Source/Editor/Experimental/EditorInteractiveToolsFramework/Private/EditorInteractiveGizmoManager.cpp
--- a/Engine/Source/Editor/Experimental/EditorInteractiveToolsFramework/Private/EditorInteractiveGizmoManager.cpp
+++ b/Engine/Source/Editor/Experimental/EditorInteractiveToolsFramework/Private/EditorInteractiveGizmoManager.cpp
@@ -52,7 +52,7 @@ void UEditorInteractiveGizmoManager::SetUsesNewTRSGizmos(const bool bUseNewTRSGi
 	{
 		if (bUseNewTRSGizmos != TransformGizmoEditorSettings->UsesNewTRSGizmo())
 		{
-			TransformGizmoEditorSettings->SetUsesLegacyGizmo(!bUseNewTRSGizmos);
+			TransformGizmoEditorSettings->SetUseExperimentalGizmo(bUseNewTRSGizmos);
 		}
 	}
 }
diff --git a/Engine/Source/Editor/Experimental/EditorInteractiveToolsFramework/Private/TransformGizmoEditorSettings.cpp b/Engine/Source/Editor/Experimental/EditorInteractiveToolsFramework/Private/TransformGizmoEditorSettings.cpp
--- a/Engine/Source/Editor/Experimental/EditorInteractiveToolsFramework/Private/TransformGizmoEditorSettings.cpp
+++ b/Engine/Source/Editor/Experimental/EditorInteractiveToolsFramework/Private/TransformGizmoEditorSettings.cpp
@@ -43,7 +43,7 @@ void UTransformGizmoEditorSettings::PostEditChangeProperty(FPropertyChangedEvent
 {
 	Super::PostEditChangeProperty(InPropertyChangedEvent);
 
-	static const FName UseLegacyGizmoName = GET_MEMBER_NAME_CHECKED(UTransformGizmoEditorSettings, bEnableLegacyGizmo);
+	static const FName UseExperimentalGizmoName = GET_MEMBER_NAME_CHECKED(UTransformGizmoEditorSettings, bUseExperimentalGizmo);
 	static const FName EnableArcballRotateName = GET_MEMBER_NAME_CHECKED(UTransformGizmoEditorSettings, bEnableArcballRotate);
 	static const FName EnableScreenRotateName = GET_MEMBER_NAME_CHECKED(UTransformGizmoEditorSettings, bEnableScreenRotate);
 	static const FName EnableAxisDrawingName = GET_MEMBER_NAME_CHECKED(UTransformGizmoEditorSettings, bEnableAxisDrawing);
@@ -56,7 +56,7 @@ void UTransformGizmoEditorSettings::PostEditChangeProperty(FPropertyChangedEvent
 
 	const FName ChangedPropertyName = InPropertyChangedEvent.Property->GetFName();
 
-	if (ChangedPropertyName == UseLegacyGizmoName)
+	if (ChangedPropertyName == UseExperimentalGizmoName)
 	{
 		BroadcastNewTRSGizmoChange();
 	}
@@ -91,11 +91,11 @@ void UTransformGizmoEditorSettings::PostEditChangeProperty(FPropertyChangedEvent
 }
 #endif // WITH_EDITOR
 
-void UTransformGizmoEditorSettings::SetUsesLegacyGizmo(bool bInUseLegacyGizmo)
+void UTransformGizmoEditorSettings::SetUseExperimentalGizmo(bool bInUseExperimentalGizmo)
 {
-	if (bEnableLegacyGizmo != bInUseLegacyGizmo)
+	if (bUseExperimentalGizmo != bInUseExperimentalGizmo)
 	{
-		bEnableLegacyGizmo = bInUseLegacyGizmo;
+		bUseExperimentalGizmo = bInUseExperimentalGizmo;
 		SaveConfig();
 
 		BroadcastNewTRSGizmoChange();
@@ -109,12 +109,12 @@ void UTransformGizmoEditorSettings::SetUsesLegacyGizmo(bool bInUseLegacyGizmo)
 
 bool UTransformGizmoEditorSettings::UsesLegacyGizmo() const
 {
-	return bEnableLegacyGizmo;
+	return !UsesNewTRSGizmo();
 }
 
 bool UTransformGizmoEditorSettings::UsesNewTRSGizmo() const
 {
-	return !UsesLegacyGizmo();
+	return bUseExperimentalGizmo;
 }
 
 void UTransformGizmoEditorSettings::SetGizmosParameters(const FGizmosParameters& InGizmosParameters)
diff --git a/Engine/Source/Editor/Experimental/EditorInteractiveToolsFramework/Public/TransformGizmoEditorSettings.h b/Engine/Source/Editor/Experimental/EditorInteractiveToolsFramework/Public/TransformGizmoEditorSettings.h
--- a/Engine/Source/Editor/Experimental/EditorInteractiveToolsFramework/Public/TransformGizmoEditorSettings.h
+++ b/Engine/Source/Editor/Experimental/EditorInteractiveToolsFramework/Public/TransformGizmoEditorSettings.h
@@ -23,47 +23,43 @@ public:
 	UPROPERTY(Config, EditAnywhere, Category = "Transform Gizmo", meta = (ShowOnlyInnerProperties))
 	FGizmosParameters GizmosParameters;
 
-	UPROPERTY(Config, EditAnywhere, Category = "Transform Gizmo", DisplayName = "Transform Gizmo Size", meta = (ClampMin = "-10.0", ClampMax = "150.0"))
+	UPROPERTY(Config, EditAnywhere, Category = "Transform Gizmo", DisplayName = "Transform Gizmo Size", meta = (ClampMin = "-10.0", ClampMax = "150.0", EditCondition = "!bUseExperimentalGizmo"))
 	float TransformGizmoSize = 0.0f;
-	//~ End Transform Gizmo Category
-
-	//~ Begin Legacy Gizmo Category
 
-	/** If true, legacy gizmos will be used. */
-	UPROPERTY(Config, EditAnywhere, Category = "Legacy Gizmo", DisplayName = "Enable Legacy Gizmo")
-	bool bEnableLegacyGizmo = true;
+	/** If true, the new TRS gizmos will be used. */
+	UPROPERTY(Config, EditAnywhere, Category = "Transform Gizmo", DisplayName = "Enable Experimental Gizmo")
+	bool bUseExperimentalGizmo = false;
 
 	/**
 	 * Allow arcball rotation with rotate widget
 	 * (updates the setting with the same name found in Level Editor Viewport Settings)
 	 */
-	UPROPERTY(EditAnywhere, Category = "Legacy Gizmo", DisplayName = "Enable Arcball Rotate", meta = (EditCondition = "bEnableLegacyGizmo"))
+	UPROPERTY(EditAnywhere, Category = "Transform Gizmo", DisplayName = "Enable Arcball Rotate", meta = (EditCondition = "!bUseExperimentalGizmo"))
 	bool bEnableArcballRotate;
 
 	/**
 	 * Allow screen rotation with rotate widget
 	 * (updates the setting with the same name found in Level Editor Viewport Settings)
 	 */
-	UPROPERTY(EditAnywhere, Category = "Legacy Gizmo", DisplayName = "Enable Screen Rotate", meta = (EditCondition = "bEnableLegacyGizmo"))
+	UPROPERTY(EditAnywhere, Category = "Transform Gizmo", DisplayName = "Enable Screen Rotate", meta = (EditCondition = "!bUseExperimentalGizmo"))
 	bool bEnableScreenRotate;
 
 	/**
 	 * If true, the Edit widget of a transform will display the axis
 	 * (updates the setting with the same name found in Level Editor Viewport Settings)
 	 */
-	UPROPERTY(EditAnywhere, Category = "Legacy Gizmo", DisplayName = "Enable Axis drawing for Transform Edit Widget", meta = (EditCondition = "bEnableLegacyGizmo"))
+	UPROPERTY(EditAnywhere, Category = "Transform Gizmo", DisplayName = "Enable Axis drawing for Transform Edit Widget", meta = (EditCondition = "!bUseExperimentalGizmo"))
 	bool bEnableAxisDrawing;
 
 	/**
 	 * Allow translate/rotate widget
 	 * (updates the setting with the same name found in Level Editor Viewport Settings)
 	 */
-	UPROPERTY(EditAnywhere, Category = "Legacy Gizmo", DisplayName = "Enable Combined Translate/Rotate Widget", meta = (EditCondition = "bEnableLegacyGizmo"))
+	UPROPERTY(EditAnywhere, Category = "Transform Gizmo", DisplayName = "Enable Combined Translate/Rotate Widget", meta = (EditCondition = "!bUseExperimentalGizmo"))
 	bool bEnableCombinedTranslateRotate;
-	//~ End Legacy Gizmo Category
+	//~ End Transform Gizmo Category
 
-public:
-	void SetUsesLegacyGizmo(bool bInUseLegacyGizmo);
+	void SetUseExperimentalGizmo(bool bInUseExperimentalGizmo);
 	bool UsesLegacyGizmo() const;
 	bool UsesNewTRSGizmo() const;
 
diff --git a/Engine/Source/Editor/UnrealEd/Private/ViewportToolbar/UnrealEdViewportToolbar.cpp b/Engine/Source/Editor/UnrealEd/Private/ViewportToolbar/UnrealEdViewportToolbar.cpp
--- a/Engine/Source/Editor/UnrealEd/Private/ViewportToolbar/UnrealEdViewportToolbar.cpp
+++ b/Engine/Source/Editor/UnrealEd/Private/ViewportToolbar/UnrealEdViewportToolbar.cpp
@@ -726,12 +726,15 @@ FToolMenuEntry CreateViewportToolbarTransformsSection()
 
 FToolMenuEntry CreateTransformsSubmenu()
 {
+	// Cache this once per session
+	const bool bIsGizmoSettingsModuleLoaded = FModuleManager::Get().IsModuleLoaded("GizmoSettings");
+
 	FToolMenuEntry Entry = FToolMenuEntry::InitSubMenu(
 		"Transform",
 		LOCTEXT("TransformsSubmenuLabel", "Transform"),
 		LOCTEXT("TransformsSubmenuTooltip", "Viewport-related transforms tools"),
 		FNewToolMenuDelegate::CreateLambda(
-			[](UToolMenu* Submenu) -> void
+			[bIsGizmoSettingsModuleLoaded](UToolMenu* Submenu) -> void
 			{
 				{
 					FToolMenuEntryToolBarData ToolBarData;
@@ -958,16 +961,21 @@ FToolMenuEntry CreateTransformsSubmenu()
 							}
 						);
 
-						GizmoSection.AddMenuEntry(
-							"TRSGizmoToggle",
-							LOCTEXT("TRSGizmoToggleLabel", "Use New TRS Gizmos"),
-							LOCTEXT("TRSGizmoToggleTooltip", "Whether or not to use the new TRS gizmos"),
-							FSlateIcon(),
-							Action,
-							EUserInterfaceActionType::ToggleButton
-						);
+						// Only show if GizmoSettings module is loaded (Experimental GizmoFramework plugin)
+						if (bIsGizmoSettingsModuleLoaded)
+						{
+							GizmoSection.AddMenuEntry(
+								"TRSGizmoToggle",
+								LOCTEXT("TRSGizmoToggleLabel", "Use Experimental Gizmos"),
+								LOCTEXT("TRSGizmoToggleTooltip", "Whether or not to use the new Transform Gizmos"),
+								FSlateIcon(),
+								Action,
+								EUserInterfaceActionType::ToggleButton
+							);
+						}
 					}
 
+					
 					TSharedRef<SWidget> GizmoScaleWidget =
 						// clang-format off
 						SNew(SHorizontalBox)
@@ -1008,8 +1016,15 @@ FToolMenuEntry CreateTransformsSubmenu()
 						"GizmoScale", GizmoScaleWidget, LOCTEXT("GizmoScaleLabel", "Gizmo Scale")
 					));
 					GizmoScaleEntry.Icon = FSlateIcon(FAppStyle::GetAppStyleSetName(), "ViewportToolbar.GizmoScale");
-					
-					
+
+					// Hide Gizmo Scale, since at the moment it is not used by new Gizmo
+					GizmoScaleEntry.Visibility = TAttribute<bool>::CreateLambda(
+						[]()
+						{
+							return !UEditorInteractiveGizmoManager::UsesNewTRSGizmos();
+						}
+					);
+
 					{
 						FUIAction PreserveNonUniformScaleAction;
 						PreserveNonUniformScaleAction.ExecuteAction.BindLambda([]
