[New TRS Gizmos] Update Properties order and disable when non experimental

diff --git a/Engine/Source/Editor/Experimental/EditorInteractiveToolsFramework/Public/EditorGizmos/TransformGizmo.h b/Engine/Source/Editor/Experimental/EditorInteractiveToolsFramework/Public/EditorGizmos/TransformGizmo.h
--- a/Engine/Source/Editor/Experimental/EditorInteractiveToolsFramework/Public/EditorGizmos/TransformGizmo.h
+++ b/Engine/Source/Editor/Experimental/EditorInteractiveToolsFramework/Public/EditorGizmos/TransformGizmo.h
@@ -113,21 +113,21 @@ struct FGizmosParameters
 	GENERATED_BODY()
 
 	/** Determines how dragging the rotate gizmo affects the selected objects */
-	UPROPERTY(EditAnywhere, Category = "Transform Gizmo")
+	UPROPERTY(EditAnywhere, Category = "Experimental")
 	TEnumAsByte<EAxisRotateMode::Type> RotateMode = EAxisRotateMode::Arc;
 
 	/**
 	 * When enabled, Ctrl+MMB activates the Y axis and Ctrl+RMB activates the Z axis.
 	 * When disabled, Ctrl+RMB activates the Y axis and Ctrl+LMB+RMB activates the Z axis.
 	 */
-	UPROPERTY(EditAnywhere, Category = "Transform Gizmo")
+	UPROPERTY(EditAnywhere, Category = "Experimental")
 	bool bCtrlMiddleDoesY = true;
 
 	/**
 	 * When enabled, the list of coordinate spaces in the viewport toolbar offers rig space: a coordinate system that is
 	 * similar to parent space but uses gimbal rotations.
 	 */
-	UPROPERTY(EditAnywhere, Category = "Transform Gizmo")
+	UPROPERTY(EditAnywhere, Category = "Experimental")
 	bool bEnableExplicit = false;
 };
 
diff --git a/Engine/Source/Editor/Experimental/EditorInteractiveToolsFramework/Public/TransformGizmoEditorSettings.h b/Engine/Source/Editor/Experimental/EditorInteractiveToolsFramework/Public/TransformGizmoEditorSettings.h
--- a/Engine/Source/Editor/Experimental/EditorInteractiveToolsFramework/Public/TransformGizmoEditorSettings.h
+++ b/Engine/Source/Editor/Experimental/EditorInteractiveToolsFramework/Public/TransformGizmoEditorSettings.h
@@ -20,16 +20,9 @@ public:
 	virtual void PostEditChangeProperty(FPropertyChangedEvent& InPropertyChangedEvent) override;
 #endif
 
-	UPROPERTY(Config, EditAnywhere, Category = "Transform Gizmo", meta = (ShowOnlyInnerProperties))
-	FGizmosParameters GizmosParameters;
-
 	UPROPERTY(Config, EditAnywhere, Category = "Transform Gizmo", DisplayName = "Transform Gizmo Size", meta = (ClampMin = "-10.0", ClampMax = "150.0", EditCondition = "!bUseExperimentalGizmo"))
 	float TransformGizmoSize = 0.0f;
 
-	/** If true, the new TRS gizmos will be used. */
-	UPROPERTY(Config, EditAnywhere, Category = "Transform Gizmo", DisplayName = "Enable Experimental Gizmo")
-	bool bUseExperimentalGizmo = false;
-
 	/**
 	 * Allow arcball rotation with rotate widget
 	 * (updates the setting with the same name found in Level Editor Viewport Settings)
@@ -59,6 +52,15 @@ public:
 	bool bEnableCombinedTranslateRotate;
 	//~ End Transform Gizmo Category
 
+	//~ Begin Experimental Gizmo Category
+	/** If true, the new TRS gizmos will be used. */
+	UPROPERTY(Config, EditAnywhere, Category = "Experimental", DisplayName = "Enable Experimental Gizmo")
+	bool bUseExperimentalGizmo = false;
+
+	UPROPERTY(Config, EditAnywhere, Category = "Experimental", meta = (ShowOnlyInnerProperties),  meta = (EditCondition = "bUseExperimentalGizmo"))
+	FGizmosParameters GizmosParameters;
+	//~ End Experimental Gizmo Category
+
 	void SetUseExperimentalGizmo(bool bInUseExperimentalGizmo);
 	bool UsesLegacyGizmo() const;
 	bool UsesNewTRSGizmo() const;
