[EditorGizmos] Transform gizmo visuals:
- Shaded material
- Plane handle thickness
- Arcball visiblity toggle

diff --git a/Engine/Source/Editor/Experimental/EditorInteractiveToolsFramework/Private/EditorGizmos/TransformGizmo.cpp b/Engine/Source/Editor/Experimental/EditorInteractiveToolsFramework/Private/EditorGizmos/TransformGizmo.cpp
--- a/Engine/Source/Editor/Experimental/EditorInteractiveToolsFramework/Private/EditorGizmos/TransformGizmo.cpp
+++ b/Engine/Source/Editor/Experimental/EditorInteractiveToolsFramework/Private/EditorGizmos/TransformGizmo.cpp
@@ -28,6 +28,8 @@
 
 DEFINE_LOG_CATEGORY_STATIC(LogTransformGizmo, Log, All);
 
+ENUM_RANGE_BY_COUNT(ETransformGizmoPartIdentifier, ETransformGizmoPartIdentifier::Max);
+
 namespace GizmoLocals
 {
 	
@@ -89,6 +91,100 @@ void MakePerpVectors(const FVector& Normal, FVector& OutPerp1, FVector& OutPerp2
 	
 }
 
+namespace UE::EditorInteractiveToolsFramework::Private
+{
+	namespace TransformGizmoLocals
+	{
+		template <EAxis::Type Axis>
+		static constexpr FVector GetAxisVector()
+		{
+			return FVector::ZeroVector;
+		}
+
+		template <>
+		constexpr FVector GetAxisVector<EAxis::X>()
+		{
+			return FVector::XAxisVector;
+		}
+
+		template <>
+		FVector GetAxisVector<EAxis::Y>()
+		{
+			return AxisDisplayInfo::GetAxisDisplayCoordinateSystem() == EAxisList::LeftUpForward
+				? -FVector::YAxisVector
+				: FVector::YAxisVector;
+		}
+
+		template <>
+		constexpr FVector GetAxisVector<EAxis::Z>()
+		{
+			return FVector::ZAxisVector;
+		}
+
+		static FVector GetAxisVector(EAxis::Type InAxis)
+		{
+			switch (InAxis)
+			{
+			case EAxis::X:
+				return GetAxisVector<EAxis::X>();
+
+			case EAxis::Y:
+				return GetAxisVector<EAxis::Y>();
+
+			case EAxis::Z:
+				return GetAxisVector<EAxis::Z>();
+
+			case EAxis::None:
+			default:
+				return FVector::ZeroVector;
+			}
+		}
+
+		template <EAxis::Type Axis>
+		static constexpr FVector GetAxisSideVector()
+		{
+			return FVector::ZeroVector;
+		}
+
+		template <>
+		constexpr FVector GetAxisSideVector<EAxis::X>()
+		{
+			return FVector::YAxisVector;
+		}
+
+		template <>
+		constexpr FVector GetAxisSideVector<EAxis::Y>()
+		{
+			return FVector::ZAxisVector;
+		}
+
+		template <>
+		constexpr FVector GetAxisSideVector<EAxis::Z>()
+		{
+			return FVector::XAxisVector;
+		}
+
+		static FVector GetAxisSideVector(EAxis::Type InAxis)
+		{
+			switch (InAxis)
+			{
+			case EAxis::X:
+				return GetAxisSideVector<EAxis::X>();
+
+			case EAxis::Y:
+				return GetAxisSideVector<EAxis::Y>();
+
+			case EAxis::Z:
+				return GetAxisSideVector<EAxis::Z>();
+
+			case EAxis::None:
+			default:
+				return FVector::ZeroVector;
+			}
+		}
+	}
+};
+
 void UTransformGizmo::SetDisallowNegativeScaling(bool bDisallow)
 {
 	bDisallowNegativeScaling = bDisallow;
@@ -114,15 +210,15 @@ void UTransformGizmo::Setup()
 	GizmoElementRoot->SetConstantScale(true);
 	GizmoElementRoot->SetHoverMaterial(CurrentAxisMaterial);
 	GizmoElementRoot->SetInteractMaterial(CurrentAxisMaterial);
-	GizmoElementRoot->SetHoverLineColor(CurrentColor);
-	GizmoElementRoot->SetInteractLineColor(CurrentColor);
+	GizmoElementRoot->SetHoverLineColor(Style.CurrentColor);
+	GizmoElementRoot->SetInteractLineColor(Style.CurrentColor);
 
 	// the main gimbal rotation element that manages the three gimbal rotation axis as a group
 	RotateGimbalElement = NewObject<UGizmoElementGimbal>();
 	RotateGimbalElement->SetHoverMaterial(CurrentAxisMaterial);
 	RotateGimbalElement->SetInteractMaterial(CurrentAxisMaterial);
-	RotateGimbalElement->SetHoverLineColor(CurrentColor);
-	RotateGimbalElement->SetInteractLineColor(CurrentColor);
+	RotateGimbalElement->SetHoverLineColor(Style.CurrentColor);
+	RotateGimbalElement->SetInteractLineColor(Style.CurrentColor);
 	GizmoElementRoot->Add(RotateGimbalElement);
 
 	bInInteraction = false;
@@ -330,7 +426,7 @@ void UTransformGizmo::SetupMaterials()
 			}
 		}
 
-		static const FString MaterialName = TEXT("/Engine/EditorMaterials/TransformGizmoMaterial.TransformGizmoMaterial");
+		static const FString MaterialName = TEXT("/Engine/EditorMaterials/TransformGizmoMaterial_PseudoLitMasked.TransformGizmoMaterial_PseudoLitMasked");
 		UMaterial* Material = FindObject<UMaterial>(nullptr, *MaterialName);
 		if (!Material)
 		{
@@ -352,13 +448,13 @@ void UTransformGizmo::SetupMaterials()
 	AxisMaterialZ->SetVectorParameterValue("GizmoColor", AxisDisplayInfo::GetAxisColor(EAxisList::Z));
 
 	GreyMaterial = UMaterialInstanceDynamic::Create(AxisMaterialBase, NULL);
-	GreyMaterial->SetVectorParameterValue("GizmoColor", GreyColor);
+	GreyMaterial->SetVectorParameterValue("GizmoColor", Style.GreyColor);
 
 	WhiteMaterial = UMaterialInstanceDynamic::Create(AxisMaterialBase, NULL);
-	WhiteMaterial->SetVectorParameterValue("GizmoColor", WhiteColor);
+	WhiteMaterial->SetVectorParameterValue("GizmoColor", Style.WhiteColor);
 
 	CurrentAxisMaterial = UMaterialInstanceDynamic::Create(AxisMaterialBase, NULL);
-	CurrentAxisMaterial->SetVectorParameterValue("GizmoColor", CurrentColor);
+	CurrentAxisMaterial->SetVectorParameterValue("GizmoColor", Style.CurrentColor);
 
 	OpaquePlaneMaterialXY = UMaterialInstanceDynamic::Create(AxisMaterialBase, NULL);
 	OpaquePlaneMaterialXY->SetVectorParameterValue("GizmoColor", FLinearColor::White);
@@ -833,17 +929,17 @@ void UTransformGizmo::EnablePlanarObjects(bool bTranslate, bool bEnableX, bool b
 	check(GizmoElementRoot);
 
 	auto EnablePlanarElement = [this](
-		TObjectPtr<UGizmoElementRectangle>& PlanarElement,
+		TObjectPtr<UGizmoElementBox>& PlanarElement,
 		ETransformGizmoPartIdentifier PartId,
 		const FVector& Axis0,
 		const FVector& Axis1,
 		const FVector& Axis2,
-		const FLinearColor& AxisColor,
+		UMaterialInterface* Material,
 		bool bEnable)
 	{
 		if (bEnable && PlanarElement == nullptr)
 		{
-			PlanarElement = MakePlanarHandle(PartId, Axis0, Axis1, Axis2, TransparentVertexColorMaterial, AxisColor);
+			PlanarElement = MakePlanarHandle(PartId, Axis0, Axis1, Axis2, Material);
 			GizmoElementRoot->Add(PlanarElement);
 		}
 
@@ -857,21 +953,21 @@ void UTransformGizmo::EnablePlanarObjects(bool bTranslate, bool bEnableX, bool b
 	const bool bEnableYZ = bEnableY && bEnableZ;
 	const bool bEnableXZ = bEnableX && bEnableZ;
 
-	const FVector XAxis(1.0f, 0.0f, 0.0f);
-	const FVector YAxis(0.0f, 1.0f, 0.0f);
-	const FVector ZAxis(0.0f, 0.0f, 1.0f);
+	const FVector XAxis = FVector::XAxisVector;
+	const FVector YAxis = FVector::YAxisVector;
+	const FVector ZAxis = FVector::ZAxisVector;
 
 	if (bTranslate)
 	{
-		EnablePlanarElement(TranslatePlanarXYElement, ETransformGizmoPartIdentifier::TranslateXYPlanar, XAxis, YAxis, ZAxis, AxisDisplayInfo::GetAxisColor(EAxisList::Z), bEnableXY);
-		EnablePlanarElement(TranslatePlanarYZElement, ETransformGizmoPartIdentifier::TranslateYZPlanar, YAxis, ZAxis, XAxis, AxisDisplayInfo::GetAxisColor(EAxisList::X), bEnableYZ);
-		EnablePlanarElement(TranslatePlanarXZElement, ETransformGizmoPartIdentifier::TranslateXZPlanar, ZAxis, XAxis, YAxis, AxisDisplayInfo::GetAxisColor(EAxisList::Y), bEnableXZ);
+		EnablePlanarElement(TranslatePlanarXYElement, ETransformGizmoPartIdentifier::TranslateXYPlanar, XAxis, YAxis, ZAxis, AxisMaterialZ, bEnableXY);
+		EnablePlanarElement(TranslatePlanarYZElement, ETransformGizmoPartIdentifier::TranslateYZPlanar, YAxis, ZAxis, XAxis, AxisMaterialX, bEnableYZ);
+		EnablePlanarElement(TranslatePlanarXZElement, ETransformGizmoPartIdentifier::TranslateXZPlanar, ZAxis, XAxis, YAxis, AxisMaterialY, bEnableXZ);
 	}
 	else
 	{
-		EnablePlanarElement(ScalePlanarXYElement, ETransformGizmoPartIdentifier::ScaleXYPlanar, XAxis, YAxis, ZAxis, AxisDisplayInfo::GetAxisColor(EAxisList::Z), bEnableXY);
-		EnablePlanarElement(ScalePlanarYZElement, ETransformGizmoPartIdentifier::ScaleYZPlanar, YAxis, ZAxis, XAxis, AxisDisplayInfo::GetAxisColor(EAxisList::X), bEnableYZ);
-		EnablePlanarElement(ScalePlanarXZElement, ETransformGizmoPartIdentifier::ScaleXZPlanar, ZAxis, XAxis, YAxis, AxisDisplayInfo::GetAxisColor(EAxisList::Y), bEnableXZ);
+		EnablePlanarElement(ScalePlanarXYElement, ETransformGizmoPartIdentifier::ScaleXYPlanar, XAxis, YAxis, ZAxis, AxisMaterialZ, bEnableXY);
+		EnablePlanarElement(ScalePlanarYZElement, ETransformGizmoPartIdentifier::ScaleYZPlanar, YAxis, ZAxis, XAxis, AxisMaterialX, bEnableYZ);
+		EnablePlanarElement(ScalePlanarXZElement, ETransformGizmoPartIdentifier::ScaleXZPlanar, ZAxis, XAxis, YAxis, AxisMaterialY, bEnableXZ);
 	}
 }
 
@@ -942,13 +1038,17 @@ void UTransformGizmo::EnableRotate(EAxisList::Type InAxisListToDraw)
 	{
 		if (!RotateScreenSpaceElement)
 		{
-			RotateScreenSpaceElement = MakeRotateCircleHandle(ETransformGizmoPartIdentifier::RotateScreenSpace, RotateScreenSpaceRadius, RotateScreenSpaceCircleColor, false);
+			RotateScreenSpaceElement = MakeRotateCircleHandle(
+				ETransformGizmoPartIdentifier::RotateScreenSpace,
+				Style.RotateCircleRadius + Style.RotateScreenSpaceRadiusOffset,
+				Style.RotateScreenSpaceCircleColor,
+				EGizmoElementDrawType::Line);
 			GizmoElementRoot->Add(RotateScreenSpaceElement);
 		}
 
 		if (!RotateArcballElement)
 		{
-			RotateArcballElement = MakeRotateCircleHandle(ETransformGizmoPartIdentifier::RotateArcball, RotateArcballSphereRadius, RotateArcballCircleColor, true);
+			RotateArcballElement = MakeArcballCircleHandle(ETransformGizmoPartIdentifier::RotateArcball, Style.RotateArcballSphereRadius, Style.RotateArcballCircleColor);
 			GizmoElementRoot->Add(RotateArcballElement);
 		}
 	}
@@ -967,25 +1067,27 @@ void UTransformGizmo::EnableScale(EAxisList::Type InAxisListToDraw)
 {
 	check(GizmoElementRoot);
 
-	bool bEnableX = static_cast<uint8>(InAxisListToDraw) & static_cast<uint8>(EAxisList::X);
-	bool bEnableY = static_cast<uint8>(InAxisListToDraw) & static_cast<uint8>(EAxisList::Y);
-	bool bEnableZ = static_cast<uint8>(InAxisListToDraw) & static_cast<uint8>(EAxisList::Z);
+	using namespace UE::EditorInteractiveToolsFramework::Private::TransformGizmoLocals;
+
+	const bool bEnableX = static_cast<uint8>(InAxisListToDraw) & static_cast<uint8>(EAxisList::X);
+	const bool bEnableY = static_cast<uint8>(InAxisListToDraw) & static_cast<uint8>(EAxisList::Y);
+	const bool bEnableZ = static_cast<uint8>(InAxisListToDraw) & static_cast<uint8>(EAxisList::Z);
 	
 	if (bEnableX && ScaleXAxisElement == nullptr)
 	{
-		ScaleXAxisElement = MakeScaleAxis(ETransformGizmoPartIdentifier::ScaleXAxis, FVector(1.0f, 0.0f, 0.0f), FVector(0.0f, 1.0f, 0.0f), AxisMaterialX);
+		ScaleXAxisElement = MakeScaleAxis(ETransformGizmoPartIdentifier::ScaleXAxis, FVector::XAxisVector, GetAxisSideVector<EAxis::X>(), AxisMaterialX);
 		GizmoElementRoot->Add(ScaleXAxisElement);
 	}
 
 	if (bEnableY && ScaleYAxisElement == nullptr)
 	{
-		ScaleYAxisElement = MakeScaleAxis(ETransformGizmoPartIdentifier::ScaleYAxis, FVector(0.0f, 1.0f, 0.0f), FVector(0.0f, 0.0f, 1.0f), AxisMaterialY);
+		ScaleYAxisElement = MakeScaleAxis(ETransformGizmoPartIdentifier::ScaleYAxis, FVector::YAxisVector, GetAxisSideVector<EAxis::Y>(), AxisMaterialY);
 		GizmoElementRoot->Add(ScaleYAxisElement);
 	}
 
 	if (bEnableZ && ScaleZAxisElement == nullptr)
 	{
-		ScaleZAxisElement = MakeScaleAxis(ETransformGizmoPartIdentifier::ScaleZAxis, FVector(0.0f, 0.0f, 1.0f), FVector(1.0f, 0.0f, 0.0f), AxisMaterialZ);
+		ScaleZAxisElement = MakeScaleAxis(ETransformGizmoPartIdentifier::ScaleZAxis, FVector::ZAxisVector, GetAxisSideVector<EAxis::Z>(), AxisMaterialZ);
 		GizmoElementRoot->Add(ScaleZAxisElement);
 	}
 
@@ -1202,13 +1304,31 @@ void UTransformGizmo::OnParametersChanged(const FGizmosParameters& InParameters)
 			}
 		}
 	}
-	
+
 	DefaultRotateMode = InParameters.RotateMode;
+
+	// @note: The clamp values should match the property metadata (ClampMin)
+	Style.AxisSizeMultiplier = FMath::Max(0.01f, InParameters.AxisSizeMultiplier);
+	Style.LineThicknessMultiplier = FMath::Max(0.01f, InParameters.DevParameters.LineThicknessMultiplier);
+	Style.TranslateAxisLength = FMath::Max(1.0f, InParameters.DevParameters.TranslateAxisLength);
+	Style.RotateArcballSphereRadius = FMath::Max(1.0f, InParameters.DevParameters.RotateAxisRadius);
+	Style.RotateAxisRadius = Style.RotateArcballSphereRadius;
+	Style.RotateCircleRadius =Style. RotateArcballSphereRadius;
+	Style.ScaleAxisLength = FMath::Max(1.0f, InParameters.DevParameters.ScaleAxisLength);
+
+	UpdateElements();
 }
 
 UGizmoElementArrow* UTransformGizmo::MakeTranslateAxis(ETransformGizmoPartIdentifier InPartId, const FVector& InAxisDir, const FVector& InSideDir, UMaterialInterface* InMaterial)
 {
-	const float SizeCoeff = GetSizeCoefficient();
+	UGizmoElementArrow* ArrowElement = NewObject<UGizmoElementArrow>();
+	ArrowElement->SetPartIdentifier(static_cast<uint32>(InPartId));
+	ArrowElement->SetHeadType(EGizmoElementArrowHeadType::Cone);
+	ArrowElement->SetEndCaps(true);
+	ArrowElement->SetNumSides(32);
+	ArrowElement->SetMaterial(InMaterial);
+	ArrowElement->SetViewDependentType(EGizmoElementViewDependentType::Axis);
+	ArrowElement->SetViewDependentAxis(InAxisDir);
 
 	FVector AxisDir = InAxisDir;
 	if (AxisDisplayInfo::GetAxisDisplayCoordinateSystem() == EAxisList::LeftUpForward)
@@ -1218,72 +1338,161 @@ UGizmoElementArrow* UTransformGizmo::MakeTranslateAxis(ETransformGizmoPartIdenti
 			AxisDir = -AxisDir;
 		}
 	}
-	
-	UGizmoElementArrow* ArrowElement = NewObject<UGizmoElementArrow>();
-	ArrowElement->SetPartIdentifier(static_cast<uint32>(InPartId));
-	ArrowElement->SetHeadType(EGizmoElementArrowHeadType::Cone);
-	ArrowElement->SetBase(AxisDir * AxisLengthOffset);
-	ArrowElement->SetDirection(AxisDir);
-	ArrowElement->SetSideDirection(InSideDir);
-	ArrowElement->SetBodyLength(TranslateAxisLength * SizeCoeff);
-	ArrowElement->SetBodyRadius(AxisRadius);
-	ArrowElement->SetHeadLength(TranslateAxisConeHeight * SizeCoeff);
-	ArrowElement->SetHeadRadius(TranslateAxisConeRadius * SizeCoeff);
-	ArrowElement->SetNumSides(32);
-	ArrowElement->SetMaterial(InMaterial);
-	ArrowElement->SetViewDependentType(EGizmoElementViewDependentType::Axis);
-	ArrowElement->SetViewDependentAxis(InAxisDir);
+
+	UpdateTranslateAxis(ArrowElement, AxisDir, InSideDir);
+
 	return ArrowElement;
 }
 
-UGizmoElementArrow* UTransformGizmo::MakeScaleAxis(ETransformGizmoPartIdentifier InPartId, const FVector& InAxisDir, const FVector& InSideDir, UMaterialInterface* InMaterial)
+void UTransformGizmo::UpdateTranslateAxis(UGizmoElementArrow* InElement)
 {
-	const float SizeCoeff = GetSizeCoefficient();
+	if (!InElement)
+	{
+		return;
+	}
 
-	FVector AxisDir = InAxisDir;
-	if (AxisDisplayInfo::GetAxisDisplayCoordinateSystem() == EAxisList::LeftUpForward)
+	UpdateTranslateAxis(InElement, InElement->GetDirection(), InElement->GetSideDirection());
+}
+
+void UTransformGizmo::UpdateTranslateAxis(UGizmoElementArrow* InElement, const EAxis::Type InAxis)
+{
+	if (!InElement)
 	{
-		if (InPartId == ETransformGizmoPartIdentifier::ScaleYAxis)
-		{
-			AxisDir = -AxisDir;
-		}
+		return;
 	}
 
+	using namespace UE::EditorInteractiveToolsFramework::Private::TransformGizmoLocals;
+
+	UpdateTranslateAxis(InElement, GetAxisVector(InAxis), GetAxisSideVector(InAxis));
+}
+
+void UTransformGizmo::UpdateTranslateAxis(UGizmoElementArrow* InElement, const FVector& InAxisDir, const FVector& InSideDir)
+{
+	if (!InElement)
+	{
+		return;
+	}
+
+	const float SizeCoeff = GetSizeCoefficient();
+
+	InElement->SetBase(InAxisDir * Style.AxisLengthOffset * Style.AxisSizeMultiplier);
+	InElement->SetDirection(InAxisDir);
+	InElement->SetSideDirection(InSideDir);
+	InElement->SetBodyLength(Style.TranslateAxisLength * Style.AxisSizeMultiplier * SizeCoeff);
+	InElement->SetBodyRadius(Style.AxisRadius);
+	InElement->SetHeadLength(Style.TranslateAxisConeHeight * SizeCoeff);
+	InElement->SetHeadRadius(Style.TranslateAxisConeRadius * SizeCoeff);
+}
+
+UGizmoElementArrow* UTransformGizmo::MakeScaleAxis(ETransformGizmoPartIdentifier InPartId, const FVector& InAxisDir, const FVector& InSideDir, UMaterialInterface* InMaterial)
+{
 	UGizmoElementArrow* ArrowElement = NewObject<UGizmoElementArrow>();
 	ArrowElement->SetPartIdentifier(static_cast<uint32>(InPartId));
 	ArrowElement->SetHeadType(EGizmoElementArrowHeadType::Cube);
-	ArrowElement->SetBase(AxisDir * AxisLengthOffset);
-	ArrowElement->SetDirection(AxisDir);
-	ArrowElement->SetSideDirection(InSideDir);
-	ArrowElement->SetBodyLength(ScaleAxisLength * SizeCoeff);
-	ArrowElement->SetBodyRadius(AxisRadius);
-	ArrowElement->SetHeadLength(ScaleAxisCubeDim * SizeCoeff);
 	ArrowElement->SetNumSides(32);
 	ArrowElement->SetMaterial(InMaterial);
 	ArrowElement->SetViewDependentType(EGizmoElementViewDependentType::Axis);
 	ArrowElement->SetViewDependentAxis(InAxisDir);
+
+	FVector AxisDir = InAxisDir;
+    if (AxisDisplayInfo::GetAxisDisplayCoordinateSystem() == EAxisList::LeftUpForward)
+    {
+    	if (InPartId == ETransformGizmoPartIdentifier::ScaleYAxis)
+    	{
+    		AxisDir = -AxisDir;
+    	}
+    }
+
+	UpdateScaleAxis(ArrowElement, AxisDir, InSideDir);
+
 	return ArrowElement;
 }
 
-UGizmoElementBox* UTransformGizmo::MakeUniformScaleHandle()
+void UTransformGizmo::UpdateScaleAxis(UGizmoElementArrow* InElement)
+{
+	if (!InElement)
+	{
+		return;
+	}
+
+	UpdateScaleAxis(InElement, InElement->GetDirection(), InElement->GetSideDirection());
+}
+
+void UTransformGizmo::UpdateScaleAxis(UGizmoElementArrow* InElement, const EAxis::Type InAxis)
+{
+	if (!InElement)
+	{
+		return;
+	}
+
+	using namespace UE::EditorInteractiveToolsFramework::Private::TransformGizmoLocals;
+
+	UpdateScaleAxis(InElement, GetAxisVector(InAxis), GetAxisSideVector(InAxis));
+}
+
+void UTransformGizmo::UpdateScaleAxis(UGizmoElementArrow* InElement, const FVector& InAxisDir, const FVector& InSideDir)
 {
+	if (!InElement)
+	{
+		return;
+	}
+
 	const float SizeCoeff = GetSizeCoefficient();
-	
+
+	InElement->SetBase(InAxisDir * Style.AxisLengthOffset * Style.AxisSizeMultiplier);
+	InElement->SetDirection(InAxisDir);
+	InElement->SetSideDirection(InSideDir);
+	InElement->SetBodyLength(Style.ScaleAxisLength * Style.AxisSizeMultiplier * SizeCoeff);
+	InElement->SetBodyRadius(Style.AxisRadius);
+	InElement->SetHeadLength(Style.ScaleAxisCubeDim * SizeCoeff);
+}
+
+void UTransformGizmo::UpdateAllScaleAxis()
+{
+	if (ScaleXAxisElement)
+	{
+		UpdateScaleAxis(ScaleXAxisElement, EAxis::X);
+	}
+
+	if (ScaleYAxisElement)
+	{
+		UpdateScaleAxis(ScaleYAxisElement, EAxis::Y);
+	}
+
+	if (ScaleZAxisElement)
+	{
+		UpdateScaleAxis(ScaleZAxisElement, EAxis::Z);
+	}
+}
+
+UGizmoElementBox* UTransformGizmo::MakeUniformScaleHandle()
+{
 	UGizmoElementBox* BoxElement = NewObject<UGizmoElementBox>();
 	BoxElement->SetPartIdentifier(static_cast<uint32>(ETransformGizmoPartIdentifier::ScaleUniform));
 	BoxElement->SetCenter(FVector::ZeroVector);
 	BoxElement->SetUpDirection(FVector::UpVector);
 	BoxElement->SetSideDirection(FVector::RightVector);
-	BoxElement->SetDimensions(FVector(ScaleAxisCubeDim, ScaleAxisCubeDim, ScaleAxisCubeDim) * SizeCoeff);
 	BoxElement->SetMaterial(GreyMaterial);
+	UpdateUniformScaleHandle(BoxElement);
+
 	return BoxElement;
 }
 
-UGizmoElementRectangle* UTransformGizmo::MakePlanarHandle(ETransformGizmoPartIdentifier InPartId, const FVector& InUpDirection, const FVector& InSideDirection, const FVector& InPlaneNormal,
-	UMaterialInterface* InMaterial, const FLinearColor& InVertexColor)
+void UTransformGizmo::UpdateUniformScaleHandle(UGizmoElementBox* InElement)
 {
+	if (!InElement)
+	{
+		return;
+	}
+
 	const float SizeCoeff = GetSizeCoefficient();
 
+	InElement->SetDimensions(FVector( Style.ScaleAxisCubeDim, Style.ScaleAxisCubeDim, Style.ScaleAxisCubeDim) * SizeCoeff);
+}
+
+UGizmoElementBox* UTransformGizmo::MakePlanarHandle(ETransformGizmoPartIdentifier InPartId, const FVector& InUpDirection, const FVector& InSideDirection, const FVector& InPlaneNormal,
+	UMaterialInterface* InMaterial)
+{
 	FVector UpDirection = InUpDirection;
 	FVector SideDirection = InSideDirection;
 	if (AxisDisplayInfo::GetAxisDisplayCoordinateSystem() == EAxisList::LeftUpForward)
@@ -1297,55 +1506,70 @@ UGizmoElementRectangle* UTransformGizmo::MakePlanarHandle(ETransformGizmoPartIde
 			UpDirection = -UpDirection;
 		}
 	}
-	
-	FVector PlanarHandleCenter = (UpDirection + SideDirection) * PlanarHandleOffset * SizeCoeff;
 
-	FLinearColor LineColor = InVertexColor;
-	FLinearColor VertexColor = LineColor;
-	VertexColor.A = LargeOuterAlpha;
+	UGizmoElementBox* BoxElement = NewObject<UGizmoElementBox>();
+	BoxElement->SetPartIdentifier(static_cast<uint32>(InPartId));
+	BoxElement->SetUpDirection(UpDirection);
+	BoxElement->SetSideDirection(SideDirection);
+	BoxElement->SetMaterial(InMaterial);
+	BoxElement->SetViewDependentType(EGizmoElementViewDependentType::Plane);
+	BoxElement->SetViewDependentAxis(InPlaneNormal);
+	UpdatePlanarHandle(BoxElement, UpDirection, SideDirection);
 
-	UGizmoElementRectangle* RectangleElement = NewObject<UGizmoElementRectangle>();
-	RectangleElement->SetPartIdentifier(static_cast<uint32>(InPartId));
-	RectangleElement->SetUpDirection(UpDirection);
-	RectangleElement->SetSideDirection(SideDirection);
-	RectangleElement->SetCenter(PlanarHandleCenter);
-	RectangleElement->SetHeight(PlanarHandleSize * SizeCoeff);
-	RectangleElement->SetWidth(PlanarHandleSize * SizeCoeff);
-	RectangleElement->SetMaterial(InMaterial);
-	RectangleElement->SetVertexColor(VertexColor);
-	RectangleElement->SetLineColor(LineColor);
-	RectangleElement->SetDrawLine(true);
-	RectangleElement->SetDrawMesh(true);
-	RectangleElement->SetHitMesh(true);
-	RectangleElement->SetViewDependentType(EGizmoElementViewDependentType::Plane);
-	RectangleElement->SetViewDependentAxis(InPlaneNormal);
-	return RectangleElement;
+	return BoxElement;
+}
+
+void UTransformGizmo::UpdatePlanarHandle(UGizmoElementBox* InElement, const FVector& InUpDirection, const FVector& InSideDirection)
+{
+	if (!InElement)
+    {
+    	return;
+    }
+
+    const float SizeCoeff = GetSizeCoefficient();
+
+    FVector PlanarHandleCenter = (InUpDirection + InSideDirection) * ((Style.TranslateAxisLength * Style.AxisSizeMultiplier) + Style.PlanarHandleOffset) * SizeCoeff;
+
+	InElement->SetCenter(PlanarHandleCenter);
+	InElement->SetDimensions(FVector(Style.PlanarHandleThickness, Style.PlanarHandleSize, Style.PlanarHandleSize) * SizeCoeff);
 }
 
 UGizmoElementRectangle* UTransformGizmo::MakeTranslateScreenSpaceHandle()
 {
-	const float SizeCoeff = GetSizeCoefficient();
-	
 	UGizmoElementRectangle* RectangleElement = NewObject<UGizmoElementRectangle>();
 	RectangleElement->SetPartIdentifier(static_cast<uint32>(ETransformGizmoPartIdentifier::TranslateScreenSpace));
 	RectangleElement->SetUpDirection(FVector::UpVector);
 	RectangleElement->SetSideDirection(FVector::RightVector);
 	RectangleElement->SetCenter(FVector::ZeroVector);
-	RectangleElement->SetHeight(TranslateScreenSpaceHandleSize * SizeCoeff);
-	RectangleElement->SetWidth(TranslateScreenSpaceHandleSize * SizeCoeff);
 	RectangleElement->SetViewAlignType(EGizmoElementViewAlignType::PointScreen);
 	RectangleElement->SetViewAlignAxis(FVector::UpVector);
 	RectangleElement->SetViewAlignNormal(-FVector::ForwardVector);
 	RectangleElement->SetMaterial(TransparentVertexColorMaterial);
-	RectangleElement->SetLineColor(ScreenSpaceColor);
+	RectangleElement->SetLineColor(Style.ScreenSpaceColor);
 	RectangleElement->SetHitMesh(true);
 	RectangleElement->SetDrawMesh(false);
 	RectangleElement->SetDrawLine(true);
-	RectangleElement->SetHoverLineThicknessMultiplier(3.0f);
-	RectangleElement->SetInteractLineThicknessMultiplier(3.0f);
+	RectangleElement->SetLineThickness(Style.AxisRadius * 2.0f);
+	RectangleElement->SetHoverLineThicknessMultiplier(1.5f);
+	RectangleElement->SetInteractLineThicknessMultiplier(1.5f);
+	UpdateTranslateScreenSpaceHandle(RectangleElement);
+
 	return RectangleElement;
 }
 
+void UTransformGizmo::UpdateTranslateScreenSpaceHandle(UGizmoElementRectangle* InElement)
+{
+	if (!InElement)
+	{
+		return;
+	}
+
+	const float SizeCoeff = GetSizeCoefficient();
+
+	InElement->SetHeight(Style.TranslateScreenSpaceHandleSize * SizeCoeff);
+	InElement->SetWidth(Style.TranslateScreenSpaceHandleSize * SizeCoeff);
+}
+
 UGizmoElementTorus* UTransformGizmo::MakeDefaultRotateAxis(const ETransformGizmoPartIdentifier InPartId)
 {
 	switch (InPartId)
@@ -1373,15 +1597,11 @@ UGizmoElementTorus* UTransformGizmo::MakeDefaultRotateAxis(const ETransformGizmo
 UGizmoElementTorus* UTransformGizmo::MakeRotateAxis(ETransformGizmoPartIdentifier InPartId, const FVector& TorusAxis0, const FVector& TorusAxis1,
 	UMaterialInterface* InMaterial, UMaterialInterface* InCurrentMaterial)
 {
-	const float SizeCoeff = GetSizeCoefficient();
-	
 	UGizmoElementTorus* RotateAxisElement = NewObject<UGizmoElementTorus>();
 	RotateAxisElement->SetPartIdentifier(static_cast<uint32>(InPartId));
 	RotateAxisElement->SetCenter(FVector::ZeroVector);
-	RotateAxisElement->SetRadius(UTransformGizmo::RotateAxisOuterRadius * SizeCoeff);
-	RotateAxisElement->SetNumSegments(UTransformGizmo::RotateAxisNumSegments);
-	RotateAxisElement->SetInnerRadius(UTransformGizmo::RotateAxisInnerRadius * SizeCoeff);
-	RotateAxisElement->SetNumInnerSlices(UTransformGizmo::RotateAxisInnerSlices);
+	RotateAxisElement->SetNumSegments(Style.RotateAxisNumSegments);
+	RotateAxisElement->SetNumInnerSlices(Style.RotateAxisInnerSlices);
 	RotateAxisElement->SetAxis0(TorusAxis0);
 	RotateAxisElement->SetAxis1(TorusAxis1);
 	const FVector TorusNormal = RotateAxisElement->GetAxis0() ^ RotateAxisElement->GetAxis1();
@@ -1391,50 +1611,141 @@ UGizmoElementTorus* UTransformGizmo::MakeRotateAxis(ETransformGizmoPartIdentifie
 	RotateAxisElement->SetViewDependentAxis(TorusNormal);
 	RotateAxisElement->SetViewAlignType(EGizmoElementViewAlignType::Axial);
 	RotateAxisElement->SetViewAlignAxialAngleTol(static_cast<float>(UE_DOUBLE_SMALL_NUMBER));
-	RotateAxisElement->SetPartialViewDependentMaxCosTol(1.0-UE_DOUBLE_SMALL_NUMBER);	
 	RotateAxisElement->SetViewAlignAxis(TorusNormal);
 	RotateAxisElement->SetViewAlignNormal(TorusAxis1);
 	RotateAxisElement->SetMaterial(InMaterial);
+	UpdateRotateAxis(RotateAxisElement);
+
 	return RotateAxisElement;
 }
 
-UGizmoElementCircle* UTransformGizmo::MakeRotateCircleHandle(ETransformGizmoPartIdentifier InPartId, float InRadius, const FLinearColor& InColor, const bool bFill)
+void UTransformGizmo::UpdateRotateAxis(UGizmoElementTorus* InElement)
 {
+	if (!InElement)
+	{
+		return;
+	}
+
 	const float SizeCoeff = GetSizeCoefficient();
-	
+
+	InElement->SetRadius(((Style.RotateAxisRadius * Style.AxisSizeMultiplier) + Style.RotateAxisOuterRadiusOffset) * SizeCoeff);
+	InElement->SetInnerRadius((Style.RotateAxisInnerRadius * Style.LineThicknessMultiplier) * SizeCoeff);
+}
+
+UGizmoElementCircle* UTransformGizmo::MakeArcballCircleHandle(ETransformGizmoPartIdentifier InPartId, float InRadius, const FLinearColor& InColor)
+{
+	const FLinearColor RotateArcballCircleLineColor = 0.25f * Style.RotateArcballCircleColor.CopyWithNewOpacity(1.0f);
+
+	UGizmoElementCircle* CircleElement = MakeRotateCircleHandle(ETransformGizmoPartIdentifier::RotateArcball, Style.RotateArcballSphereRadius, Style.RotateArcballCircleColor,
+		EGizmoElementDrawType::FillAndLine, RotateArcballCircleLineColor);
+
+	// Effectively hides the line when not hovering/interacting
+	CircleElement->SetLineColor(FLinearColor::Transparent);
+
+	return CircleElement;
+}
+
+UGizmoElementCircle* UTransformGizmo::MakeRotateCircleHandle(ETransformGizmoPartIdentifier InPartId, float InRadius, const FLinearColor& InColor,
+	const EGizmoElementDrawType InDrawType, const TOptional<FLinearColor>& InLineColorOverride)
+{
 	UGizmoElementCircle* CircleElement = NewObject<UGizmoElementCircle>();
 	CircleElement->SetPartIdentifier(static_cast<uint32>(InPartId));
 	CircleElement->SetCenter(FVector::ZeroVector);
-	CircleElement->SetRadius(InRadius * SizeCoeff);
 	CircleElement->SetAxis0(FVector::UpVector);
 	CircleElement->SetAxis1(-FVector::RightVector);
 	CircleElement->SetLineColor(InColor);
 	CircleElement->SetViewAlignType(EGizmoElementViewAlignType::PointOnly);
 	CircleElement->SetViewAlignNormal(-FVector::ForwardVector);
 
-	if (bFill)
+	// Disable Draw/Hit by default, and toggle below based on the provided flags
+	CircleElement->SetDrawMesh(false);
+	CircleElement->SetHitMesh(false);
+	CircleElement->SetDrawLine(false);
+	CircleElement->SetHitLine(false);
+
+	if (EnumHasAnyFlags(InDrawType, EGizmoElementDrawType::Fill))
 	{
-		const FLinearColor LightColor(InColor.R, InColor.G, InColor.B, InColor.A * 0.5f);
-		CircleElement->SetVertexColor(LightColor);
+		CircleElement->SetDrawMesh(true);
+		CircleElement->SetHitMesh(true);
+
+		CircleElement->SetVertexColor(FLinearColor::Transparent);
 		CircleElement->SetMaterial(TransparentVertexColorMaterial);
 
-		CircleElement->SetHoverVertexColor(InColor);
+		const FLinearColor LightColor(InColor.R, InColor.G, InColor.B, 0.15f);
+
+		CircleElement->SetHoverVertexColor(LightColor);
 		CircleElement->SetHoverMaterial(TransparentVertexColorMaterial);
-		
-		CircleElement->SetInteractVertexColor(InColor);
+
+		CircleElement->SetInteractVertexColor(LightColor);
 		CircleElement->SetInteractMaterial(TransparentVertexColorMaterial);
 	}
-	else
+
+	if (EnumHasAnyFlags(InDrawType, EGizmoElementDrawType::Line))
 	{
 		CircleElement->SetDrawLine(true);
 		CircleElement->SetHitLine(true);
-		CircleElement->SetDrawMesh(false);
-		CircleElement->SetHitMesh(false);
+
+		if (InLineColorOverride.IsSet())
+		{
+			const FLinearColor LineColor = InLineColorOverride.GetValue();
+
+			CircleElement->SetLineColor(LineColor);
+			CircleElement->SetHoverLineColor(LineColor);
+			CircleElement->SetInteractLineColor(LineColor);
+		}
 	}
 
+	UpdateRotateCircleHandle(CircleElement, InRadius);
+
 	return CircleElement;
 }
 
+void UTransformGizmo::UpdateRotateCircleHandle(UGizmoElementCircle* InElement, float InRadius)
+{
+	if (!InElement)
+	{
+		return;
+	}
+
+	const float SizeCoeff = GetSizeCoefficient();
+
+	InElement->SetRadius((InRadius * Style.AxisSizeMultiplier) * SizeCoeff);
+	InElement->SetLineThickness(Style.LineThicknessMultiplier);
+}
+
+void UTransformGizmo::UpdateElements()
+{
+	const FVector XAxis = FVector::XAxisVector;
+	const FVector YAxis = FVector::YAxisVector;
+	const FVector ZAxis = FVector::ZAxisVector;
+
+	UpdateTranslateAxis(TranslateXAxisElement);
+	UpdateTranslateAxis(TranslateYAxisElement);
+	UpdateTranslateAxis(TranslateZAxisElement);
+
+	UpdatePlanarHandle(TranslatePlanarXYElement, XAxis, YAxis);
+	UpdatePlanarHandle(TranslatePlanarYZElement, YAxis, ZAxis);
+	UpdatePlanarHandle(TranslatePlanarXZElement, XAxis, ZAxis);
+
+	UpdateTranslateScreenSpaceHandle(TranslateScreenSpaceElement);
+
+	UpdateScaleAxis(ScaleXAxisElement);
+	UpdateScaleAxis(ScaleYAxisElement);
+	UpdateScaleAxis(ScaleZAxisElement);
+
+	UpdatePlanarHandle(ScalePlanarXYElement, XAxis, YAxis);
+	UpdatePlanarHandle(ScalePlanarYZElement, YAxis, ZAxis);
+	UpdatePlanarHandle(ScalePlanarXZElement, XAxis, ZAxis);
+
+	UpdateUniformScaleHandle(ScaleUniformElement);
+
+	UpdateRotateAxis(RotateXAxisElement);
+	UpdateRotateAxis(RotateYAxisElement);
+	UpdateRotateAxis(RotateZAxisElement);
+
+	UpdateRotateCircleHandle(RotateArcballElement, Style.RotateArcballSphereRadius);
+	UpdateRotateCircleHandle(RotateScreenSpaceElement, Style.RotateCircleRadius + Style.RotateScreenSpaceRadiusOffset);
+}
 
 void UTransformGizmo::ClearActiveTarget()
 {
@@ -1647,6 +1958,19 @@ void UTransformGizmo::UpdateHoverState(const bool bInHover, const ETransformGizm
 
 	switch (InHitPartId)
 	{
+	case ETransformGizmoPartIdentifier::ScaleXAxis:
+		HitTarget->UpdateHoverState(!bInHover, static_cast<uint32>(ETransformGizmoPartIdentifier::ScaleYAxis));
+		HitTarget->UpdateHoverState(!bInHover, static_cast<uint32>(ETransformGizmoPartIdentifier::ScaleZAxis));
+		break;
+	case ETransformGizmoPartIdentifier::ScaleYAxis:
+		HitTarget->UpdateHoverState(!bInHover, static_cast<uint32>(ETransformGizmoPartIdentifier::ScaleXAxis));
+		HitTarget->UpdateHoverState(!bInHover, static_cast<uint32>(ETransformGizmoPartIdentifier::ScaleZAxis));
+		break;
+	case ETransformGizmoPartIdentifier::ScaleZAxis:
+		HitTarget->UpdateHoverState(!bInHover, static_cast<uint32>(ETransformGizmoPartIdentifier::ScaleXAxis));
+		HitTarget->UpdateHoverState(!bInHover, static_cast<uint32>(ETransformGizmoPartIdentifier::ScaleYAxis));
+		break;
+
 	case ETransformGizmoPartIdentifier::ScaleUniform:
 		HitTarget->UpdateHoverState(bInHover, static_cast<uint32>(ETransformGizmoPartIdentifier::ScaleXAxis));
 		HitTarget->UpdateHoverState(bInHover, static_cast<uint32>(ETransformGizmoPartIdentifier::ScaleYAxis));
@@ -1707,6 +2031,13 @@ void UTransformGizmo::UpdateInteractingState(const bool bInInteracting, const ET
 	{
 		switch (InHitPartId)
 		{
+		case ETransformGizmoPartIdentifier::RotateArcball:
+			{
+				HitTarget->UpdateInteractingState(bInInteracting, static_cast<uint32>(ETransformGizmoPartIdentifier::RotateArcball));
+				GizmoElementRoot->UpdatePartVisibleState(!bInInteracting, static_cast<uint32>(ETransformGizmoPartIdentifier::RotateScreenSpace));
+				break;
+			}
+
 		case ETransformGizmoPartIdentifier::ScaleUniform:
 			{
 				HitTarget->UpdateInteractingState(bInInteracting, static_cast<uint32>(ETransformGizmoPartIdentifier::ScaleXAxis));
@@ -2391,6 +2722,7 @@ void UTransformGizmo::OnClickPressScale(const FInputDeviceRay& PressPos)
 			DebugData.InteractionStart.SetLocation(HitPoint);
 			DebugData.InteractionCurrent.SetLocation(HitPoint);
 			DebugData.InteractionPlaneNormal = InteractionPlanarNormal;
+			DebugData.InteractionScreenDirection = InteractionScreenAxisDirection;
 		}
 	}
 
@@ -2457,16 +2789,22 @@ void UTransformGizmo::OnClickDragScale(const FInputDeviceRay& DragPos)
 
 void UTransformGizmo::OnClickReleaseScaleAxis(const FInputDeviceRay& InReleasePos)
 {
+	UpdateAllScaleAxis();
+
 	bInInteraction = false;
 }
 
 void UTransformGizmo::OnClickReleaseScalePlanar(const FInputDeviceRay& InReleasePos)
 {
+	UpdateAllScaleAxis();
+
 	bInInteraction = false;
 }
 
 void UTransformGizmo::OnClickReleaseScaleXYZ(const FInputDeviceRay& InReleasePos)
 {
+	UpdateAllScaleAxis();
+
 	bInInteraction = false;
 }
 
@@ -2651,7 +2989,7 @@ FVector2D UTransformGizmo::GetWorldToScreenRotateAxisDir(const FInputDeviceRay&
 	}
 
 	// compute nearest point
-	const float Radius = GetWorldRadius(RotateAxisOuterRadius);
+	const float Radius = GetWorldRadius(Style.RotateAxisRadius + Style.RotateAxisOuterRadiusOffset);
 	FVector ClosestPointOnCircle;
 	GizmoMath::ClosetPointOnCircle(QueryPoint, WorldOrigin, InWorldAxis, Radius, ClosestPointOnCircle);
 
@@ -2916,7 +3254,7 @@ void UTransformGizmo::OnClickPressGimbalRotateAxis(const FInputDeviceRay& InPres
 	}
 
 	bTrySwitchingToNormalPull = bIndirectManipulation && RotateMode == EAxisRotateMode::Pull;
-	
+
 	bInInteraction = true;
 	SetModeLastHitPart(EGizmoTransformMode::Rotate, LastHitPart);
 }
@@ -3193,7 +3531,7 @@ void UTransformGizmo::OnClickPressArcBallRotate(const FInputDeviceRay& PressPos)
 	
 	const FVector& RayOrigin = PressPos.WorldRay.Origin;
 	const FVector& RayDir = PressPos.WorldRay.Direction;		
-	const double SphereRadius = GetWorldRadius(RotateArcballSphereRadius);
+	const double SphereRadius = GetWorldRadius(Style.RotateArcballSphereRadius);
 
 	StartRotation = CurrentRotation = CurrentTransform.GetRotation();
 	InteractionPlanarOrigin = CurrentTransform.GetLocation();
@@ -3251,7 +3589,7 @@ void UTransformGizmo::OnClickDragArcBallRotate(const FInputDeviceRay& DragPos)
 {
 	const FVector& RayOrigin = DragPos.WorldRay.Origin;
 	const FVector& RayDir = DragPos.WorldRay.Direction;		
-	const float SphereRadius = GetWorldRadius(RotateArcballSphereRadius);
+	const float SphereRadius = GetWorldRadius(Style.RotateArcballSphereRadius);
 
 	// compute projection
 	ArcBallLocals::GetSphereAndHyperbolicProjection(
diff --git a/Engine/Source/Editor/Experimental/EditorInteractiveToolsFramework/Public/EditorGizmos/TransformGizmo.h b/Engine/Source/Editor/Experimental/EditorInteractiveToolsFramework/Public/EditorGizmos/TransformGizmo.h
--- a/Engine/Source/Editor/Experimental/EditorInteractiveToolsFramework/Public/EditorGizmos/TransformGizmo.h
+++ b/Engine/Source/Editor/Experimental/EditorInteractiveToolsFramework/Public/EditorGizmos/TransformGizmo.h
@@ -8,6 +8,7 @@
 #include "BaseGizmos/TransformProxy.h"
 #include "Containers/EnumAsByte.h"
 #include "CoreMinimal.h"
+#include "BaseGizmos/GizmoElementShared.h"
 #include "EditorGizmos/TransformGizmoInterfaces.h"
 #include "InputState.h"
 #include "InteractiveGizmo.h"
@@ -93,7 +94,7 @@ enum class ETransformGizmoPartIdentifier
 	ScaleXYPlanar,
 	ScaleYZPlanar,
 	ScaleXZPlanar,
-	ScaleUniform, 
+	ScaleUniform,
 	Max
 };
 
@@ -107,6 +108,25 @@ namespace EAxisRotateMode
 	};
 }
 
+/** @todo: To be removed for release, for dev/tweaking only. */
+USTRUCT()
+struct FGizmoDevParameters
+{
+	GENERATED_BODY()
+
+	UPROPERTY(EditAnywhere, Category = "Transform Gizmo", meta = (ClampMin = 0.01f))
+	float LineThicknessMultiplier = 1.0f;
+
+	UPROPERTY(EditAnywhere, Category = "Transform Gizmo", meta = (ClampMin = 1.0f))
+	float TranslateAxisLength = 70.0f;
+
+	UPROPERTY(EditAnywhere, Category = "Transform Gizmo", meta = (ClampMin = 1.0f))
+	float ScaleAxisLength = 70.0f;
+
+	UPROPERTY(EditAnywhere, Category = "Transform Gizmo", meta = (ClampMin = 1.0f))
+	float RotateAxisRadius = 70.0f;
+};
+
 USTRUCT()
 struct FGizmosParameters
 {
@@ -129,6 +149,13 @@ struct FGizmosParameters
 	 */
 	UPROPERTY(EditAnywhere, Category = "Experimental")
 	bool bEnableExplicit = false;
+
+	/** Multiplies the values of TranslateAxisLength, ScaleAxisLength, and RotateAxisRadius. */
+	UPROPERTY()
+	float AxisSizeMultiplier = 1.0f;
+
+	UPROPERTY(EditAnywhere, Category = "Transform Gizmo|Development", meta=(ShowOnlyInnerProperties))
+	FGizmoDevParameters DevParameters;
 };
 
 /**
@@ -143,50 +170,67 @@ class EDITORINTERACTIVETOOLSFRAMEWORK_API UTransformGizmo : public UInteractiveG
 
 public:
 
-	static constexpr float AxisRadius = 1.5f;
-	static constexpr float AxisLengthOffset = 20.0f;
-
-	static constexpr float TranslateAxisLength = 70.0f;
-	static constexpr float TranslateAxisConeAngle = 16.0f;
-	static constexpr float TranslateAxisConeHeight = 22.0f;
-	static constexpr float TranslateAxisConeRadius = 7.0f;
-	static constexpr float TranslateScreenSpaceHandleSize = 14.0f;
-
-	// Rotate constants
-	static constexpr float RotateArcballSphereRadius = 70.0f;
-	static constexpr float RotateAxisOuterRadius = 73.0f;
-	static constexpr float RotateAxisInnerRadius = 1.25f;
-	static constexpr int32 RotateAxisNumSegments = 64;
-	static constexpr int32 RotateAxisInnerSlices = 8;
-	static constexpr float RotateOuterCircleRadius = 73.0f;
-	static constexpr float RotateScreenSpaceRadius = 83.0f;
-
-	static constexpr float ScaleAxisLength = 70.0f;
-	static constexpr float ScaleAxisCubeDim = 12.0f;
-
-	static constexpr float PlanarHandleOffset = 55.0f;
-	static constexpr float PlanarHandleSize = 15.0f;
-
-	static constexpr float AxisTransp = 0.8f;
-	UE_DEPRECATED(5.6, "Use AxisDisplayInfo::GetAxisColor(EAxisList::X) instead")
-	static constexpr FLinearColor AxisColorX = FLinearColor(0.594f, 0.0197f, 0.0f);
-	UE_DEPRECATED(5.6, "Use AxisDisplayInfo::GetAxisColor(EAxisList::Y) instead")
-	static constexpr FLinearColor AxisColorY = FLinearColor(0.1349f, 0.3959f, 0.0f);
-	UE_DEPRECATED(5.6, "Use AxisDisplayInfo::GetAxisColor(EAxisList::Z) instead")
-	static constexpr FLinearColor AxisColorZ = FLinearColor(0.0251f, 0.207f, 0.85f);
-	static constexpr FLinearColor ScreenAxisColor = FLinearColor(0.76f, 0.72f, 0.14f);
-	static constexpr FLinearColor PlaneColorXY = FLinearColor(1.0f, 1.0f, 0.0f);
-	static constexpr FLinearColor ArcBallColor = FLinearColor(0.5f, 0.5f, 0.5f, 0.03f);
-	static constexpr FLinearColor ScreenSpaceColor = FLinearColor(0.765f, 0.765f, 0.765f);
-	static constexpr FLinearColor CurrentColor = FLinearColor(1.0f, 1.0f, 0.0f);
-	static constexpr FLinearColor GreyColor = FLinearColor(0.50f, 0.50f, 0.50f);
-	static constexpr FLinearColor WhiteColor = FLinearColor(1.0f, 1.0f, 1.0f);
-
-	static constexpr FLinearColor RotateScreenSpaceCircleColor = WhiteColor;
-	static constexpr FLinearColor RotateOuterCircleColor = GreyColor;
-	static constexpr FLinearColor RotateArcballCircleColor = FLinearColor(0.50f, 0.50f, 0.50f, 0.1f);
-
-	static constexpr float LargeOuterAlpha = 0.5f;
+	/** Contains all styling properties. */
+	struct FGizmoStyle
+	{
+		/** Applies to axis lengths and radii */
+		float AxisSizeMultiplier = 1.0f;
+
+		/** Applies to lines */
+		float LineThicknessMultiplier = 1.0f;
+
+		static constexpr float AxisRadius = 1.5f;
+		static constexpr float AxisLengthOffset = 20.0f;
+
+		// Transform
+		float TranslateAxisLength = 70.0f;
+		static constexpr float TranslateAxisConeAngle = 16.0f;
+		static constexpr float TranslateAxisConeHeight = 22.0f;
+		static constexpr float TranslateAxisConeRadius = 7.0f;
+		static constexpr float TranslateScreenSpaceHandleSize = 14.0f;
+
+		float RotateArcballSphereRadius = 70.0f;
+		float RotateAxisRadius = 70.0f;
+		float RotateCircleRadius = 70.0f;
+
+		// Rotation
+		static constexpr float RotateAxisInnerRadius = 1.25f;
+		static constexpr int32 RotateAxisNumSegments = 64;
+		static constexpr int32 RotateAxisInnerSlices = 8;
+
+		static constexpr float RotateAxisOuterRadiusOffset = 3.0f;  // Offset from RotateAxisRadius
+		static constexpr float RotateOuterCircleRadiusOffset = 3.0f; // Offset from RotateCircleRadius
+		static constexpr float RotateScreenSpaceRadiusOffset = 13.0f; // Offset from RotateCircleRadius
+
+		// Scale
+		float ScaleAxisLength = 70.0f;
+		static constexpr float ScaleAxisCubeDim = 12.0f;
+
+		static constexpr float PlanarHandleOffset = -15.0f; // Offset from TranslateAxisLength
+		static constexpr float PlanarHandleSize = 15.0f;
+		static constexpr float PlanarHandleThickness = 2.0f;
+
+		static constexpr float AxisTransparency = 0.8f;
+		UE_DEPRECATED(5.6, "Use AxisDisplayInfo::GetAxisColor(EAxisList::X) instead")
+		static constexpr FLinearColor AxisColorX = FLinearColor(0.594f, 0.0197f, 0.0f);
+		UE_DEPRECATED(5.6, "Use AxisDisplayInfo::GetAxisColor(EAxisList::Y) instead")
+		static constexpr FLinearColor AxisColorY = FLinearColor(0.1349f, 0.3959f, 0.0f);
+		UE_DEPRECATED(5.6, "Use AxisDisplayInfo::GetAxisColor(EAxisList::Z) instead")
+		static constexpr FLinearColor AxisColorZ = FLinearColor(0.0251f, 0.207f, 0.85f);
+		static constexpr FLinearColor ScreenAxisColor = FLinearColor(0.76f, 0.72f, 0.14f);
+		static constexpr FLinearColor PlaneColorXY = FLinearColor(1.0f, 1.0f, 0.0f);
+		static constexpr FLinearColor ArcBallColor = FLinearColor(0.5f, 0.5f, 0.5f, 0.03f);
+		static constexpr FLinearColor ScreenSpaceColor = FLinearColor(0.765f, 0.765f, 0.765f);
+		static constexpr FLinearColor CurrentColor = FLinearColor(1.0f, 1.0f, 0.0f);
+		static constexpr FLinearColor GreyColor = FLinearColor(0.50f, 0.50f, 0.50f);
+		static constexpr FLinearColor WhiteColor = FLinearColor(1.0f, 1.0f, 1.0f);
+
+		static constexpr FLinearColor RotateScreenSpaceCircleColor = WhiteColor;
+		static constexpr FLinearColor RotateOuterCircleColor = GreyColor;
+		static constexpr FLinearColor RotateArcballCircleColor = WhiteColor.CopyWithNewOpacity(0.1f);
+
+		static constexpr float LargeOuterAlpha = 0.5f;
+	} Style;
 
 
 public:
@@ -349,15 +393,15 @@ protected:
 
 	/** Translate planar XY handle */
 	UPROPERTY()
-	TObjectPtr<UGizmoElementRectangle> TranslatePlanarXYElement;
+	TObjectPtr<UGizmoElementBox> TranslatePlanarXYElement;
 
 	/** Translate planar YZ handle */
 	UPROPERTY()
-	TObjectPtr<UGizmoElementRectangle> TranslatePlanarYZElement;
+	TObjectPtr<UGizmoElementBox> TranslatePlanarYZElement;
 
 	/** Translate planar XZ handle */
 	UPROPERTY()
-	TObjectPtr<UGizmoElementRectangle> TranslatePlanarXZElement;
+	TObjectPtr<UGizmoElementBox> TranslatePlanarXZElement;
 
 	/** Rotate X Axis */
 	UPROPERTY()
@@ -409,15 +453,15 @@ protected:
 
 	/** Scale planar XY handle */
 	UPROPERTY()
-	TObjectPtr<UGizmoElementRectangle> ScalePlanarXYElement;
+	TObjectPtr<UGizmoElementBox> ScalePlanarXYElement;
 
 	/** Scale planar YZ handle */
 	UPROPERTY()
-	TObjectPtr<UGizmoElementRectangle> ScalePlanarYZElement;
+	TObjectPtr<UGizmoElementBox> ScalePlanarYZElement;
 
 	/** Scale planar XZ handle */
 	UPROPERTY()
-	TObjectPtr<UGizmoElementRectangle> ScalePlanarXZElement;
+	TObjectPtr<UGizmoElementBox> ScalePlanarXZElement;
 
 	/** Uniform scale object */
 	UPROPERTY()
@@ -484,9 +528,32 @@ protected:
 	/** Construct translate axis handle */
 	virtual UGizmoElementArrow* MakeTranslateAxis(ETransformGizmoPartIdentifier InPartId, const FVector& InAxisDir, const FVector& InSideDir, UMaterialInterface* InMaterial);
 
+	/** Updates the given translate axis handle with the current parameters (ie. size coefficient) */
+	virtual void UpdateTranslateAxis(UGizmoElementArrow* InElement);
+
+	/** Updates the given translate axis handle with the current parameters (ie. size coefficient) */
+	virtual void UpdateTranslateAxis(UGizmoElementArrow* InElement, const EAxis::Type InAxis);
+
+	/** Updates the given translate axis handle with the current parameters (ie. size coefficient) */
+	virtual void UpdateTranslateAxis(UGizmoElementArrow* InElement, const FVector& InAxisDir, const FVector& InSideDir);
+
 	/** Construct scale axis handle */
 	virtual UGizmoElementArrow* MakeScaleAxis(ETransformGizmoPartIdentifier InPartId, const FVector& InAxisDir, const FVector& InSideDir, UMaterialInterface* InMaterial);
 
+	/**	Updates the given scale axis handle with the current parameters (ie. size coefficient). */
+	virtual void UpdateScaleAxis(UGizmoElementArrow* InElement);
+
+	/**	Updates the given scale axis handle with the current parameters (ie. size coefficient). */
+	virtual void UpdateScaleAxis(UGizmoElementArrow* InElement, const EAxis::Type InAxis);
+
+	/**
+	 * Updates the given scale axis handle with the current parameters (ie. size coefficient).
+	 */
+	virtual void UpdateScaleAxis(UGizmoElementArrow* InElement, const FVector& InAxisDir, const FVector& InSideDir);
+
+	/** Updates all scale axis handles with the current parameters (ie. size coefficient) */
+	virtual void UpdateAllScaleAxis();
+
 	/** Construct the default rotate axis handle based on the ETransformGizmoPartIdentifier. */
 	UGizmoElementTorus* MakeDefaultRotateAxis(const ETransformGizmoPartIdentifier InPartId);
 
@@ -494,18 +561,40 @@ protected:
 	virtual UGizmoElementTorus* MakeRotateAxis(ETransformGizmoPartIdentifier InPartId, const FVector& TorusAxis0, const FVector& TorusAxis1,
 		UMaterialInterface* InMaterial, UMaterialInterface* InCurrentMaterial);
 
+	/** Updates the given rotate axis handle with the current parameters (ie. size coefficient) */
+	virtual void UpdateRotateAxis(UGizmoElementTorus* InElement);
+
 	/** Construct uniform scale handle */
 	virtual UGizmoElementBox* MakeUniformScaleHandle();
 
+	/** Updates the given uniform scale handle with the current parameters (ie. size coefficient) */
+	virtual void UpdateUniformScaleHandle(UGizmoElementBox* InElement);
+
 	/** Construct planar axis handle */
-	virtual UGizmoElementRectangle* MakePlanarHandle(ETransformGizmoPartIdentifier InPartId, const FVector& InUpDirection, const FVector& InSideDirection, const FVector& InPlaneNormal,
-		UMaterialInterface* InMaterial, const FLinearColor& InVertexColor);
+	virtual UGizmoElementBox* MakePlanarHandle(ETransformGizmoPartIdentifier InPartId, const FVector& InUpDirection, const FVector& InSideDirection, const FVector& InPlaneNormal,
+		UMaterialInterface* InMaterial);
+
+	/** Updates the given uniform scale handle with the current parameters (ie. size coefficient) */
+	virtual void UpdatePlanarHandle(UGizmoElementBox* InElement, const FVector& InUpDirection, const FVector& InSideDirection);
 
 	/** Construct translate screen space handle */
 	virtual UGizmoElementRectangle* MakeTranslateScreenSpaceHandle();
 
+	/** Updates the given translate screen space handle with the current parameters (ie. size coefficient) */
+	virtual void UpdateTranslateScreenSpaceHandle(UGizmoElementRectangle* InElement);
+
+	/** Construct arcball screen space handle */
+	virtual UGizmoElementCircle* MakeArcballCircleHandle(ETransformGizmoPartIdentifier InPartId, float InRadius, const FLinearColor& InColor);
+
 	/** Construct rotate screen space handle */
-	virtual UGizmoElementCircle* MakeRotateCircleHandle(ETransformGizmoPartIdentifier InPartId, float InRadius, const FLinearColor& InColor, const bool bFill);
+	virtual UGizmoElementCircle* MakeRotateCircleHandle(ETransformGizmoPartIdentifier InPartId, float InRadius, const FLinearColor& InColor,
+		const EGizmoElementDrawType InDrawType = EGizmoElementDrawType::Fill, const TOptional<FLinearColor>& InLineColorOverride = TOptional<FLinearColor>());
+
+	/** Updates the given rotate screen space handle with the current parameters (ie. size coefficient) */
+	virtual void UpdateRotateCircleHandle(UGizmoElementCircle* InElement, float InRadius);
+
+	/** Updates all gizmo elements based on current parameters */
+	virtual void UpdateElements();
 
 	/** Get gizmo transform based on cached current transform. */
 	virtual FTransform GetGizmoTransform() const;
@@ -990,6 +1079,9 @@ private:
 		FTransform InteractionCurrent;
 		FVector InteractionPlaneNormal;
 
+		/** Can indicate a 2D drag direction, etc. */
+		FVector2D InteractionScreenDirection;
+
 		/** Debug attributes to display the pull direction */
 		bool bDebugRotate = false;
 		FVector DebugDirection = FVector::ZeroVector;
diff --git a/Engine/Source/Runtime/InteractiveToolsFramework/Public/BaseGizmos/GizmoElementShared.h b/Engine/Source/Runtime/InteractiveToolsFramework/Public/BaseGizmos/GizmoElementShared.h
--- a/Engine/Source/Runtime/InteractiveToolsFramework/Public/BaseGizmos/GizmoElementShared.h
+++ b/Engine/Source/Runtime/InteractiveToolsFramework/Public/BaseGizmos/GizmoElementShared.h
@@ -2,7 +2,8 @@
 
 #pragma once
 
-#include "CoreMinimal.h"
+#include "CoreTypes.h"
+
 #include "GizmoElementShared.generated.h"
 
 //
@@ -77,3 +78,15 @@ enum class EGizmoElementPartialType
 	PartialViewDependent
 };
 
+// Can be used to specify what visuals should be created for a given element, when multiple are available, ie. Fill and Line
+UENUM()
+enum class EGizmoElementDrawType : uint8
+{
+	None = 0x00,
+
+	Fill = 1 << 1,
+	Line = 1 << 2,
+
+	FillAndLine = Fill | Line
+};
+ENUM_CLASS_FLAGS(EGizmoElementDrawType)
