[Viewport Toolbar + New TRS Gizmos] Dynamically update coord space menu to add/remove explicit coordinates

diff --git a/Engine/Source/Editor/Experimental/EditorInteractiveToolsFramework/Private/EditorInteractiveGizmoManager.cpp b/Engine/Source/Editor/Experimental/EditorInteractiveToolsFramework/Private/EditorInteractiveGizmoManager.cpp
--- a/Engine/Source/Editor/Experimental/EditorInteractiveToolsFramework/Private/EditorInteractiveGizmoManager.cpp
+++ b/Engine/Source/Editor/Experimental/EditorInteractiveToolsFramework/Private/EditorInteractiveGizmoManager.cpp
@@ -80,7 +80,7 @@ const TOptional<FGizmosParameters>& UEditorInteractiveGizmoManager::GetDefaultGi
 	if (const UTransformGizmoEditorSettings* TransformGizmoEditorSettings = GetDefault<UTransformGizmoEditorSettings>())
 	{
 		GizmoManagerLocals::OptDefaultParameters = TransformGizmoEditorSettings->GizmosParameters;
-	}\
+	}
 
 	static const TOptional<FGizmosParameters> Invalid;
 	return GizmoManagerLocals::OptDefaultParameters.IsSet() ? GizmoManagerLocals::OptDefaultParameters : Invalid;
@@ -88,7 +88,7 @@ const TOptional<FGizmosParameters>& UEditorInteractiveGizmoManager::GetDefaultGi
 
 bool UEditorInteractiveGizmoManager::IsExplicitModeEnabled()
 {
-	return GizmoManagerLocals::OptDefaultParameters.IsSet() ? GizmoManagerLocals::OptDefaultParameters->bEnableExplicit : false;
+	return GetDefaultGizmosParameters().IsSet() ? GetDefaultGizmosParameters()->bEnableExplicit : false;
 }
 
 UEditorInteractiveGizmoManager::UEditorInteractiveGizmoManager() :
