[ITF] Added GizmoElementText, for screen-space labels

diff --git a/Engine/Source/Editor/Experimental/EditorInteractiveToolsFramework/Private/EditorGizmos/GizmoElementText.cpp b/Engine/Source/Editor/Experimental/EditorInteractiveToolsFramework/Private/EditorGizmos/GizmoElementText.cpp
new file mode 100644
--- /dev/null
+++ b/Engine/Source/Editor/Experimental/EditorInteractiveToolsFramework/Private/EditorGizmos/GizmoElementText.cpp
@@ -0,0 +1,172 @@
+// Copyright Epic Games, Inc. All Rights Reserved.
+
+#include "EditorGizmos/GizmoElementText.h"
+
+#include "CanvasItem.h"
+#include "CanvasTypes.h"
+
+#include "Engine/Engine.h"
+
+void UGizmoElementText::Render(IToolsContextRenderAPI* RenderAPI, const FRenderTraversalState& RenderState)
+{
+	// Empty, work done in DrawHUD
+}
+
+void UGizmoElementText::DrawHUD(FCanvas* Canvas, IToolsContextRenderAPI* RenderAPI, const FRenderTraversalState& RenderState)
+{
+	if (!Canvas)
+	{
+		return;
+	}
+
+	if (Text.IsEmpty())
+	{
+		return;
+	}
+
+	FRenderTraversalState CurrentRenderState(RenderState);
+	bool bVisibleViewDependent = UpdateRenderState(RenderAPI, Center, CurrentRenderState);
+
+	if (bVisibleViewDependent)
+	{
+		const FVector WorldCenter = CurrentRenderState.LocalToWorldTransform.TransformPosition(FVector::ZeroVector);
+
+		FVector4 ScreenCenter = RenderAPI->GetSceneView()->WorldToScreen(WorldCenter);
+		FVector2D ScreenCenter2D = FVector2D(ScreenCenter.X, ScreenCenter.Y);
+
+		RenderAPI->GetSceneView()->ScreenToPixel(ScreenCenter, ScreenCenter2D);
+
+		FCanvasTextItem TextItem(ScreenCenter2D, Text, GetFont(), CurrentRenderState.GetCurrentVertexColor());
+		if (bUseShadow)
+		{
+			TextItem.EnableShadow(ShadowColor, FVector2D(ShadowOffset));
+		}
+
+		// TextItem.Position
+		TextItem.bCentreX = HorizontalAlignment == EGizmoElementTextHorizontalAlignment::Center;
+		TextItem.bCentreY = VerticalAlignment == EGizmoElementTextVerticalAlignment::Center;
+
+		TextItem.Draw(Canvas);
+
+		LastDrawnSize = FVector2f(TextItem.DrawnSize);
+	}
+}
+
+FInputRayHit UGizmoElementText::LineTrace(const UGizmoViewContext* ViewContext, const FLineTraceTraversalState& LineTraceState, const FVector& RayOrigin, const FVector& RayDirection)
+{
+	// Text isn't hit-testable
+	return FInputRayHit();
+}
+
+void UGizmoElementText::SetText(const FText& InText)
+{
+	if (!Text.EqualTo(InText))
+	{
+		Text = InText;
+	}
+}
+
+const FText& UGizmoElementText::GetText() const
+{
+	return Text;
+}
+
+void UGizmoElementText::SetFont(UFont* InFont)
+{
+	if (Font != InFont)
+	{
+		Font = InFont;
+	}
+}
+
+const UFont* UGizmoElementText::GetFont() const
+{
+	if (Font)
+	{
+		return Font;
+	}
+
+	return GEngine->GetSmallFont();
+}
+
+void UGizmoElementText::SetHorizontalAlignment(const EGizmoElementTextHorizontalAlignment InHorizontalAlignment)
+{
+	HorizontalAlignment = InHorizontalAlignment;
+}
+
+const EGizmoElementTextHorizontalAlignment UGizmoElementText::GetHorizontalAlignment() const
+{
+	return HorizontalAlignment;
+}
+
+void UGizmoElementText::SetVerticalAlignment(const EGizmoElementTextVerticalAlignment InVerticalAlignment)
+{
+	VerticalAlignment = InVerticalAlignment;
+}
+
+const EGizmoElementTextVerticalAlignment UGizmoElementText::GetVerticalAlignment() const
+{
+	return VerticalAlignment;
+}
+
+void UGizmoElementText::SetCenter(const FVector& InCenter)
+{
+	if (Center != InCenter)
+	{
+		Center = InCenter;
+	}
+}
+
+FVector UGizmoElementText::GetCenter() const
+{
+	return Center;
+}
+
+void UGizmoElementText::SetEnableShadow(const bool bInEnable)
+{
+	if (bUseShadow != bInEnable)
+	{
+		bUseShadow = bInEnable;
+	}
+}
+
+bool UGizmoElementText::GetEnableShadow() const
+{
+	return bUseShadow;
+}
+
+void UGizmoElementText::SetShadowOffset(const FVector2f& InOffset)
+{
+	if (!ShadowOffset.Equals(InOffset))
+	{
+		ShadowOffset = InOffset;
+	}
+}
+
+const FVector2f& UGizmoElementText::GetShadowOffset() const
+{
+	return ShadowOffset;
+}
+
+void UGizmoElementText::SetShadowColor(const FLinearColor& InShadowColor)
+{
+	if (!ShadowColor.Equals(InShadowColor))
+	{
+		ShadowColor = InShadowColor;
+	}
+}
+
+const FLinearColor& UGizmoElementText::GetShadowColor() const
+{
+	return ShadowColor;
+}
+
+FVector2f UGizmoElementText::GetDrawnSize() const
+{
+	return LastDrawnSize;
+}
+
+bool UGizmoElementText::UpdateRenderState(IToolsContextRenderAPI* RenderAPI, const FVector& InLocalCenter, FRenderTraversalState& InOutRenderState)
+{
+	return GetText().IsEmpty() || Super::UpdateRenderState(RenderAPI, InLocalCenter, InOutRenderState);
+}
diff --git a/Engine/Source/Editor/Experimental/EditorInteractiveToolsFramework/Private/EditorGizmos/GizmoElementText.h b/Engine/Source/Editor/Experimental/EditorInteractiveToolsFramework/Private/EditorGizmos/GizmoElementText.h
new file mode 100644
--- /dev/null
+++ b/Engine/Source/Editor/Experimental/EditorInteractiveToolsFramework/Private/EditorGizmos/GizmoElementText.h
@@ -0,0 +1,99 @@
+// Copyright Epic Games, Inc. All Rights Reserved.
+
+#pragma once
+
+#include "BaseGizmos/GizmoElementBase.h"
+
+#include "GizmoElementText.generated.h"
+
+UENUM()
+enum class EGizmoElementTextHorizontalAlignment : uint8
+{
+	Left = 0,
+	Center = 1
+};
+
+UENUM()
+enum class EGizmoElementTextVerticalAlignment : uint8
+{
+	Bottom = 0,
+	Center = 1
+};
+
+/**
+ * Simple object intended to be used as part of 3D Gizmos.
+ * Draws a text label based on parameters.
+ */
+UCLASS(Transient, MinimalAPI)
+class UGizmoElementText : public UGizmoElementBase
+{
+	GENERATED_BODY()
+
+public:
+	//~ Begin UGizmoElementBase Interface.
+	virtual void Render(IToolsContextRenderAPI* RenderAPI, const FRenderTraversalState& RenderState) override;
+	virtual void DrawHUD(FCanvas* Canvas, IToolsContextRenderAPI* RenderAPI, const FRenderTraversalState& RenderState) override;
+	virtual FInputRayHit LineTrace(const UGizmoViewContext* ViewContext, const FLineTraceTraversalState& LineTraceState, const FVector& RayOrigin, const FVector& RayDirection) override;
+	//~ End UGizmoElementBase Interface.
+
+	void SetText(const FText& InText);
+	const FText& GetText() const;
+
+	void SetFont(UFont* InFont);
+	const UFont* GetFont() const;
+
+	void SetHorizontalAlignment(const EGizmoElementTextHorizontalAlignment InHorizontalAlignment);
+	const EGizmoElementTextHorizontalAlignment GetHorizontalAlignment() const;
+
+	void SetVerticalAlignment(const EGizmoElementTextVerticalAlignment InVerticalAlignment);
+	const EGizmoElementTextVerticalAlignment GetVerticalAlignment() const;
+
+	// Location of the text label
+	void SetCenter(const FVector& InCenter);
+	FVector GetCenter() const;
+
+	void SetEnableShadow(const bool bInEnable);
+	bool GetEnableShadow() const;
+
+	void SetShadowOffset(const FVector2f& InOffset);
+	const FVector2f& GetShadowOffset() const;
+
+	void SetShadowColor(const FLinearColor& InShadowColor);
+	const FLinearColor& GetShadowColor() const;
+
+	/** Returns the last drawn size of the text, if any. Useful to calculate the approximate bounds. */
+	FVector2f GetDrawnSize() const;
+
+protected:
+	//~ Begin UGizmoElementBase Interface.
+	virtual bool UpdateRenderState(IToolsContextRenderAPI* RenderAPI, const FVector& InLocalCenter, FRenderTraversalState& InOutRenderState) override;
+	//~ End UGizmoElementBase Interface.
+
+protected:
+	UPROPERTY()
+	FText Text;
+
+	UPROPERTY()
+	TObjectPtr<UFont> Font;
+
+	UPROPERTY()
+	EGizmoElementTextHorizontalAlignment HorizontalAlignment = EGizmoElementTextHorizontalAlignment::Left;
+
+	UPROPERTY()
+	EGizmoElementTextVerticalAlignment VerticalAlignment = EGizmoElementTextVerticalAlignment::Bottom;
+
+	// Location of center of the text label
+	UPROPERTY()
+	FVector Center = FVector::ZeroVector;
+
+	UPROPERTY()
+	bool bUseShadow = true;
+
+	UPROPERTY()
+	FVector2f ShadowOffset = FVector2f(1.0f, 0.0f);
+
+	UPROPERTY()
+	FLinearColor ShadowColor = FLinearColor::Black;
+
+	FVector2f LastDrawnSize = FVector2f::ZeroVector;
+};
diff --git a/Engine/Source/Runtime/InteractiveToolsFramework/Private/BaseGizmos/GizmoElementBase.cpp b/Engine/Source/Runtime/InteractiveToolsFramework/Private/BaseGizmos/GizmoElementBase.cpp
--- a/Engine/Source/Runtime/InteractiveToolsFramework/Private/BaseGizmos/GizmoElementBase.cpp
+++ b/Engine/Source/Runtime/InteractiveToolsFramework/Private/BaseGizmos/GizmoElementBase.cpp
@@ -307,6 +307,11 @@ bool UGizmoElementBase::GetVisibleState() const
 	return (static_cast<uint8>(ElementState) & static_cast<uint8>(EGizmoElementState::Visible));
 }
 
+void UGizmoElementBase::DrawHUD(FCanvas* Canvas, IToolsContextRenderAPI* RenderAPI, const FRenderTraversalState& RenderState)
+{
+	// Empty default implementation
+}
+
 void UGizmoElementBase::SetVisibleState(bool bVisible)
 {
 	ElementState = (bVisible ? ElementState | EGizmoElementState::Visible : ElementState & (~EGizmoElementState::Visible));
@@ -682,4 +687,3 @@ void UGizmoElementBase::ClearInteractVertexColor()
 {
 	MeshRenderAttributes.InteractVertexColor.Reset();
 }
-
diff --git a/Engine/Source/Runtime/InteractiveToolsFramework/Private/BaseGizmos/GizmoElementGroup.cpp b/Engine/Source/Runtime/InteractiveToolsFramework/Private/BaseGizmos/GizmoElementGroup.cpp
--- a/Engine/Source/Runtime/InteractiveToolsFramework/Private/BaseGizmos/GizmoElementGroup.cpp
+++ b/Engine/Source/Runtime/InteractiveToolsFramework/Private/BaseGizmos/GizmoElementGroup.cpp
@@ -36,6 +36,27 @@ void UGizmoElementGroupBase::Render(IToolsContextRenderAPI* RenderAPI, const FRe
 	}
 }
 
+void UGizmoElementGroupBase::DrawHUD(FCanvas* Canvas, IToolsContextRenderAPI* RenderAPI, const FRenderTraversalState& RenderState)
+{
+	FRenderTraversalState CurrentRenderState(RenderState);
+	bool bVisibleViewDependent = UpdateRenderState(RenderAPI, FVector::ZeroVector, CurrentRenderState);
+
+	if (bVisibleViewDependent)
+	{
+		ApplyUniformConstantScaleToTransform(CurrentRenderState.PixelToWorldScale, CurrentRenderState.LocalToWorldTransform);
+
+		// Continue render even if not visible so all transforms will be cached
+		// for subsequent line tracing.
+		for (UGizmoElementBase* Element : Elements)
+		{
+			if (Element)
+			{
+				Element->DrawHUD(Canvas, RenderAPI, CurrentRenderState);
+			}
+		}
+	}
+}
+
 FInputRayHit UGizmoElementGroupBase::LineTrace(const UGizmoViewContext* ViewContext, const FLineTraceTraversalState& LineTraceState, const FVector& RayOrigin, const FVector& RayDirection)
 {
 	FInputRayHit Hit;
diff --git a/Engine/Source/Runtime/InteractiveToolsFramework/Public/BaseGizmos/GizmoElementBase.h b/Engine/Source/Runtime/InteractiveToolsFramework/Public/BaseGizmos/GizmoElementBase.h
--- a/Engine/Source/Runtime/InteractiveToolsFramework/Public/BaseGizmos/GizmoElementBase.h
+++ b/Engine/Source/Runtime/InteractiveToolsFramework/Public/BaseGizmos/GizmoElementBase.h
@@ -116,6 +116,9 @@ public:
 	// Render enabled visible element.
 	INTERACTIVETOOLSFRAMEWORK_API virtual void Render(IToolsContextRenderAPI* RenderAPI, const FRenderTraversalState& RenderState) PURE_VIRTUAL(UGizmoElementBase::Render);
 
+	// Do screen space drawing.
+	INTERACTIVETOOLSFRAMEWORK_API virtual void DrawHUD(FCanvas* Canvas, IToolsContextRenderAPI* RenderAPI, const FRenderTraversalState& RenderState);
+
 	// Line trace enabled hittable element.
 	INTERACTIVETOOLSFRAMEWORK_API virtual FInputRayHit LineTrace(const UGizmoViewContext* ViewContext, const FLineTraceTraversalState& LineTraceState, const FVector& RayOrigin, const FVector& RayDirection) PURE_VIRTUAL(UGizmoElementBase::LineTrace, return FInputRayHit(););
 
diff --git a/Engine/Source/Runtime/InteractiveToolsFramework/Public/BaseGizmos/GizmoElementGroup.h b/Engine/Source/Runtime/InteractiveToolsFramework/Public/BaseGizmos/GizmoElementGroup.h
--- a/Engine/Source/Runtime/InteractiveToolsFramework/Public/BaseGizmos/GizmoElementGroup.h
+++ b/Engine/Source/Runtime/InteractiveToolsFramework/Public/BaseGizmos/GizmoElementGroup.h
@@ -21,6 +21,7 @@ public:
 
 	//~ Begin UGizmoElementBase Interface.
 	INTERACTIVETOOLSFRAMEWORK_API virtual void Render(IToolsContextRenderAPI* RenderAPI, const FRenderTraversalState& RenderState) override;
+	INTERACTIVETOOLSFRAMEWORK_API virtual void DrawHUD(FCanvas* Canvas, IToolsContextRenderAPI* RenderAPI, const FRenderTraversalState& RenderState) override;
 	INTERACTIVETOOLSFRAMEWORK_API virtual FInputRayHit LineTrace(const UGizmoViewContext* ViewContext, const FLineTraceTraversalState& LineTraceState, const FVector& RayOrigin, const FVector& RayDirection) override;
 	//~ End UGizmoElementBase Interface.
 
