Marquee selection: fixed frustum computation

frustum planes can easily be invalid when calculated from a vertical/horizontal line with a low near-plane value
the tolerance is lowered when calling GetSafeNormal during the plane computation to avoid that

#rnx

diff --git a/Engine/Source/Editor/Experimental/EditorInteractiveToolsFramework/Private/EditorDragTools/EditorFrustumSelect.cpp b/Engine/Source/Editor/Experimental/EditorInteractiveToolsFramework/Private/EditorDragTools/EditorFrustumSelect.cpp
--- a/Engine/Source/Editor/Experimental/EditorInteractiveToolsFramework/Private/EditorDragTools/EditorFrustumSelect.cpp
+++ b/Engine/Source/Editor/Experimental/EditorInteractiveToolsFramework/Private/EditorDragTools/EditorFrustumSelect.cpp
@@ -25,31 +25,117 @@
 #include "Subsystems/BrushEditingSubsystem.h"
 #include "UnrealEdGlobals.h"
 
-namespace UE::EditorDragTools::Private
+namespace UE::EditorDragTools
 {
-TArray<FTypedElementHandle> GetElementsIntersectingFrustum(
-	const AActor* Actor,
-	const FConvexVolume& InFrustum,
-	const IEditorViewportClientProxy* InEditorViewportClientProxy,
-	const FWorldSelectionElementArgs& SelectionArgs
-)
+
+namespace Private
 {
-	if (InEditorViewportClientProxy && InEditorViewportClientProxy->IsActorVisible(Actor))
+	TArray<FTypedElementHandle> GetElementsIntersectingFrustum(
+		const AActor* Actor,
+		const FConvexVolume& InFrustum,
+		const IEditorViewportClientProxy* InEditorViewportClientProxy,
+		const FWorldSelectionElementArgs& SelectionArgs
+	)
 	{
-		if (FTypedElementHandle ActorHandle = UEngineElementsLibrary::AcquireEditorActorElementHandle(Actor))
+		if (InEditorViewportClientProxy && InEditorViewportClientProxy->IsActorVisible(Actor))
 		{
-			if (TTypedElement<ITypedElementWorldInterface> WorldElement =
-					UTypedElementRegistry::GetInstance()->GetElement<ITypedElementWorldInterface>(ActorHandle))
+			if (FTypedElementHandle ActorHandle = UEngineElementsLibrary::AcquireEditorActorElementHandle(Actor))
 			{
-				return WorldElement.GetSelectionElementsInConvexVolume(InFrustum, SelectionArgs);
+				if (TTypedElement<ITypedElementWorldInterface> WorldElement =
+						UTypedElementRegistry::GetInstance()->GetElement<ITypedElementWorldInterface>(ActorHandle))
+				{
+					return WorldElement.GetSelectionElementsInConvexVolume(InFrustum, SelectionArgs);
+				}
 			}
 		}
+
+		return {};
 	}
+}
+
+FViewportFrustum::FViewportFrustum(const FEditorViewportClient& InViewportClient, const FSceneView& InView,
+	const FVector& InStart, const FVector& InEnd, const bool InUseBoxFrustum)
+	: ViewportClient(InViewportClient)
+	, View(InView)
+	, Start(InStart)
+	, End(InEnd) 
+	, bUseBoxFrustum(InUseBoxFrustum)
+{}
+	
+void FViewportFrustum::Calculate(FConvexVolume& OutFrustum) const
+{
+	if (bUseBoxFrustum)
+	{
+		const FVector& CamPoint = ViewportClient.GetViewLocation();
+
+		// extend the 2D box of 1 pixel if needed to avoid degenerated volume 
+		const double DX = Start.X == End.X ? 0.5 : 0.0;
+		const double Left = FMath::Min(Start.X, End.X) - DX;
+		const double Right = FMath::Max(Start.X, End.X) + DX;
+		const double DY = Start.Y == End.Y ? 0.5 : 0.0;
+		const double Bottom = FMath::Min(Start.Y, End.Y) - DY;
+		const double Top = FMath::Max(Start.Y, End.Y) + DY;
+		
+		// Deproject the four corners of the selection box		
+		const FVector2D Point1(Left, Bottom); // Upper Left Corner
+		const FVector2D Point2(Right, Bottom); // Upper Right Corner
+		const FVector2D Point3(Right, Top); // Lower Right Corner
+		const FVector2D Point4(Left, Top); // Lower Left Corner
 
-	return {};
+		FVector BoxPoint1, BoxPoint2, BoxPoint3, BoxPoint4;
+		FVector WorldDir1, WorldDir2, WorldDir3, WorldDir4;
+		View.DeprojectFVector2D(Point1, BoxPoint1, WorldDir1);
+		View.DeprojectFVector2D(Point2, BoxPoint2, WorldDir2);
+		View.DeprojectFVector2D(Point3, BoxPoint3, WorldDir3);
+		View.DeprojectFVector2D(Point4, BoxPoint4, WorldDir4);
+
+		// Use the camera position and the selection box to create the bounding planes
+		auto MakePlane = [](const FVector& A, const FVector& B, const FVector& C)
+		{
+			// the default tolerance used in GetSafeNormal (UE_SMALL_NUMBER) is not enough to handle low near plane values
+			// see TPlane<T>::TPlane(TVector<T> A, TVector<T> B, TVector<T> C) 
+			static constexpr double SmallNumberCubed = UE_SMALL_NUMBER * UE_SMALL_NUMBER * UE_SMALL_NUMBER;
+			const FVector Normal = (B - A ^ C - A).GetSafeNormal(SmallNumberCubed);
+			const double W = A | Normal;
+			FPlane Plane(Normal, W);
+			return MoveTemp(Plane);
+		};
+		
+		FPlane TopPlane = MakePlane(BoxPoint1, BoxPoint2, CamPoint); // Top Plane
+		FPlane RightPlane = MakePlane(BoxPoint2, BoxPoint3, CamPoint); // Right Plane
+		FPlane BottomPlane = MakePlane(BoxPoint3, BoxPoint4, CamPoint); // Bottom Plane
+		FPlane LeftPlane = MakePlane(BoxPoint4, BoxPoint1, CamPoint); // Left Plane
+
+		// Try to get all six planes to create a frustum.
+		// The frustum is built with the first four planes corresponding to the sides of the frustum.
+		OutFrustum.Planes.Empty();
+		OutFrustum.Planes.Emplace(MoveTemp(TopPlane));
+		OutFrustum.Planes.Emplace(MoveTemp(RightPlane));
+		OutFrustum.Planes.Emplace(MoveTemp(BottomPlane));
+		OutFrustum.Planes.Emplace(MoveTemp(LeftPlane));
+		
+		FPlane NearPlane;
+		if ( View.ViewMatrices.GetViewProjectionMatrix().GetFrustumNearPlane(NearPlane) )
+		{
+			OutFrustum.Planes.Emplace(MoveTemp(NearPlane));
+		}
+
+		FPlane FarPlane;
+		if ( View.ViewMatrices.GetViewProjectionMatrix().GetFrustumFarPlane(FarPlane) )
+		{
+			OutFrustum.Planes.Emplace(MoveTemp(FarPlane));
+		}
+		
+		OutFrustum.Init();
+	}
+	else
+	{
+		OutFrustum = View.ViewFrustum;
+		OutFrustum.Init();
+	}
 }
 
-} // namespace UE::EditorDragTools::Private
+} // namespace UE::EditorDragTools
 
 void FEditorFrustumSelect::Render(const FSceneView* View, FCanvas* Canvas)
 {
@@ -358,55 +444,6 @@ void FEditorFrustumSelect::CalculateFrustum(const FSceneView* InView, FConvexVol
 		return;
 	}
 
-	if (bUseBoxFrustum)
-	{
-		FVector CamPoint = EditorViewportClient->GetViewLocation();
-		FVector BoxPoint1, BoxPoint2, BoxPoint3, BoxPoint4;
-		FVector WorldDir1, WorldDir2, WorldDir3, WorldDir4;
-		// extend the 2D box of 1 pixel if needed to avoid degenerated volume
-		const double DX = Start.X == End.X ? 0.5 : 0.0;
-		const double Left = FMath::Min(Start.X, End.X) - DX;
-		const double Right = FMath::Max(Start.X, End.X) + DX;
-		const double DY = Start.Y == End.Y ? 0.5 : 0.0;
-		const double Bottom = FMath::Min(Start.Y, End.Y) - DY;
-		const double Top = FMath::Max(Start.Y, End.Y) + DY;
-		// Deproject the four corners of the selection box
-		FVector2D Point1(Left, Bottom); // Upper Left Corner
-		FVector2D Point2(Right, Bottom); // Upper Right Corner
-		FVector2D Point3(Right, Top); // Lower Right Corner
-		FVector2D Point4(Left, Top); // Lower Left Corner
-		InView->DeprojectFVector2D(Point1, BoxPoint1, WorldDir1);
-		InView->DeprojectFVector2D(Point2, BoxPoint2, WorldDir2);
-		InView->DeprojectFVector2D(Point3, BoxPoint3, WorldDir3);
-		InView->DeprojectFVector2D(Point4, BoxPoint4, WorldDir4);
-		// Use the camera position and the selection box to create the bounding planes
-		FPlane TopPlane(BoxPoint1, BoxPoint2, CamPoint); // Top Plane
-		FPlane RightPlane(BoxPoint2, BoxPoint3, CamPoint); // Right Plane
-		FPlane BottomPlane(BoxPoint3, BoxPoint4, CamPoint); // Bottom Plane
-		FPlane LeftPlane(BoxPoint4, BoxPoint1, CamPoint); // Left Plane
-
-		// Try to get all six planes to create a frustum.
-		// The frustum is built with the first four planes corresponding to the sides of the frustum.
-		FPlane NearPlane;
-		FPlane FarPlane;
-		OutFrustum.Planes.Empty();
-		OutFrustum.Planes.Add(TopPlane);
-		OutFrustum.Planes.Add(RightPlane);
-		OutFrustum.Planes.Add(BottomPlane);
-		OutFrustum.Planes.Add(LeftPlane);
-		if (InView->ViewMatrices.GetViewProjectionMatrix().GetFrustumNearPlane(NearPlane))
-		{
-			OutFrustum.Planes.Add(NearPlane);
-		}
-		if (InView->ViewMatrices.GetViewProjectionMatrix().GetFrustumFarPlane(FarPlane))
-		{
-			OutFrustum.Planes.Add(FarPlane);
-		}
-		OutFrustum.Init();
-	}
-	else
-	{
-		OutFrustum = InView->ViewFrustum;
-		OutFrustum.Init();
-	}
+	UE::EditorDragTools::FViewportFrustum ViewportFrustum(*EditorViewportClient, *InView, Start, End, bUseBoxFrustum);
+	ViewportFrustum.Calculate(OutFrustum);
 }
diff --git a/Engine/Source/Editor/Experimental/EditorInteractiveToolsFramework/Public/EditorDragTools/EditorFrustumSelect.h b/Engine/Source/Editor/Experimental/EditorInteractiveToolsFramework/Public/EditorDragTools/EditorFrustumSelect.h
--- a/Engine/Source/Editor/Experimental/EditorInteractiveToolsFramework/Public/EditorDragTools/EditorFrustumSelect.h
+++ b/Engine/Source/Editor/Experimental/EditorInteractiveToolsFramework/Public/EditorDragTools/EditorFrustumSelect.h
@@ -56,3 +56,39 @@ private:
 	 */
 	void CalculateFrustum(const FSceneView* InView, FConvexVolume& OutFrustum, bool bUseBoxFrustum) const;
 };
+
+namespace UE::EditorDragTools
+{
+
+/**
+ * FViewportFrustum calculates a frustum box from 2D start and end positions in a viewport.  
+ */
+
+struct EDITORINTERACTIVETOOLSFRAMEWORK_API FViewportFrustum
+{
+public:
+	/** 
+	 * @param InViewportClient	Information about the viewport client (view location)
+	 * @param InView			Information about scene projection
+	 * @param InStart			The 2D screen space start position of the dragged box
+	 * @param InEnd				The 2D screen space end position of the dragged box
+	 * @param InUseBoxFrustum	If true a frustum out of the current dragged box will be created (false will use the view frustum)
+	 */
+	FViewportFrustum(const FEditorViewportClient& InViewportClient, const FSceneView& InView,
+		const FVector& InStart, const FVector& InEnd, const bool InUseBoxFrustum);
+
+	/** 
+	 * @param OutFrustum		The created frustum
+	 */
+	void Calculate(FConvexVolume& OutFrustum) const;
+
+private:
+
+	const FEditorViewportClient& ViewportClient;
+	const FSceneView& View;
+	const FVector& Start;
+	const FVector& End;
+	bool bUseBoxFrustum = false;
+};
+
+}
\ No newline at end of file
diff --git a/Engine/Source/Editor/UnrealEd/Private/DragTool_FrustumSelect.cpp b/Engine/Source/Editor/UnrealEd/Private/DragTool_FrustumSelect.cpp
--- a/Engine/Source/Editor/UnrealEd/Private/DragTool_FrustumSelect.cpp
+++ b/Engine/Source/Editor/UnrealEd/Private/DragTool_FrustumSelect.cpp
@@ -25,6 +25,7 @@
 #include "Subsystems/BrushEditingSubsystem.h"
 #include "LevelEditorSubsystem.h"
 #include "SceneView.h"
+#include "EditorDragTools/EditorFrustumSelect.h"
 
 ///////////////////////////////////////////////////////////////////////////////
 //
@@ -280,57 +281,13 @@ bool FDragTool_ActorFrustumSelect::IntersectsFrustum( const UModel& InModel, int
 
 void FDragTool_ActorFrustumSelect::CalculateFrustum( const FSceneView* InView, FConvexVolume& OutFrustum, bool bUseBoxFrustum ) const
 {
-	if( bUseBoxFrustum )
+	if (!InView || !EditorViewportClient)
 	{
-		FVector CamPoint = EditorViewportClient->GetViewLocation();
-		FVector BoxPoint1, BoxPoint2, BoxPoint3, BoxPoint4;
-		FVector WorldDir1, WorldDir2, WorldDir3, WorldDir4;
-		// extend the 2D box of 1 pixel if needed to avoid degenerated volume 
-		const double DX = Start.X == End.X ? 0.5 : 0.0;
-		const double Left = FMath::Min(Start.X, End.X) - DX;
-		const double Right = FMath::Max(Start.X, End.X) + DX;
-		const double DY = Start.Y == End.Y ? 0.5 : 0.0;
-		const double Bottom = FMath::Min(Start.Y, End.Y) - DY;
-		const double Top = FMath::Max(Start.Y, End.Y) + DY;
-		// Deproject the four corners of the selection box		
-		FVector2D Point1(Left, Bottom); // Upper Left Corner
-		FVector2D Point2(Right, Bottom); // Upper Right Corner
-		FVector2D Point3(Right, Top); // Lower Right Corner
-		FVector2D Point4(Left, Top); // Lower Left Corner
-		InView->DeprojectFVector2D(Point1, BoxPoint1, WorldDir1);
-		InView->DeprojectFVector2D(Point2, BoxPoint2, WorldDir2);
-		InView->DeprojectFVector2D(Point3, BoxPoint3, WorldDir3);
-		InView->DeprojectFVector2D(Point4, BoxPoint4, WorldDir4);
-		// Use the camera position and the selection box to create the bounding planes
-		FPlane TopPlane(BoxPoint1, BoxPoint2, CamPoint); // Top Plane
-		FPlane RightPlane(BoxPoint2, BoxPoint3, CamPoint); // Right Plane
-		FPlane BottomPlane(BoxPoint3, BoxPoint4, CamPoint); // Bottom Plane
-		FPlane LeftPlane(BoxPoint4, BoxPoint1, CamPoint); // Left Plane
-
-		// Try to get all six planes to create a frustum.
-		// The frustum is built with the first four planes corresponding to the sides of the frustum.
-		FPlane NearPlane;
-		FPlane FarPlane;
-		OutFrustum.Planes.Empty();
-		OutFrustum.Planes.Add(TopPlane);
-		OutFrustum.Planes.Add(RightPlane);
-		OutFrustum.Planes.Add(BottomPlane);
-		OutFrustum.Planes.Add(LeftPlane);
-		if ( InView->ViewMatrices.GetViewProjectionMatrix().GetFrustumNearPlane(NearPlane) )
-		{
-			OutFrustum.Planes.Add(NearPlane);
-		}
-		if ( InView->ViewMatrices.GetViewProjectionMatrix().GetFrustumFarPlane(FarPlane) )
-		{
-			OutFrustum.Planes.Add(FarPlane);
-		}
-		OutFrustum.Init();
-	}
-	else
-	{
-		OutFrustum = InView->ViewFrustum;
-		OutFrustum.Init();
+		return;
 	}
+
+	UE::EditorDragTools::FViewportFrustum ViewportFrustum(*EditorViewportClient, *InView, Start, End, bUseBoxFrustum);
+	ViewportFrustum.Calculate(OutFrustum);
 }
 
 void FDragTool_ActorFrustumSelect::AddHoverEffect( AActor& InActor )
