Updated headers using UnrealCodeFixup to make sure dllstorage is on methods/staticvars instead of types. Part 5/n

diff --git a/Engine/Source/Editor/Experimental/EditorInteractiveToolsFramework/Public/EditorGizmos/TransformGizmo.h b/Engine/Source/Editor/Experimental/EditorInteractiveToolsFramework/Public/EditorGizmos/TransformGizmo.h
--- a/Engine/Source/Editor/Experimental/EditorInteractiveToolsFramework/Public/EditorGizmos/TransformGizmo.h
+++ b/Engine/Source/Editor/Experimental/EditorInteractiveToolsFramework/Public/EditorGizmos/TransformGizmo.h
@@ -30,6 +30,8 @@
 
 #include "TransformGizmo.generated.h"
 
+#define UE_API EDITORINTERACTIVETOOLSFRAMEWORK_API
+
 class IGizmoAxisSource;
 class IGizmoStateTarget;
 class IGizmoTransformSource;
@@ -163,8 +165,8 @@ struct FGizmosParameters
  * applied to a UTransformProxy target object. By default the Gizmo will be
  * a standard XYZ translate/rotate Gizmo (axis and plane translation).
  */
-UCLASS()
-class EDITORINTERACTIVETOOLSFRAMEWORK_API UTransformGizmo : public UInteractiveGizmo, public IHoverBehaviorTarget, public IClickDragBehaviorTarget
+UCLASS(MinimalAPI)
+class UTransformGizmo : public UInteractiveGizmo, public IHoverBehaviorTarget, public IClickDragBehaviorTarget
 {
 	GENERATED_BODY()
 
@@ -241,26 +243,26 @@ public:
 	 *
 	 * TODO: Should this affect uniform scaling too?
 	 */
-	virtual void SetDisallowNegativeScaling(bool bDisallow);
+	UE_API virtual void SetDisallowNegativeScaling(bool bDisallow);
 
 	// UInteractiveGizmo overrides
-	virtual void Setup() override;
-	virtual void Shutdown() override;
-	virtual void Render(IToolsContextRenderAPI* RenderAPI) override;
-	virtual void Tick(float DeltaTime) override;
+	UE_API virtual void Setup() override;
+	UE_API virtual void Shutdown() override;
+	UE_API virtual void Render(IToolsContextRenderAPI* RenderAPI) override;
+	UE_API virtual void Tick(float DeltaTime) override;
 
 	// IHoverBehaviorTarget implementation
-	virtual FInputRayHit BeginHoverSequenceHitTest(const FInputDeviceRay& DevicePos) override;
-	virtual void OnBeginHover(const FInputDeviceRay& DevicePos) override;
-	virtual bool OnUpdateHover(const FInputDeviceRay& DevicePos) override;
-	virtual void OnEndHover() override;
+	UE_API virtual FInputRayHit BeginHoverSequenceHitTest(const FInputDeviceRay& DevicePos) override;
+	UE_API virtual void OnBeginHover(const FInputDeviceRay& DevicePos) override;
+	UE_API virtual bool OnUpdateHover(const FInputDeviceRay& DevicePos) override;
+	UE_API virtual void OnEndHover() override;
 
 	// IClickDragBehaviorTarget implementation
-	virtual FInputRayHit CanBeginClickDragSequence(const FInputDeviceRay& PressPos) override;
-	virtual void OnClickPress(const FInputDeviceRay& PressPos) override;
-	virtual void OnClickDrag(const FInputDeviceRay& DragPos) override;
-	virtual void OnClickRelease(const FInputDeviceRay& ReleasePos) override;
-	virtual void OnTerminateDragSequence() override;
+	UE_API virtual FInputRayHit CanBeginClickDragSequence(const FInputDeviceRay& PressPos) override;
+	UE_API virtual void OnClickPress(const FInputDeviceRay& PressPos) override;
+	UE_API virtual void OnClickDrag(const FInputDeviceRay& DragPos) override;
+	UE_API virtual void OnClickRelease(const FInputDeviceRay& ReleasePos) override;
+	UE_API virtual void OnTerminateDragSequence() override;
 
 	/**
 	 * Set the active target object for the Gizmo
@@ -268,38 +270,38 @@ public:
 	 * @param TransactionProvider optional IToolContextTransactionProvider implementation to use - by default uses GizmoManager
 	 * @param InStateTarget optional IGizmoStateTarget implementation to use - will create a new UGizmoObjectModifyStateTarget none provided
 	 */
-	virtual void SetActiveTarget(UTransformProxy* Target, IToolContextTransactionProvider* TransactionProvider = nullptr, IGizmoStateTarget* InStateTarget = nullptr);
+	UE_API virtual void SetActiveTarget(UTransformProxy* Target, IToolContextTransactionProvider* TransactionProvider = nullptr, IGizmoStateTarget* InStateTarget = nullptr);
 
 	/**
 	 * Clear the active target object for the Gizmo
 	 */
-	virtual void ClearActiveTarget();
+	UE_API virtual void ClearActiveTarget();
 
 	/**
 	 * Explicitly set the child scale. Mainly useful to "reset" the child scale to (1,1,1) when re-using Gizmo across multiple transform actions.
 	 * @warning does not generate change/modify events!!
 	 */
-	virtual void SetNewChildScale(const FVector& NewChildScale);
+	UE_API virtual void SetNewChildScale(const FVector& NewChildScale);
 
 	/**
 	 * Set visibility for this Gizmo
 	 */
-	virtual void SetVisibility(bool bVisible);
+	UE_API virtual void SetVisibility(bool bVisible);
 
 	/**
 	 * Set customization function for this Gizmo
 	 */
-	void SetCustomizationFunction(const TFunction<const FGizmoCustomization()>& InFunction);
+	UE_API void SetCustomizationFunction(const TFunction<const FGizmoCustomization()>& InFunction);
 
 	/**
 	 * Handle widget mode changed.
 	 */
-	void HandleWidgetModeChanged(UE::Widget::EWidgetMode InWidgetMode);
+	UE_API void HandleWidgetModeChanged(UE::Widget::EWidgetMode InWidgetMode);
 
 	/**
 	 * Handle user parameters changes
 	 */
-	void OnParametersChanged(const FGizmosParameters& InParameters);
+	UE_API void OnParametersChanged(const FGizmosParameters& InParameters);
 	
 public:
 
@@ -472,7 +474,7 @@ protected:
 	TObjectPtr<UGizmoConstantFrameAxisSource> CameraAxisSource;
 
 	// internal function that updates CameraAxisSource by getting current view state from GizmoManager
-	void UpdateCameraAxisSource();
+	UE_API void UpdateCameraAxisSource();
 
 	/** The state target is created internally during SetActiveTarget() if no one is provided. */
 	UPROPERTY()
@@ -493,250 +495,250 @@ protected:
 	void RenderDebug(IToolsContextRenderAPI* RenderAPI);
 
 	/** Setup behaviors */
-	virtual void SetupBehaviors();
+	UE_API virtual void SetupBehaviors();
 
 	/** Setup indirect behaviors */
-	void SetupIndirectBehaviors();
+	UE_API void SetupIndirectBehaviors();
 
 	/** Setup materials */
-	virtual void SetupMaterials();
+	UE_API virtual void SetupMaterials();
 
 	/** Setup on click functions */
-	virtual void SetupOnClickFunctions();
+	UE_API virtual void SetupOnClickFunctions();
 
 	/** Update current gizmo mode based on transform source */
-	virtual void UpdateMode();
+	UE_API virtual void UpdateMode();
 
 	/** Update current gizmo rotation mode (default or gimbal)*/
-	void UpdateRotationMode();
+	UE_API void UpdateRotationMode();
 
 	/** Enable the given mode with the specified axes, EAxisList::Type::None will hide objects associated with mode */
-	virtual void EnableMode(EGizmoTransformMode InGizmoMode, EAxisList::Type InAxisListToDraw);
+	UE_API virtual void EnableMode(EGizmoTransformMode InGizmoMode, EAxisList::Type InAxisListToDraw);
 
 	/** Enable translate using specified axis list */
-	virtual void EnableTranslate(EAxisList::Type InAxisListToDraw);
+	UE_API virtual void EnableTranslate(EAxisList::Type InAxisListToDraw);
 
 	/** Enable rotate using specified axis list */
-	virtual void EnableRotate(EAxisList::Type InAxisListToDraw);
+	UE_API virtual void EnableRotate(EAxisList::Type InAxisListToDraw);
 
 	/** Enable scale using specified axis list */
-	virtual void EnableScale(EAxisList::Type InAxisListToDraw);
+	UE_API virtual void EnableScale(EAxisList::Type InAxisListToDraw);
 
 	/** Enable planar handles used by translate and scale */
-	virtual void EnablePlanarObjects(bool bTranslate, bool bEnableX, bool bEnableY, bool bEnableZ);
+	UE_API virtual void EnablePlanarObjects(bool bTranslate, bool bEnableX, bool bEnableY, bool bEnableZ);
 
 	/** Construct translate axis handle */
-	virtual UGizmoElementArrow* MakeTranslateAxis(ETransformGizmoPartIdentifier InPartId, const FVector& InAxisDir, const FVector& InSideDir, UMaterialInterface* InMaterial);
+	UE_API virtual UGizmoElementArrow* MakeTranslateAxis(ETransformGizmoPartIdentifier InPartId, const FVector& InAxisDir, const FVector& InSideDir, UMaterialInterface* InMaterial);
 
 	/** Updates the given translate axis handle with the current parameters (ie. size coefficient) */
-	virtual void UpdateTranslateAxis(UGizmoElementArrow* InElement);
+	UE_API virtual void UpdateTranslateAxis(UGizmoElementArrow* InElement);
 
 	/** Updates the given translate axis handle with the current parameters (ie. size coefficient) */
-	virtual void UpdateTranslateAxis(UGizmoElementArrow* InElement, const EAxis::Type InAxis);
+	UE_API virtual void UpdateTranslateAxis(UGizmoElementArrow* InElement, const EAxis::Type InAxis);
 
 	/** Updates the given translate axis handle with the current parameters (ie. size coefficient) */
-	virtual void UpdateTranslateAxis(UGizmoElementArrow* InElement, const FVector& InAxisDir, const FVector& InSideDir);
+	UE_API virtual void UpdateTranslateAxis(UGizmoElementArrow* InElement, const FVector& InAxisDir, const FVector& InSideDir);
 
 	/** Construct scale axis handle */
-	virtual UGizmoElementArrow* MakeScaleAxis(ETransformGizmoPartIdentifier InPartId, const FVector& InAxisDir, const FVector& InSideDir, UMaterialInterface* InMaterial);
+	UE_API virtual UGizmoElementArrow* MakeScaleAxis(ETransformGizmoPartIdentifier InPartId, const FVector& InAxisDir, const FVector& InSideDir, UMaterialInterface* InMaterial);
 
 	/**	Updates the given scale axis handle with the current parameters (ie. size coefficient). */
-	virtual void UpdateScaleAxis(UGizmoElementArrow* InElement);
+	UE_API virtual void UpdateScaleAxis(UGizmoElementArrow* InElement);
 
 	/**	Updates the given scale axis handle with the current parameters (ie. size coefficient). */
-	virtual void UpdateScaleAxis(UGizmoElementArrow* InElement, const EAxis::Type InAxis);
+	UE_API virtual void UpdateScaleAxis(UGizmoElementArrow* InElement, const EAxis::Type InAxis);
 
 	/**
 	 * Updates the given scale axis handle with the current parameters (ie. size coefficient).
 	 */
-	virtual void UpdateScaleAxis(UGizmoElementArrow* InElement, const FVector& InAxisDir, const FVector& InSideDir);
+	UE_API virtual void UpdateScaleAxis(UGizmoElementArrow* InElement, const FVector& InAxisDir, const FVector& InSideDir);
 
 	/** Updates all scale axis handles with the current parameters (ie. size coefficient) */
-	virtual void UpdateAllScaleAxis();
+	UE_API virtual void UpdateAllScaleAxis();
 
 	/** Construct the default rotate axis handle based on the ETransformGizmoPartIdentifier. */
-	UGizmoElementTorus* MakeDefaultRotateAxis(const ETransformGizmoPartIdentifier InPartId);
+	UE_API UGizmoElementTorus* MakeDefaultRotateAxis(const ETransformGizmoPartIdentifier InPartId);
 
 	/** Construct rotate axis handle */
-	virtual UGizmoElementTorus* MakeRotateAxis(ETransformGizmoPartIdentifier InPartId, const FVector& TorusAxis0, const FVector& TorusAxis1,
+	UE_API virtual UGizmoElementTorus* MakeRotateAxis(ETransformGizmoPartIdentifier InPartId, const FVector& TorusAxis0, const FVector& TorusAxis1,
 		UMaterialInterface* InMaterial, UMaterialInterface* InCurrentMaterial);
 
 	/** Updates the given rotate axis handle with the current parameters (ie. size coefficient) */
-	virtual void UpdateRotateAxis(UGizmoElementTorus* InElement);
+	UE_API virtual void UpdateRotateAxis(UGizmoElementTorus* InElement);
 
 	/** Construct uniform scale handle */
-	virtual UGizmoElementBox* MakeUniformScaleHandle();
+	UE_API virtual UGizmoElementBox* MakeUniformScaleHandle();
 
 	/** Updates the given uniform scale handle with the current parameters (ie. size coefficient) */
-	virtual void UpdateUniformScaleHandle(UGizmoElementBox* InElement);
+	UE_API virtual void UpdateUniformScaleHandle(UGizmoElementBox* InElement);
 
 	/** Construct planar axis handle */
-	virtual UGizmoElementBox* MakePlanarHandle(ETransformGizmoPartIdentifier InPartId, const FVector& InUpDirection, const FVector& InSideDirection, const FVector& InPlaneNormal,
+	UE_API virtual UGizmoElementBox* MakePlanarHandle(ETransformGizmoPartIdentifier InPartId, const FVector& InUpDirection, const FVector& InSideDirection, const FVector& InPlaneNormal,
 		UMaterialInterface* InMaterial);
 
 	/** Updates the given uniform scale handle with the current parameters (ie. size coefficient) */
-	virtual void UpdatePlanarHandle(UGizmoElementBox* InElement, const FVector& InUpDirection, const FVector& InSideDirection);
+	UE_API virtual void UpdatePlanarHandle(UGizmoElementBox* InElement, const FVector& InUpDirection, const FVector& InSideDirection);
 
 	/** Construct translate screen space handle */
-	virtual UGizmoElementRectangle* MakeTranslateScreenSpaceHandle();
+	UE_API virtual UGizmoElementRectangle* MakeTranslateScreenSpaceHandle();
 
 	/** Updates the given translate screen space handle with the current parameters (ie. size coefficient) */
-	virtual void UpdateTranslateScreenSpaceHandle(UGizmoElementRectangle* InElement);
+	UE_API virtual void UpdateTranslateScreenSpaceHandle(UGizmoElementRectangle* InElement);
 
 	/** Construct arcball screen space handle */
-	virtual UGizmoElementCircle* MakeArcballCircleHandle(ETransformGizmoPartIdentifier InPartId, float InRadius, const FLinearColor& InColor);
+	UE_API virtual UGizmoElementCircle* MakeArcballCircleHandle(ETransformGizmoPartIdentifier InPartId, float InRadius, const FLinearColor& InColor);
 
 	/** Construct rotate screen space handle */
-	virtual UGizmoElementCircle* MakeRotateCircleHandle(ETransformGizmoPartIdentifier InPartId, float InRadius, const FLinearColor& InColor,
+	UE_API virtual UGizmoElementCircle* MakeRotateCircleHandle(ETransformGizmoPartIdentifier InPartId, float InRadius, const FLinearColor& InColor,
 		const EGizmoElementDrawType InDrawType = EGizmoElementDrawType::Fill, const TOptional<FLinearColor>& InLineColorOverride = TOptional<FLinearColor>());
 
 	/** Updates the given rotate screen space handle with the current parameters (ie. size coefficient) */
-	virtual void UpdateRotateCircleHandle(UGizmoElementCircle* InElement, float InRadius);
+	UE_API virtual void UpdateRotateCircleHandle(UGizmoElementCircle* InElement, float InRadius);
 
 	/** Updates all gizmo elements based on current parameters */
-	virtual void UpdateElements();
+	UE_API virtual void UpdateElements();
 
 	/** Get gizmo transform based on cached current transform. */
-	virtual FTransform GetGizmoTransform() const;
+	UE_API virtual FTransform GetGizmoTransform() const;
 
 	/** Determine hit part and update hover state based on current input ray */
-	virtual FInputRayHit UpdateHoveredPart(const FInputDeviceRay& DevicePos);
+	UE_API virtual FInputRayHit UpdateHoveredPart(const FInputDeviceRay& DevicePos);
 
 	/** Get current interaction axis */
-	virtual FVector GetWorldAxis(const FVector& InAxis) const;
+	UE_API virtual FVector GetWorldAxis(const FVector& InAxis) const;
 
 	/** Get gimbal rotation axis (InAxis is expected to be in [0,2] range) */
-	FVector GetGimbalRotationAxis(const int32 InAxis) const;
+	UE_API FVector GetGimbalRotationAxis(const int32 InAxis) const;
 
 	/** Get nearest param along input ray to the current interaction axis */
-	virtual float GetNearestRayParamToInteractionAxis(const FInputDeviceRay& InRay);
+	UE_API virtual float GetNearestRayParamToInteractionAxis(const FInputDeviceRay& InRay);
 
 	/** Return true if input ray intersects current interaction plane and return param along ray in OutHitParam */
-	virtual bool GetRayParamIntersectionWithInteractionPlane(const FInputDeviceRay& InRay, FVector::FReal& OutHitParam);
+	UE_API virtual bool GetRayParamIntersectionWithInteractionPlane(const FInputDeviceRay& InRay, FVector::FReal& OutHitParam);
 
 	/** Update hover state for given part id */
-	virtual void UpdateHoverState(const bool bInHover, const ETransformGizmoPartIdentifier InPartId);
+	UE_API virtual void UpdateHoverState(const bool bInHover, const ETransformGizmoPartIdentifier InPartId);
 
 	/** Reset all hover states related to the transform mode to false */
-	void ResetHoverStates(const EGizmoTransformMode InMode);
+	UE_API void ResetHoverStates(const EGizmoTransformMode InMode);
 
 	/** Update interacting state for given part id */
-	virtual void UpdateInteractingState(const bool bInInteracting, const ETransformGizmoPartIdentifier InPartId, const bool bIdOnly = false);
+	UE_API virtual void UpdateInteractingState(const bool bInInteracting, const ETransformGizmoPartIdentifier InPartId, const bool bIdOnly = false);
 
 	/** Reset all interacting states related to the transform mode to false */
-	void ResetInteractingStates(const EGizmoTransformMode InMode);
+	UE_API void ResetInteractingStates(const EGizmoTransformMode InMode);
 
 	/** Called at the start of a sequence of gizmo transform edits */
-	virtual void BeginTransformEditSequence();
+	UE_API virtual void BeginTransformEditSequence();
 
 	/** Called at the end of a sequence of gizmo transform edits. */
-	virtual void EndTransformEditSequence();
+	UE_API virtual void EndTransformEditSequence();
 	
 	/**
 	 * Translate axis click-drag handling methods 
 	 */ 
 
 	/** Handle click press for translate X axis */
-	virtual void OnClickPressTranslateXAxis(const FInputDeviceRay& PressPos);
+	UE_API virtual void OnClickPressTranslateXAxis(const FInputDeviceRay& PressPos);
 
 	/** Handle click press for translate Y axis */
-	virtual void OnClickPressTranslateYAxis(const FInputDeviceRay& PressPos);
+	UE_API virtual void OnClickPressTranslateYAxis(const FInputDeviceRay& PressPos);
 
 	/** Handle click press for translate Z axis */
-	virtual void OnClickPressTranslateZAxis(const FInputDeviceRay& PressPos);
+	UE_API virtual void OnClickPressTranslateZAxis(const FInputDeviceRay& PressPos);
 
 	/** Handle click press for translate axes */
-	virtual void OnClickPressAxis(const FInputDeviceRay& PressPos);
+	UE_API virtual void OnClickPressAxis(const FInputDeviceRay& PressPos);
 
 	/** Handle click drag for translate axes */
-	virtual void OnClickDragTranslateAxis(const FInputDeviceRay& DragPos);
+	UE_API virtual void OnClickDragTranslateAxis(const FInputDeviceRay& DragPos);
 
 	/** Handle click release for translate axes */
-	virtual void OnClickReleaseTranslateAxis(const FInputDeviceRay& ReleasePos);
+	UE_API virtual void OnClickReleaseTranslateAxis(const FInputDeviceRay& ReleasePos);
 
 	/**
 	 * Translate and scale planar click-drag handling methods
 	 */
 
 	 /** Handle click press for translate XY planar */
-	virtual void OnClickPressTranslateXYPlanar(const FInputDeviceRay& PressPos);
+	UE_API virtual void OnClickPressTranslateXYPlanar(const FInputDeviceRay& PressPos);
 
 	/** Handle click press for translate YZ planar */
-	virtual void OnClickPressTranslateYZPlanar(const FInputDeviceRay& PressPos);
+	UE_API virtual void OnClickPressTranslateYZPlanar(const FInputDeviceRay& PressPos);
 
 	/** Handle click press for translate XZ planar */
-	virtual void OnClickPressTranslateXZPlanar(const FInputDeviceRay& PressPos);
+	UE_API virtual void OnClickPressTranslateXZPlanar(const FInputDeviceRay& PressPos);
 
 	/** Handle click press for generic planar */
-	virtual void OnClickPressPlanar(const FInputDeviceRay& PressPos);
+	UE_API virtual void OnClickPressPlanar(const FInputDeviceRay& PressPos);
 
 	/** Handle click drag for translate planar */
-	virtual void OnClickDragTranslatePlanar(const FInputDeviceRay& DragPos);
+	UE_API virtual void OnClickDragTranslatePlanar(const FInputDeviceRay& DragPos);
 
 	/** Handle click release for translate planar */
-	virtual void OnClickReleaseTranslatePlanar(const FInputDeviceRay& ReleasePos);
+	UE_API virtual void OnClickReleaseTranslatePlanar(const FInputDeviceRay& ReleasePos);
 
 	/** Compute translate axis delta based on start/end params */
-	virtual FVector ComputeAxisTranslateDelta(double InStartParam, double InEndParam);
+	UE_API virtual FVector ComputeAxisTranslateDelta(double InStartParam, double InEndParam);
 
 	/** Compute translate planar delta based on world space start/end points */
-	virtual FVector ComputePlanarTranslateDelta(const FVector& InStartPoint, const FVector& InEndPoint);
+	UE_API virtual FVector ComputePlanarTranslateDelta(const FVector& InStartPoint, const FVector& InEndPoint);
 
 	/**
 	 * Screen-space translate interaction methods
 	 */
 
 	 /** Handle click press for screen-space translate */
-	virtual void OnClickPressScreenSpaceTranslate(const FInputDeviceRay& PressPos);
+	UE_API virtual void OnClickPressScreenSpaceTranslate(const FInputDeviceRay& PressPos);
 
 	/** Handle click drag for screen-space translate */
-	virtual void OnClickDragScreenSpaceTranslate(const FInputDeviceRay& DragPos);
+	UE_API virtual void OnClickDragScreenSpaceTranslate(const FInputDeviceRay& DragPos);
 
 	/** Handle click release for screen-space translate */
-	virtual void OnClickReleaseScreenSpaceTranslate(const FInputDeviceRay& ReleasePos);
+	UE_API virtual void OnClickReleaseScreenSpaceTranslate(const FInputDeviceRay& ReleasePos);
 
 	/**
 	 * Rotate interaction methods
 	 */
 
 	/** Handle click press for rotate X axis */
-	virtual void OnClickPressRotateXAxis(const FInputDeviceRay& InPressPos);
+	UE_API virtual void OnClickPressRotateXAxis(const FInputDeviceRay& InPressPos);
 
 	/** Handle click press for rotate Y axis */
-	virtual void OnClickPressRotateYAxis(const FInputDeviceRay& InPressPos);
+	UE_API virtual void OnClickPressRotateYAxis(const FInputDeviceRay& InPressPos);
 
 	/** Handle click press for rotate Z axis */
-	virtual void OnClickPressRotateZAxis(const FInputDeviceRay& InPressPos);
+	UE_API virtual void OnClickPressRotateZAxis(const FInputDeviceRay& InPressPos);
 
 	/** Handle click press for any rotate axis */
-	void OnClickPressRotateAxis(const FInputDeviceRay& InPressPos);
+	UE_API void OnClickPressRotateAxis(const FInputDeviceRay& InPressPos);
 
 	/** Handle click drag for rotate axis */
-	virtual void OnClickDragRotateAxis(const FInputDeviceRay& InDragPos);
+	UE_API virtual void OnClickDragRotateAxis(const FInputDeviceRay& InDragPos);
 
 	/** Handle click release for rotate axes */
-	virtual void OnClickReleaseRotateAxis(const FInputDeviceRay& InReleasePos);
+	UE_API virtual void OnClickReleaseRotateAxis(const FInputDeviceRay& InReleasePos);
 
 	/** Get screen-space axis for rotation drag */
-	FVector2D GetScreenRotateAxisDir(const FInputDeviceRay& InPressPos);
+	UE_API FVector2D GetScreenRotateAxisDir(const FInputDeviceRay& InPressPos);
 	
 	/** Get screen-space axis for gimbal rotation drag */
-	FVector2D GetScreenGimbalRotateAxisDir(const FInputDeviceRay& InPressPos);
+	UE_API FVector2D GetScreenGimbalRotateAxisDir(const FInputDeviceRay& InPressPos);
 
 	/** Get screen-space axis based on a world space axis */
-	FVector2D GetWorldToScreenRotateAxisDir(const FInputDeviceRay& InPressPos, const FVector& InWorldAxis);
+	UE_API FVector2D GetWorldToScreenRotateAxisDir(const FInputDeviceRay& InPressPos, const FVector& InWorldAxis);
 
 	/** Compute angle (degrees) delta based on screen-space start/end positions. Intended to be used in combination with ComputeAxisRotateDelta */
-	virtual double ComputeAxisRotateDeltaAngle(const FVector2D& InStartPos, const FInputDeviceRay& InDragPos);
+	UE_API virtual double ComputeAxisRotateDeltaAngle(const FVector2D& InStartPos, const FInputDeviceRay& InDragPos);
 
 	/** Compute rotate delta based on screen-space angle delta (in degrees) */
-	virtual FQuat ComputeAxisRotateDelta(const double& InDeltaAngle);
+	UE_API virtual FQuat ComputeAxisRotateDelta(const double& InDeltaAngle);
 
 	/** Compute gimbal rotate delta based on screen-space start/end positions */
-	FQuat ComputeGimbalRotateDelta(const FVector2D& InStartPos, const FVector2D& InEndPos);
+	UE_API FQuat ComputeGimbalRotateDelta(const FVector2D& InStartPos, const FVector2D& InEndPos);
 
 	/** Prepares data for arc rotation. This will return false if this is not possible (the rotate handle is perpendicular to the view) */
-	bool OnClickPressRotateArc( const FInputDeviceRay& InPressPos,
+	UE_API bool OnClickPressRotateArc( const FInputDeviceRay& InPressPos,
 		const FVector& InPlaneNormal, const FVector& InPlaneAxis1, const FVector& InPlaneAxis2);
 
 	/**
@@ -745,141 +747,141 @@ protected:
 	 */
 
 	/** Handle click press for any gimbal rotate axis */
-	void OnClickPressGimbalRotateAxis(const FInputDeviceRay& InPressPos);
+	UE_API void OnClickPressGimbalRotateAxis(const FInputDeviceRay& InPressPos);
 
 	/** Handle click drag for gimbal rotate axis */
-	virtual void OnClickDragGimbalRotateAxis(const FInputDeviceRay& InDragPos);
+	UE_API virtual void OnClickDragGimbalRotateAxis(const FInputDeviceRay& InDragPos);
 	
 	/**
 	 * Screen-space rotate interaction methods
 	 */
 
 	 /** Handle click press for screen-space rotate */
-	virtual void OnClickPressScreenSpaceRotate(const FInputDeviceRay& PressPos);
+	UE_API virtual void OnClickPressScreenSpaceRotate(const FInputDeviceRay& PressPos);
 
 	/** Handle click drag for screen-space rotate */
-	virtual void OnClickDragScreenSpaceRotate(const FInputDeviceRay& InDragPos);
+	UE_API virtual void OnClickDragScreenSpaceRotate(const FInputDeviceRay& InDragPos);
 
 	/** Handle click release for screen-space rotate */
-	virtual void OnClickReleaseScreenSpaceRotate(const FInputDeviceRay& ReleasePos);
+	UE_API virtual void OnClickReleaseScreenSpaceRotate(const FInputDeviceRay& ReleasePos);
 
 	/** Compute rotate delta based on start/end angles */
-	virtual FQuat ComputeAngularRotateDelta(FQuat::FReal InStartAngle, FQuat::FReal InEndAngle);
+	UE_API virtual FQuat ComputeAngularRotateDelta(FQuat::FReal InStartAngle, FQuat::FReal InEndAngle);
 	
 	/**
 	 * Arc ball rotate interaction methods
 	 */
 
 	/** Handle click press for arc ball rotate */
-	virtual void OnClickPressArcBallRotate(const FInputDeviceRay& PressPos);
+	UE_API virtual void OnClickPressArcBallRotate(const FInputDeviceRay& PressPos);
 
 	/** Handle click drag for arc ball rotate */
-	virtual void OnClickDragArcBallRotate(const FInputDeviceRay& DragPos);
+	UE_API virtual void OnClickDragArcBallRotate(const FInputDeviceRay& DragPos);
 
 	/** Handle click release for arc ball rotate */
-	virtual void OnClickReleaseArcBallRotate(const FInputDeviceRay& ReleasePos);
+	UE_API virtual void OnClickReleaseArcBallRotate(const FInputDeviceRay& ReleasePos);
 
 	/** Get the arc ball sphere world radius */
-	float GetWorldRadius(const float InRadius) const;
+	UE_API float GetWorldRadius(const float InRadius) const;
 
 	/** */
-	float GetSizeCoefficient() const;
+	UE_API float GetSizeCoefficient() const;
 	
 	/**
 	* Scale click-drag handling methods
 	*/
 
 	/** Handle click press for scale X axis */
-	virtual void OnClickPressScaleXAxis(const FInputDeviceRay& PressPos);
+	UE_API virtual void OnClickPressScaleXAxis(const FInputDeviceRay& PressPos);
 
 	/** Handle click press for scale Y axis */
-	virtual void OnClickPressScaleYAxis(const FInputDeviceRay& PressPos);
+	UE_API virtual void OnClickPressScaleYAxis(const FInputDeviceRay& PressPos);
 
 	/** Handle click press for scale Z axis */
-	virtual void OnClickPressScaleZAxis(const FInputDeviceRay& PressPos);
+	UE_API virtual void OnClickPressScaleZAxis(const FInputDeviceRay& PressPos);
 
 	/** Handle click press for scale XY planar */
-	virtual void OnClickPressScaleXYPlanar(const FInputDeviceRay& PressPos);
+	UE_API virtual void OnClickPressScaleXYPlanar(const FInputDeviceRay& PressPos);
 
 	/** Handle click press for scale YZ planar */
-	virtual void OnClickPressScaleYZPlanar(const FInputDeviceRay& PressPos);
+	UE_API virtual void OnClickPressScaleYZPlanar(const FInputDeviceRay& PressPos);
 
 	/** Handle click press for scale XZ planar */
-	virtual void OnClickPressScaleXZPlanar(const FInputDeviceRay& PressPos);
+	UE_API virtual void OnClickPressScaleXZPlanar(const FInputDeviceRay& PressPos);
 
 	/** Handle click press for uniform scale */
-	virtual void OnClickPressScaleXYZ(const FInputDeviceRay& PressPos);
+	UE_API virtual void OnClickPressScaleXYZ(const FInputDeviceRay& PressPos);
 
 	/** Handle click press for all scale methods */
-	virtual void OnClickPressScale(const FInputDeviceRay& PressPos);
+	UE_API virtual void OnClickPressScale(const FInputDeviceRay& PressPos);
 
 	/** Handle click drag for scale axes */
-	virtual void OnClickDragScaleAxis(const FInputDeviceRay& DragPos);
+	UE_API virtual void OnClickDragScaleAxis(const FInputDeviceRay& DragPos);
 
 	/** Handle click drag for scale planar */
-	virtual void OnClickDragScalePlanar(const FInputDeviceRay& DragPos);
+	UE_API virtual void OnClickDragScalePlanar(const FInputDeviceRay& DragPos);
 
 	/** Handle click drag for uniform scale */
-	virtual void OnClickDragScaleXYZ(const FInputDeviceRay& DragPos);
+	UE_API virtual void OnClickDragScaleXYZ(const FInputDeviceRay& DragPos);
 
 	/** Handle click drag for all scale */
-	virtual void OnClickDragScale(const FInputDeviceRay& DragPos);
+	UE_API virtual void OnClickDragScale(const FInputDeviceRay& DragPos);
 
 	/** Handle click release for scale axes */
-	virtual void OnClickReleaseScaleAxis(const FInputDeviceRay& ReleasePos);
+	UE_API virtual void OnClickReleaseScaleAxis(const FInputDeviceRay& ReleasePos);
 
 	/** Handle click release for scale planar */
-	virtual void OnClickReleaseScalePlanar(const FInputDeviceRay& ReleasePos);
+	UE_API virtual void OnClickReleaseScalePlanar(const FInputDeviceRay& ReleasePos);
 
 	/** Handle click release for uniform scale */
-	virtual void OnClickReleaseScaleXYZ(const FInputDeviceRay& ReleasePos);
+	UE_API virtual void OnClickReleaseScaleXYZ(const FInputDeviceRay& ReleasePos);
 
 	/** Compute scale delta based on screen space start/end positions */
-	virtual FVector ComputeScaleDelta(const FVector2D& InStartPos, const FVector2D& InEndPos, FVector2D& OutScreenDelta);
+	UE_API virtual FVector ComputeScaleDelta(const FVector2D& InStartPos, const FVector2D& InEndPos, FVector2D& OutScreenDelta);
 
 	/**
 	 * Screen-space helper method
 	 */
 
 	/** Returns 2D vector projection of input axis onto input view plane */
-	static FVector2D GetScreenProjectedAxis(const UGizmoViewContext* View, const FVector& InLocalAxis, const FTransform& InLocalToWorld = FTransform::Identity);
+	static UE_API FVector2D GetScreenProjectedAxis(const UGizmoViewContext* View, const FVector& InLocalAxis, const FTransform& InLocalToWorld = FTransform::Identity);
 
 	/**
 	 * Apply transform delta methods
 	 */
 	 
 	/** Apply translate delta to transform proxy */
-	virtual void ApplyTranslateDelta(const FVector& InTranslateDelta);
+	UE_API virtual void ApplyTranslateDelta(const FVector& InTranslateDelta);
 
 	/** Apply rotate delta to transform proxy */
-	virtual void ApplyRotateDelta(const FQuat& InRotateDelta);
+	UE_API virtual void ApplyRotateDelta(const FQuat& InRotateDelta);
 
 	/** Apply scale delta to transform proxy */
-	virtual void ApplyScaleDelta(const FVector& InScaleDelta);
+	UE_API virtual void ApplyScaleDelta(const FVector& InScaleDelta);
 
 	// Axis and Plane TransformSources use this function to execute worldgrid snap queries
-	bool PositionSnapFunction(const FVector& WorldPosition, FVector& SnappedPositionOut) const;
-	FQuat RotationSnapFunction(const FQuat& DeltaRotation) const;
+	UE_API bool PositionSnapFunction(const FVector& WorldPosition, FVector& SnappedPositionOut) const;
+	UE_API FQuat RotationSnapFunction(const FQuat& DeltaRotation) const;
 
-	void SnapTranslateDelta(FVector& InOutWorldDelta) const;
-	void SnapRotateDelta(FQuat& InOutWorldDelta) const;
-	void SnapRotateAngleDelta(double& InOutAngleDelta) const;
-	void SnapScaleDelta(FVector& InOutLocalScaleDelta) const;
+	UE_API void SnapTranslateDelta(FVector& InOutWorldDelta) const;
+	UE_API void SnapRotateDelta(FQuat& InOutWorldDelta) const;
+	UE_API void SnapRotateAngleDelta(double& InOutAngleDelta) const;
+	UE_API void SnapScaleDelta(FVector& InOutLocalScaleDelta) const;
 
 	// Get max part identifier.
-	virtual uint32 GetMaxPartIdentifier() const;
+	UE_API virtual uint32 GetMaxPartIdentifier() const;
 
 	// Verify part identifier is within recognized range of transform gizmo part ids
-	virtual bool VerifyPartIdentifier(uint32 InPartIdentifier) const;
+	UE_API virtual bool VerifyPartIdentifier(uint32 InPartIdentifier) const;
 
 	// Returns whether the gizmo is visible in the viewport.
-	bool IsVisible(const EViewportContext InViewportContext = EViewportContext::Focused) const;
+	UE_API bool IsVisible(const EViewportContext InViewportContext = EViewportContext::Focused) const;
 
 	// Returns whether the gizmo can interact. Note that this can be true even if the gizmo is not visible to support indirect manipulation.
-	bool CanInteract(const EViewportContext InViewportContext = EViewportContext::Focused) const;
+	UE_API bool CanInteract(const EViewportContext InViewportContext = EViewportContext::Focused) const;
 
 	// Returns the current rotation context.
-	const FRotationContext& GetRotationContext() const;
+	UE_API const FRotationContext& GetRotationContext() const;
 	
 protected:
 
@@ -938,8 +940,8 @@ protected:
 	/** Last hit part per mode */
 	ETransformGizmoPartIdentifier LastHitPartPerMode[static_cast<uint8>(EGizmoTransformMode::Max)];
 	
-	ETransformGizmoPartIdentifier GetCurrentModeLastHitPart() const;
-	void SetModeLastHitPart(const EGizmoTransformMode InMode, const ETransformGizmoPartIdentifier InIdentifier);
+	UE_API ETransformGizmoPartIdentifier GetCurrentModeLastHitPart() const;
+	UE_API void SetModeLastHitPart(const EGizmoTransformMode InMode, const ETransformGizmoPartIdentifier InIdentifier);
 	
 	//
 	// The values below are used in the context of a single click-drag interaction, ie if bInInteraction = true
@@ -1093,3 +1095,5 @@ private:
 		double InteractionRadius;
 	} DebugData;
 };
+
+#undef UE_API
diff --git a/Engine/Source/Editor/Experimental/EditorInteractiveToolsFramework/Public/EditorGizmos/TransformGizmoInterfaces.h b/Engine/Source/Editor/Experimental/EditorInteractiveToolsFramework/Public/EditorGizmos/TransformGizmoInterfaces.h
--- a/Engine/Source/Editor/Experimental/EditorInteractiveToolsFramework/Public/EditorGizmos/TransformGizmoInterfaces.h
+++ b/Engine/Source/Editor/Experimental/EditorInteractiveToolsFramework/Public/EditorGizmos/TransformGizmoInterfaces.h
@@ -62,8 +62,8 @@ struct FRotationContext
 	FTransform Offset = FTransform::Identity;
 };
 
-UINTERFACE()
-class EDITORINTERACTIVETOOLSFRAMEWORK_API UTransformGizmoSource : public UInterface
+UINTERFACE(MinimalAPI)
+class UTransformGizmoSource : public UInterface
 {
 	GENERATED_BODY()
 };
@@ -71,7 +71,7 @@ class EDITORINTERACTIVETOOLSFRAMEWORK_API UTransformGizmoSource : public UInterf
 /**
  * ITransformGizmoSource is an interface for providing gizmo mode configuration information.
  */
-class EDITORINTERACTIVETOOLSFRAMEWORK_API ITransformGizmoSource
+class ITransformGizmoSource
 {
 	GENERATED_BODY()
 public:
diff --git a/Engine/Source/Editor/Experimental/EditorInteractiveToolsFramework/Public/EditorInteractiveGizmoConditionalBuilder.h b/Engine/Source/Editor/Experimental/EditorInteractiveToolsFramework/Public/EditorInteractiveGizmoConditionalBuilder.h
--- a/Engine/Source/Editor/Experimental/EditorInteractiveToolsFramework/Public/EditorInteractiveGizmoConditionalBuilder.h
+++ b/Engine/Source/Editor/Experimental/EditorInteractiveToolsFramework/Public/EditorInteractiveGizmoConditionalBuilder.h
@@ -14,7 +14,7 @@
  * from SatsifiesCondition(), the gizmo builder with highest priority will be used. If 
  * there are multiple builders the highest priority, multiple gizmos will be built.
  */
-struct EDITORINTERACTIVETOOLSFRAMEWORK_API FEditorGizmoTypePriority
+struct FEditorGizmoTypePriority
 {
 	static constexpr int DEFAULT_GIZMO_TYPE_PRIORITY = 50;
 
@@ -59,13 +59,13 @@ struct EDITORINTERACTIVETOOLSFRAMEWORK_API FEditorGizmoTypePriority
  * on selection. The builder's priority should be specified such that when more than one builder
  * is discovered, the builder with highest priority will be built.
  */
-UINTERFACE()
-class EDITORINTERACTIVETOOLSFRAMEWORK_API UEditorInteractiveGizmoConditionalBuilder : public UInterface
+UINTERFACE(MinimalAPI)
+class UEditorInteractiveGizmoConditionalBuilder : public UInterface
 {
 	GENERATED_BODY()
 };
 
-class EDITORINTERACTIVETOOLSFRAMEWORK_API IEditorInteractiveGizmoConditionalBuilder
+class IEditorInteractiveGizmoConditionalBuilder
 {
 	GENERATED_BODY()
 
diff --git a/Engine/Source/Editor/Experimental/EditorInteractiveToolsFramework/Public/EditorInteractiveGizmoManager.h b/Engine/Source/Editor/Experimental/EditorInteractiveToolsFramework/Public/EditorInteractiveGizmoManager.h
--- a/Engine/Source/Editor/Experimental/EditorInteractiveToolsFramework/Public/EditorInteractiveGizmoManager.h
+++ b/Engine/Source/Editor/Experimental/EditorInteractiveToolsFramework/Public/EditorInteractiveGizmoManager.h
@@ -14,6 +14,8 @@
 
 #include "EditorInteractiveGizmoManager.generated.h"
 
+#define UE_API EDITORINTERACTIVETOOLSFRAMEWORK_API
+
 class FCanvas;
 class FEditorModeTools;
 class IToolsContextQueriesAPI;
@@ -44,28 +46,28 @@ struct FActiveEditorGizmo
  * When the section changes, the highest priority builders for which SatisfiesCondition() return true, will be used to
  * build gizmos.
  */
-UCLASS(Transient)
-class EDITORINTERACTIVETOOLSFRAMEWORK_API UEditorInteractiveGizmoManager : public UInteractiveGizmoManager
+UCLASS(MinimalAPI, Transient)
+class UEditorInteractiveGizmoManager : public UInteractiveGizmoManager
 {
 	GENERATED_BODY()
 
 protected:
 	friend class UEditorInteractiveToolsContext;		// to call Initialize/Shutdown
 
-	UEditorInteractiveGizmoManager();
+	UE_API UEditorInteractiveGizmoManager();
 
 	/** Initialize the GizmoManager with the necessary Context-level state. UEdModeInteractiveToolsContext calls this, you should not. */
-	virtual void InitializeWithEditorModeManager(IToolsContextQueriesAPI* QueriesAPI, IToolsContextTransactionsAPI* TransactionsAPI, UInputRouter* InputRouter, FEditorModeTools* InEditorModeManager);
+	UE_API virtual void InitializeWithEditorModeManager(IToolsContextQueriesAPI* QueriesAPI, IToolsContextTransactionsAPI* TransactionsAPI, UInputRouter* InputRouter, FEditorModeTools* InEditorModeManager);
 
 	// UInteractiveGizmoManager interface
-	virtual void Shutdown() override;
+	UE_API virtual void Shutdown() override;
 
 public:
 
 	// UInteractiveGizmoManager interface
-	virtual void Tick(float DeltaTime) override;
-	virtual void Render(IToolsContextRenderAPI* RenderAPI) override;
-	virtual void DrawHUD(FCanvas* Canvas, IToolsContextRenderAPI* RenderAPI) override;
+	UE_API virtual void Tick(float DeltaTime) override;
+	UE_API virtual void Render(IToolsContextRenderAPI* RenderAPI) override;
+	UE_API virtual void DrawHUD(FCanvas* Canvas, IToolsContextRenderAPI* RenderAPI) override;
 
 	/**
 	 * Register a new Editor gizmo type.
@@ -74,14 +76,14 @@ public:
 	 * - Accessory and Primary gizmo builders must derive from UInteractiveGizmoBuilder (of from a builder derived from it)
 	 *   and must implement the IEditorInteractiveConditionalGizmoBuilder and IEditorInteractiveSelectionGizmoBuilder interfaces.
 	 */
-	void RegisterEditorGizmoType(EEditorGizmoCategory InGizmoCategory, UInteractiveGizmoBuilder* InGizmoBuilder);
+	UE_API void RegisterEditorGizmoType(EEditorGizmoCategory InGizmoCategory, UInteractiveGizmoBuilder* InGizmoBuilder);
 
 	/**
 	* Remove an Editor gizmo type from the set of known Editor gizmo types
 	* @param InGizmoBuilder same object pointer that was passed to RegisterEditorGizmoType()
 	* @return true if gizmo type was found and deregistered
 	*/
-	void DeregisterEditorGizmoType(EEditorGizmoCategory InGizmoCategory, UInteractiveGizmoBuilder* InGizmoBuilder);
+	UE_API void DeregisterEditorGizmoType(EEditorGizmoCategory InGizmoCategory, UInteractiveGizmoBuilder* InGizmoBuilder);
 
 	/**
 	 * Get all qualified Editor gizmo builders for the specified category, based on the current state. Qualification is determined by the gizmo builder
@@ -91,7 +93,7 @@ public:
 	 * @param InToolBuilderState current selection and other state
 	 * @return array of qualified Gizmo selection builders based on current state
 	 */
-	void GetQualifiedEditorGizmoBuilders(EEditorGizmoCategory InGizmoCategory, const FToolBuilderState& InToolBuilderState, TArray<UInteractiveGizmoBuilder*>& InFoundBuilders);
+	UE_API void GetQualifiedEditorGizmoBuilders(EEditorGizmoCategory InGizmoCategory, const FToolBuilderState& InToolBuilderState, TArray<UInteractiveGizmoBuilder*>& InFoundBuilders);
 
 	/** 
 	 * Set how auto gizmo resolution should occur when CreateGizmosForCurrentState is invoked. If bSearchLocalOnly is true, only the current
@@ -116,60 +118,60 @@ public:
 	 * @param Gizmo the Gizmo to shutdown and remove
 	 * @return true if the Gizmo was found and removed
 	 */
-	virtual bool DestroyEditorGizmo(UInteractiveGizmo* Gizmo);
+	UE_API virtual bool DestroyEditorGizmo(UInteractiveGizmo* Gizmo);
 
 	/**
 	 * Shutdown and remove all active auto gizmos
 	 */
-	virtual void DestroyAllEditorGizmos();
+	UE_API virtual void DestroyAllEditorGizmos();
 
 	/** Try to activate a new Gizmo instance (UInteractiveGizmoManager override) */	
-	virtual UInteractiveGizmo* CreateGizmo(
+	UE_API virtual UInteractiveGizmo* CreateGizmo(
 		const FString& BuilderIdentifier, const FString& InstanceIdentifier = FString(), void* Owner = nullptr) override;
 
 	/** Shutdown and remove a Gizmo (UInteractiveGizmoManager override) */
-	virtual bool DestroyGizmo(UInteractiveGizmo* InGizmo) override;
+	UE_API virtual bool DestroyGizmo(UInteractiveGizmo* InGizmo) override;
 
 	/** instance/builder identifiers for transform gizmo */
-	static const FString& TransformInstanceIdentifier();
-	static const FString& TransformBuilderIdentifier();
+	static UE_API const FString& TransformInstanceIdentifier();
+	static UE_API const FString& TransformBuilderIdentifier();
 
 	/**
 	 * Returns true if the new TRS gizmos are used.
 	 */
-	static bool UsesNewTRSGizmos();
+	static UE_API bool UsesNewTRSGizmos();
 
 	/**
 	 * Updates the current New TRS Gizmo state and notifies that change using OnUsesNewTRSGizmosChangedDelegate  
 	 */
-	static void SetUsesNewTRSGizmos(const bool bUseNewTRSGizmos);
+	static UE_API void SetUsesNewTRSGizmos(const bool bUseNewTRSGizmos);
 
 	/**
 	 * Delegate to notify from bUseNewTRSGizmos changes  
 	 */
 	DECLARE_MULTICAST_DELEGATE_OneParam(FOnUsesNewTRSGizmosChanged, const bool bUseNewTRSGizmos);
-	static FOnUsesNewTRSGizmosChanged& OnUsesNewTRSGizmosChangedDelegate();
+	static UE_API FOnUsesNewTRSGizmosChanged& OnUsesNewTRSGizmosChangedDelegate();
 
 	/**
 	 * Notifies FGizmosParameters changes across bound transform gizmos
 	 */
-	static void SetGizmosParameters(const FGizmosParameters& InParameters);
+	static UE_API void SetGizmosParameters(const FGizmosParameters& InParameters);
 
 	/**
 	 * Delegate to notify from FGizmosParameters changes  
 	 */
 	DECLARE_MULTICAST_DELEGATE_OneParam(FOnGizmosParametersChanged, const FGizmosParameters& InParameters);
-    static FOnGizmosParametersChanged& OnGizmosParametersChangedDelegate();
+    static UE_API FOnGizmosParametersChanged& OnGizmosParametersChangedDelegate();
 
 	/**
 	 * Returns the default gizmos parameters if set
 	 */
-    static const TOptional<FGizmosParameters>& GetDefaultGizmosParameters();
+    static UE_API const TOptional<FGizmosParameters>& GetDefaultGizmosParameters();
 
 	/**
 	 * Returns true if the Explicit mode is enabled;
 	 */
-	static bool IsExplicitModeEnabled();
+	static UE_API bool IsExplicitModeEnabled();
 
 protected:
 
@@ -177,18 +179,18 @@ protected:
 	 * Returns true if selection gizmos should be visible. 
 	 * @todo move this to a gizmo context object
 	 */
-	virtual bool GetShowEditorGizmos();
+	UE_API virtual bool GetShowEditorGizmos();
 
 	/**
 	 * Returns true if gizmos should be visible based on the current view's engine show flag.
 	 * @todo move this to a gizmo context object
 	 */
-	virtual bool GetShowEditorGizmosForView(IToolsContextRenderAPI* RenderAPI);
+	UE_API virtual bool GetShowEditorGizmosForView(IToolsContextRenderAPI* RenderAPI);
 
 	/**
 	 * Updates active selection gizmos when show selection state changes
 	 */
-	void UpdateActiveEditorGizmos();
+	UE_API void UpdateActiveEditorGizmos();
 
 	/** Actual registry */
 	UPROPERTY()
@@ -217,7 +219,9 @@ protected:
 private:
 
 	/** Returns the existing default Gizmo instance if any. */
-	UTransformGizmo* FindDefaultTransformGizmo() const;
+	UE_API UTransformGizmo* FindDefaultTransformGizmo() const;
 
-	void OnEditorModeChanged(const FName& InModeID, bool bInIsEnteringMode);
+	UE_API void OnEditorModeChanged(const FName& InModeID, bool bInIsEnteringMode);
 };
+
+#undef UE_API
diff --git a/Engine/Source/Editor/Experimental/EditorInteractiveToolsFramework/Public/EditorInteractiveGizmoRegistry.h b/Engine/Source/Editor/Experimental/EditorInteractiveToolsFramework/Public/EditorInteractiveGizmoRegistry.h
--- a/Engine/Source/Editor/Experimental/EditorInteractiveToolsFramework/Public/EditorInteractiveGizmoRegistry.h
+++ b/Engine/Source/Editor/Experimental/EditorInteractiveToolsFramework/Public/EditorInteractiveGizmoRegistry.h
@@ -17,6 +17,8 @@
 
 #include "EditorInteractiveGizmoRegistry.generated.h"
 
+#define UE_API EDITORINTERACTIVETOOLSFRAMEWORK_API
+
 class UInteractiveGizmoBuilder;
 struct FToolBuilderState;
 template <typename T> struct TObjectPtr;
@@ -99,17 +101,17 @@ public:
  * the subystem is global to the Editor, the manager is local to the Interactive
  * Tools Context.
  */
-UCLASS()
-class EDITORINTERACTIVETOOLSFRAMEWORK_API UEditorInteractiveGizmoRegistry : public UObject
+UCLASS(MinimalAPI)
+class UEditorInteractiveGizmoRegistry : public UObject
 {
 	GENERATED_BODY()
 
 public:
 
-	UEditorInteractiveGizmoRegistry();
+	UE_API UEditorInteractiveGizmoRegistry();
 
 	/** Shutdown the registry, called by the gizmo subsystem and gizmo manager when they are shutdown/deinitialized. */
-	virtual void Shutdown();
+	UE_API virtual void Shutdown();
 
 	/**
 	 * Register a new Editor gizmo type.
@@ -118,19 +120,19 @@ public:
 	 * - Accessory and Primary gizmo builders must derive from UInteractiveGizmoBuilder (of from a builder derived from it)
 	 *   and must implement the IEditorInteractiveConditionalGizmoBuilder and IEditorInteractiveSelectionGizmoBuilder interfaces.
 	 */
-	void RegisterEditorGizmoType(EEditorGizmoCategory InGizmoCategory, UInteractiveGizmoBuilder* InGizmoBuilder);
+	UE_API void RegisterEditorGizmoType(EEditorGizmoCategory InGizmoCategory, UInteractiveGizmoBuilder* InGizmoBuilder);
 
 	/**
 	* Remove an Editor gizmo type from the set of known Editor gizmo types
 	* @param InGizmoBuilder same object pointer that was passed to RegisterEditorGizmoType()
 	* @return true if gizmo type was found and deregistered
 	*/
-	void DeregisterEditorGizmoType(EEditorGizmoCategory InGizmoCategory, UInteractiveGizmoBuilder* InGizmoBuilder);
+	UE_API void DeregisterEditorGizmoType(EEditorGizmoCategory InGizmoCategory, UInteractiveGizmoBuilder* InGizmoBuilder);
 
 	/**
 	 * Clear all registered gizmo types
 	 */
-	void ClearEditorGizmoTypes();
+	UE_API void ClearEditorGizmoTypes();
 
 	/**
 	 * Get all qualified Editor gizmo builders for the specified category, based on the current state. Qualification is determined by the gizmo builder
@@ -142,11 +144,13 @@ public:
 	 *                      builders which may add to or replace the input found builders.
 	 * @return array of qualified Gizmo selection builders based on current state
 	 */
-	void GetQualifiedEditorGizmoBuilders(EEditorGizmoCategory InGizmoCategory, const FToolBuilderState& InToolBuilderState, TArray<UInteractiveGizmoBuilder*>& FoundBuilders);
+	UE_API void GetQualifiedEditorGizmoBuilders(EEditorGizmoCategory InGizmoCategory, const FToolBuilderState& InToolBuilderState, TArray<UInteractiveGizmoBuilder*>& FoundBuilders);
 
 private:
 
 	/** Current set of Gizmo Builders */
 	UPROPERTY()
 	TMap<EEditorGizmoCategory, TObjectPtr<UEditorGizmoRegistryCategoryEntry>> GizmoCategoryMap;
-};
\ No newline at end of file
+};
+
+#undef UE_API
diff --git a/Engine/Source/Editor/Experimental/EditorInteractiveToolsFramework/Public/EditorInteractiveGizmoSelectionBuilder.h b/Engine/Source/Editor/Experimental/EditorInteractiveToolsFramework/Public/EditorInteractiveGizmoSelectionBuilder.h
--- a/Engine/Source/Editor/Experimental/EditorInteractiveToolsFramework/Public/EditorInteractiveGizmoSelectionBuilder.h
+++ b/Engine/Source/Editor/Experimental/EditorInteractiveToolsFramework/Public/EditorInteractiveGizmoSelectionBuilder.h
@@ -11,27 +11,29 @@
 
 #include "EditorInteractiveGizmoSelectionBuilder.generated.h"
 
+#define UE_API EDITORINTERACTIVETOOLSFRAMEWORK_API
+
 class UInteractiveGizmo;
 class UObject;
 class UTransformProxy;
 struct FToolBuilderState;
 
-class EDITORINTERACTIVETOOLSFRAMEWORK_API FEditorGizmoSelectionBuilderHelper
+class FEditorGizmoSelectionBuilderHelper
 {
 public:
 	/**
 	 * Utility method that creates a transform proxy based on the current selection.
 	 */
-	static UTransformProxy* CreateTransformProxyForSelection(const FToolBuilderState& SceneState);
+	static UE_API UTransformProxy* CreateTransformProxyForSelection(const FToolBuilderState& SceneState);
 };
 
-UINTERFACE()
-class EDITORINTERACTIVETOOLSFRAMEWORK_API UEditorInteractiveGizmoSelectionBuilder : public UInterface
+UINTERFACE(MinimalAPI)
+class UEditorInteractiveGizmoSelectionBuilder : public UInterface
 {
 	GENERATED_BODY()
 };
 
-class EDITORINTERACTIVETOOLSFRAMEWORK_API IEditorInteractiveGizmoSelectionBuilder
+class IEditorInteractiveGizmoSelectionBuilder
 {
 	GENERATED_BODY()
 
@@ -46,3 +48,5 @@ public:
 	virtual void UpdateGizmoForSelection(UInteractiveGizmo* Gizmo, const FToolBuilderState& SceneState) = 0;
 
 };
+
+#undef UE_API
diff --git a/Engine/Source/Editor/Experimental/EditorInteractiveToolsFramework/Public/EditorInteractiveGizmoSubsystem.h b/Engine/Source/Editor/Experimental/EditorInteractiveToolsFramework/Public/EditorInteractiveGizmoSubsystem.h
--- a/Engine/Source/Editor/Experimental/EditorInteractiveToolsFramework/Public/EditorInteractiveGizmoSubsystem.h
+++ b/Engine/Source/Editor/Experimental/EditorInteractiveToolsFramework/Public/EditorInteractiveGizmoSubsystem.h
@@ -14,6 +14,8 @@
 
 #include "EditorInteractiveGizmoSubsystem.generated.h"
 
+#define UE_API EDITORINTERACTIVETOOLSFRAMEWORK_API
+
 class FSubsystemCollectionBase;
 class UInteractiveGizmoBuilder;
 class UObject;
@@ -35,19 +37,19 @@ struct FToolBuilderState;
  * - OnEditorGizmoSubsystemDeregisterEditorGizmoTypes()
  * to register and dergister their gizmo builders.
  */
- UCLASS()
-class EDITORINTERACTIVETOOLSFRAMEWORK_API UEditorInteractiveGizmoSubsystem : public UEditorSubsystem
+ UCLASS(MinimalAPI)
+class UEditorInteractiveGizmoSubsystem : public UEditorSubsystem
 {
 	GENERATED_BODY()
 
 public:
 
-	UEditorInteractiveGizmoSubsystem();
+	UE_API UEditorInteractiveGizmoSubsystem();
 
 	//~ Begin USubsystem interface
-	virtual void Initialize(FSubsystemCollectionBase& Collection) override;
+	UE_API virtual void Initialize(FSubsystemCollectionBase& Collection) override;
 
-	virtual void Deinitialize() override;
+	UE_API virtual void Deinitialize() override;
 	//~ End USubsystem interface
 
 	/**
@@ -65,12 +67,12 @@ public:
 	/**
 	 * Registers all built-in Editor gizmo types and broadcast registration event.
 	 */
-	void RegisterBuiltinEditorGizmoTypes();
+	UE_API void RegisterBuiltinEditorGizmoTypes();
 
 	/**
 	 * Removes all built-in Editor gizmo types and broadcast deregistration event.
 	 */
-	void DeregisterBuiltinEditorGizmoTypes();
+	UE_API void DeregisterBuiltinEditorGizmoTypes();
 
 	/**
 	 * Register a new Editor gizmo type which will be global to the Editor.
@@ -79,14 +81,14 @@ public:
 	 * - Accessory and Primary gizmo builders must derive from UInteractiveGizmoBuilder (of from a builder derived from it)
 	 *   and must implement the IEditorInteractiveConditionalGizmoBuilder and IEditorInteractiveSelectionGizmoBuilder interfaces.
 	 */
-	void RegisterGlobalEditorGizmoType(EEditorGizmoCategory InGizmoCategory, UInteractiveGizmoBuilder* InGizmoBuilder);
+	UE_API void RegisterGlobalEditorGizmoType(EEditorGizmoCategory InGizmoCategory, UInteractiveGizmoBuilder* InGizmoBuilder);
 
 	/**
 	* Remove an Editor gizmo type from the set of known Editor gizmo types
 	* @param InGizmoBuilder same object pointer that was passed to RegisterEditorGizmoType()
 	* @return true if gizmo type was found and deregistered
 	*/
-	void DeregisterGlobalEditorGizmoType(EEditorGizmoCategory InGizmoCategory, UInteractiveGizmoBuilder* InGizmoBuilder);
+	UE_API void DeregisterGlobalEditorGizmoType(EEditorGizmoCategory InGizmoCategory, UInteractiveGizmoBuilder* InGizmoBuilder);
 
 	/**
 	 * Get all qualified Editor gizmo builders for the specified category, based on the current state. Qualification is determined by the gizmo builder
@@ -96,7 +98,7 @@ public:
 	 * @param InToolBuilderState current selection and other state
 	 * @return array of qualified Gizmo selection builders based on current state
 	 */
-	void GetQualifiedGlobalEditorGizmoBuilders(EEditorGizmoCategory InGizmoCategory, const FToolBuilderState& InToolBuilderState, TArray<UInteractiveGizmoBuilder*>& InFoundBuilders);
+	UE_API void GetQualifiedGlobalEditorGizmoBuilders(EEditorGizmoCategory InGizmoCategory, const FToolBuilderState& InToolBuilderState, TArray<UInteractiveGizmoBuilder*>& InFoundBuilders);
 
 	/** 
 	 * Get transform gizmo builder used to build the Level Editor TRS gizmo. 
@@ -122,4 +124,6 @@ private:
 	/** Call to deregister gizmo types */
 	FOnEditorGizmoSubsystemDeregisterGlobalEditorGizmoTypes DeregisterGlobalEditorGizmoTypesDelegate;
 
-};
\ No newline at end of file
+};
+
+#undef UE_API
diff --git a/Engine/Source/Editor/Experimental/EditorInteractiveToolsFramework/Public/EditorInteractiveToolsFrameworkModule.h b/Engine/Source/Editor/Experimental/EditorInteractiveToolsFramework/Public/EditorInteractiveToolsFrameworkModule.h
--- a/Engine/Source/Editor/Experimental/EditorInteractiveToolsFramework/Public/EditorInteractiveToolsFrameworkModule.h
+++ b/Engine/Source/Editor/Experimental/EditorInteractiveToolsFramework/Public/EditorInteractiveToolsFrameworkModule.h
@@ -7,13 +7,15 @@
 #include "Modules/ModuleInterface.h"
 #include "Modules/ModuleManager.h"
 
+#define UE_API EDITORINTERACTIVETOOLSFRAMEWORK_API
 
-class EDITORINTERACTIVETOOLSFRAMEWORK_API FEditorInteractiveToolsFrameworkGlobals
+
+class FEditorInteractiveToolsFrameworkGlobals
 {
 public:
 	// This is the key returned by AddComponentTargetFactory() for the FStaticMeshComponentTargetFactory created/registered
 	// in StartupModule() below. Use this key to find/remove that module registration if you need to.
-	static int32 RegisteredStaticMeshTargetFactoryKey;
+	static UE_API int32 RegisteredStaticMeshTargetFactoryKey;
 };
 
 class FEditorInteractiveToolsFrameworkModule : public IModuleInterface
@@ -24,3 +26,5 @@ public:
 	virtual void StartupModule() override;
 	virtual void ShutdownModule() override;
 };
+
+#undef UE_API
diff --git a/Engine/Source/Editor/Experimental/EditorInteractiveToolsFramework/Public/GizmoEdModeInterface.h b/Engine/Source/Editor/Experimental/EditorInteractiveToolsFramework/Public/GizmoEdModeInterface.h
--- a/Engine/Source/Editor/Experimental/EditorInteractiveToolsFramework/Public/GizmoEdModeInterface.h
+++ b/Engine/Source/Editor/Experimental/EditorInteractiveToolsFramework/Public/GizmoEdModeInterface.h
@@ -24,7 +24,7 @@ class UGizmoEdModeInterface : public UInterface
 	GENERATED_BODY()
 };
 
-class EDITORINTERACTIVETOOLSFRAMEWORK_API IGizmoEdModeInterface
+class IGizmoEdModeInterface
 {
 	GENERATED_BODY()
 
@@ -32,4 +32,4 @@ public:
 	
 	virtual bool BeginTransform(const FGizmoState& InState) = 0;
 	virtual bool EndTransform(const FGizmoState& InState) = 0;
-};
\ No newline at end of file
+};
diff --git a/Engine/Source/Editor/Experimental/EditorInteractiveToolsFramework/Public/InteractiveToolStack.h b/Engine/Source/Editor/Experimental/EditorInteractiveToolsFramework/Public/InteractiveToolStack.h
--- a/Engine/Source/Editor/Experimental/EditorInteractiveToolsFramework/Public/InteractiveToolStack.h
+++ b/Engine/Source/Editor/Experimental/EditorInteractiveToolsFramework/Public/InteractiveToolStack.h
@@ -20,6 +20,8 @@
 
 #include "InteractiveToolStack.generated.h"
 
+#define UE_API EDITORINTERACTIVETOOLSFRAMEWORK_API
+
 class FText;
 class FUICommandInfo;
 class UEdMode;
@@ -37,7 +39,7 @@ struct FSlateIcon;
  * See: 'UWidgetEditorToolPaletteMode::Enter' as reference.
  */
 USTRUCT()
-struct EDITORINTERACTIVETOOLSFRAMEWORK_API FInteractiveToolStack
+struct FInteractiveToolStack
 {
 	GENERATED_BODY()
 
@@ -51,7 +53,7 @@ public:
 	 * @param UICommand				Command to map tool start / stop actions to
 	 * @param ToolStackIdentifier	Unique string identifier for the tool stack
 	 */
-	static void RegisterToolStack(UEdMode* EdMode, UEditorInteractiveToolsContext* UseToolsContext, TSharedPtr<FUICommandInfo> UICommand, FString ToolStackIdentifier);
+	static UE_API void RegisterToolStack(UEdMode* EdMode, UEditorInteractiveToolsContext* UseToolsContext, TSharedPtr<FUICommandInfo> UICommand, FString ToolStackIdentifier);
 
 	/**
 	 * Helper method to add tools to stacks
@@ -63,33 +65,33 @@ public:
 	 * @param ToolIdentifier		Unique string identifier for the tool, used to check if tool is active
 	 * @param Builder				Builder for tool to be used by actions
 	 */
-	static void AddToolToStack(UEdMode* EdMode, UEditorInteractiveToolsContext* UseToolsContext, TSharedPtr<FUICommandInfo> UICommand, FString ToolStackIdentifier, FString ToolIdentifier, UInteractiveToolBuilder* Builder);
+	static UE_API void AddToolToStack(UEdMode* EdMode, UEditorInteractiveToolsContext* UseToolsContext, TSharedPtr<FUICommandInfo> UICommand, FString ToolStackIdentifier, FString ToolIdentifier, UInteractiveToolBuilder* Builder);
 
 public:
 
 	/** Add the tool to our internal tool array */
-	int32 AddTool(const FString& InToolName, const TSharedPtr<FUICommandInfo>& InUICommand);
+	UE_API int32 AddTool(const FString& InToolName, const TSharedPtr<FUICommandInfo>& InUICommand);
 
 	/** Signal that a tool was activated increasing its priority in the toolstack for future chord inputs */
-	void NotifyToolActivated(const FString& InToolName);
+	UE_API void NotifyToolActivated(const FString& InToolName);
 
 	/** Get Tool that should be activated for the given input chord */
-	const FString GetLastActiveToolByChord(const FInputChord& InInputChord) const;
+	UE_API const FString GetLastActiveToolByChord(const FInputChord& InInputChord) const;
 
 	/** Get most recently activated tool within this stack */
-	const FString GetLastActiveTool() const;
+	UE_API const FString GetLastActiveTool() const;
 
 	/** Get most recently activated tool's icon within this stack */
-	const FSlateIcon& GetLastActiveToolIcon() const;
+	UE_API const FSlateIcon& GetLastActiveToolIcon() const;
 
 	/** Get most recently activated tool's label within this stack */
-	const FText& GetLastActiveToolLabel() const;
+	UE_API const FText& GetLastActiveToolLabel() const;
 
 	/** Get most recently activated tool's description within this stack */
-	const FText& GetLastActiveToolDescription() const;
+	UE_API const FText& GetLastActiveToolDescription() const;
 
 	/** Get the original tool stack */
-	const TArray<TPair<FString, TSharedPtr<FUICommandInfo>>>& GetStack() const;
+	UE_API const TArray<TPair<FString, TSharedPtr<FUICommandInfo>>>& GetStack() const;
 
 protected:
 
@@ -98,4 +100,6 @@ protected:
 
 	/** Original stack ordering */
 	TArray<TPair<FString, TSharedPtr<FUICommandInfo>>> OriginalStack;
-};
\ No newline at end of file
+};
+
+#undef UE_API
diff --git a/Engine/Source/Editor/Experimental/EditorInteractiveToolsFramework/Public/ToolContexts/ToolStackContext.h b/Engine/Source/Editor/Experimental/EditorInteractiveToolsFramework/Public/ToolContexts/ToolStackContext.h
--- a/Engine/Source/Editor/Experimental/EditorInteractiveToolsFramework/Public/ToolContexts/ToolStackContext.h
+++ b/Engine/Source/Editor/Experimental/EditorInteractiveToolsFramework/Public/ToolContexts/ToolStackContext.h
@@ -17,6 +17,8 @@
 
 #include "ToolStackContext.generated.h"
 
+#define UE_API EDITORINTERACTIVETOOLSFRAMEWORK_API
+
 class FUICommandInfo;
 class UInteractiveToolBuilder;
 struct FInteractiveToolStack;
@@ -24,12 +26,12 @@ struct FInteractiveToolStack;
 /**
  * Context needed to support tool stacks, add this to your ContextObjectStore if you want to use tool stacks
  */
-UCLASS(Transient)
-class EDITORINTERACTIVETOOLSFRAMEWORK_API UToolStackContext : public UObject
+UCLASS(MinimalAPI, Transient)
+class UToolStackContext : public UObject
 {
 	GENERATED_BODY()
 
-	virtual ~UToolStackContext();
+	UE_API virtual ~UToolStackContext();
 
 public:
 	/**
@@ -40,7 +42,7 @@ public:
 	 * @param	ToolStackIdentifier	Unique string identifier for the tool stack
 	 * @param	ToolScope			Scope to determine appropriate tool context for stack
 	 */
-	virtual void RegisterToolStack(TSharedPtr<FUICommandInfo> UICommand, const FString& ToolStackIdentifier, EToolsContextScope ToolScope = EToolsContextScope::Default);
+	UE_API virtual void RegisterToolStack(TSharedPtr<FUICommandInfo> UICommand, const FString& ToolStackIdentifier, EToolsContextScope ToolScope = EToolsContextScope::Default);
 
 	/**
 	 * Registers and maps the provided UI command to actions that control a tool within a tool stack
@@ -51,21 +53,21 @@ public:
 	 * @param	Builder				Builder for tool to be used by actions
 	 * @param	ToolScope			Scope to determine lifetime of tool (Editor, Mode, etc)
 	 */
-	virtual void AddToolToStack(TSharedPtr<FUICommandInfo> UICommand, const FString& ToolStackIdentifier, const FString& ToolIdentifier, UInteractiveToolBuilder* Builder, EToolsContextScope ToolScope = EToolsContextScope::Default);
+	UE_API virtual void AddToolToStack(TSharedPtr<FUICommandInfo> UICommand, const FString& ToolStackIdentifier, const FString& ToolIdentifier, UInteractiveToolBuilder* Builder, EToolsContextScope ToolScope = EToolsContextScope::Default);
 
 	/**
 	 * Unregister a Tool stack, unregistration of individual tools is handled seperately
 	 * Note: Currently this is never called since we properly free resources on destruction
 	 * @param Identifier string used to identify this stack
 	 */
-	virtual void UnregisterToolStack(const FString& Identifier);
+	UE_API virtual void UnregisterToolStack(const FString& Identifier);
 
 	/**
 	 * Access an existing tool stack
 	 * @param	Identifier string used to identify this stack
 	 * @return	Ptr to stack if it exists, nullptr otherwise
 	 */
-	virtual FInteractiveToolStack* AccessToolStack(const FString& Identifier);
+	UE_API virtual FInteractiveToolStack* AccessToolStack(const FString& Identifier);
 
 public:
 	/** EdMode owning the tools referenced by this stack */
@@ -86,3 +88,5 @@ protected:
 	/** List of Tools stacks this Mode has registered in the EditorToolsContext, to be unregistered on Mode shutdown */
 	TArray<TPair<TSharedPtr<FUICommandInfo>, FString>> RegisteredEditorToolStacks;
 };
+
+#undef UE_API
diff --git a/Engine/Source/Editor/Experimental/EditorInteractiveToolsFramework/Public/ToolContexts/WidgetToolsContext.h b/Engine/Source/Editor/Experimental/EditorInteractiveToolsFramework/Public/ToolContexts/WidgetToolsContext.h
--- a/Engine/Source/Editor/Experimental/EditorInteractiveToolsFramework/Public/ToolContexts/WidgetToolsContext.h
+++ b/Engine/Source/Editor/Experimental/EditorInteractiveToolsFramework/Public/ToolContexts/WidgetToolsContext.h
@@ -12,6 +12,8 @@
 
 #include "WidgetToolsContext.generated.h"
 
+#define UE_API EDITORINTERACTIVETOOLSFRAMEWORK_API
+
 class FPaintArgs;
 class FSlateRect;
 class FSlateWindowElementList;
@@ -27,24 +29,26 @@ struct FPointerEvent;
  * UWidgetToolsContext extends UModeManagerInteractiveToolsContext with methods needed for
  * tools operating on general widgets that do not have a viewport.
  */
-UCLASS(Transient)
-class EDITORINTERACTIVETOOLSFRAMEWORK_API UWidgetToolsContext : public UModeManagerInteractiveToolsContext
+UCLASS(MinimalAPI, Transient)
+class UWidgetToolsContext : public UModeManagerInteractiveToolsContext
 {
 	GENERATED_BODY()
 
 public:
 
 	// ~Begin SWidget "Proxy" Interface
-	bool OnKeyChar(const FGeometry& MyGeometry, const FCharacterEvent& InCharacterEvent);
-	bool OnKeyDown(const FGeometry& MyGeometry, const FKeyEvent& InKeyEvent);
-	bool OnKeyUp(const FGeometry& MyGeometry, const FKeyEvent& InKeyEvent);
-	bool OnMouseButtonDown(const FGeometry& MyGeometry, const FPointerEvent& MouseEvent);
-	bool OnMouseButtonUp(const FGeometry& MyGeometry, const FPointerEvent& MouseEvent);
-	bool OnMouseMove(const FGeometry& MyGeometry, const FPointerEvent& MouseEvent);
-	bool OnMouseWheel(const FGeometry& MyGeometry, const FPointerEvent& MouseEvent);
-	bool OnMouseButtonDoubleClick(const FGeometry& InMyGeometry, const FPointerEvent& InMouseEvent);
-	bool OnDragDetected(const FGeometry& MyGeometry, const FPointerEvent& MouseEvent);
-	void OnMouseCaptureLost(const FCaptureLostEvent& CaptureLostEvent);
-	int32 OnPaint(const FPaintArgs& Args, const FGeometry& AllottedGeometry, const FSlateRect& MyCullingRect, FSlateWindowElementList& OutDrawElements, int32 LayerId, const FWidgetStyle& InWidgetStyle, bool bParentEnabled) const;
+	UE_API bool OnKeyChar(const FGeometry& MyGeometry, const FCharacterEvent& InCharacterEvent);
+	UE_API bool OnKeyDown(const FGeometry& MyGeometry, const FKeyEvent& InKeyEvent);
+	UE_API bool OnKeyUp(const FGeometry& MyGeometry, const FKeyEvent& InKeyEvent);
+	UE_API bool OnMouseButtonDown(const FGeometry& MyGeometry, const FPointerEvent& MouseEvent);
+	UE_API bool OnMouseButtonUp(const FGeometry& MyGeometry, const FPointerEvent& MouseEvent);
+	UE_API bool OnMouseMove(const FGeometry& MyGeometry, const FPointerEvent& MouseEvent);
+	UE_API bool OnMouseWheel(const FGeometry& MyGeometry, const FPointerEvent& MouseEvent);
+	UE_API bool OnMouseButtonDoubleClick(const FGeometry& InMyGeometry, const FPointerEvent& InMouseEvent);
+	UE_API bool OnDragDetected(const FGeometry& MyGeometry, const FPointerEvent& MouseEvent);
+	UE_API void OnMouseCaptureLost(const FCaptureLostEvent& CaptureLostEvent);
+	UE_API int32 OnPaint(const FPaintArgs& Args, const FGeometry& AllottedGeometry, const FSlateRect& MyCullingRect, FSlateWindowElementList& OutDrawElements, int32 LayerId, const FWidgetStyle& InWidgetStyle, bool bParentEnabled) const;
 	// ~End SWidget "Proxy" Interface
 };
+
+#undef UE_API
diff --git a/Engine/Source/Editor/Experimental/EditorInteractiveToolsFramework/Public/TransformGizmoEditorSettings.h b/Engine/Source/Editor/Experimental/EditorInteractiveToolsFramework/Public/TransformGizmoEditorSettings.h
--- a/Engine/Source/Editor/Experimental/EditorInteractiveToolsFramework/Public/TransformGizmoEditorSettings.h
+++ b/Engine/Source/Editor/Experimental/EditorInteractiveToolsFramework/Public/TransformGizmoEditorSettings.h
@@ -1,4 +1,4 @@
-﻿// Copyright Epic Games, Inc. All Rights Reserved.
+// Copyright Epic Games, Inc. All Rights Reserved.
 
 #pragma once
 
@@ -7,17 +7,19 @@
 
 #include "TransformGizmoEditorSettings.generated.h"
 
-UCLASS(BlueprintType, Config = EditorPerProjectUserSettings, meta = (DisplayName = "Transform Gizmo"))
-class EDITORINTERACTIVETOOLSFRAMEWORK_API UTransformGizmoEditorSettings : public UDeveloperSettings
+#define UE_API EDITORINTERACTIVETOOLSFRAMEWORK_API
+
+UCLASS(MinimalAPI, BlueprintType, Config = EditorPerProjectUserSettings, meta = (DisplayName = "Transform Gizmo"))
+class UTransformGizmoEditorSettings : public UDeveloperSettings
 {
 	GENERATED_BODY()
 
 public:
-	UTransformGizmoEditorSettings(const FObjectInitializer& ObjectInitializer);
+	UE_API UTransformGizmoEditorSettings(const FObjectInitializer& ObjectInitializer);
 
 #if WITH_EDITOR
-	virtual void PostEditChangeChainProperty(FPropertyChangedChainEvent& InPropertyChangedEvent) override;
-	virtual void PostEditChangeProperty(FPropertyChangedEvent& InPropertyChangedEvent) override;
+	UE_API virtual void PostEditChangeChainProperty(FPropertyChangedChainEvent& InPropertyChangedEvent) override;
+	UE_API virtual void PostEditChangeProperty(FPropertyChangedEvent& InPropertyChangedEvent) override;
 #endif
 
 	UPROPERTY(Config, EditAnywhere, Category = "Transform Gizmo", DisplayName = "Transform Gizmo Size", meta = (ClampMin = "-10.0", ClampMax = "150.0", EditCondition = "!bUseExperimentalGizmo"))
@@ -61,16 +63,18 @@ public:
 	FGizmosParameters GizmosParameters;
 	//~ End Experimental Gizmo Category
 
-	void SetUseExperimentalGizmo(bool bInUseExperimentalGizmo);
-	bool UsesLegacyGizmo() const;
-	bool UsesNewTRSGizmo() const;
+	UE_API void SetUseExperimentalGizmo(bool bInUseExperimentalGizmo);
+	UE_API bool UsesLegacyGizmo() const;
+	UE_API bool UsesNewTRSGizmo() const;
 
-	void SetGizmosParameters(const FGizmosParameters& InGizmosParameters);
-	void SetTransformGizmoSize(float InTransformGizmoSize);
+	UE_API void SetGizmosParameters(const FGizmosParameters& InGizmosParameters);
+	UE_API void SetTransformGizmoSize(float InTransformGizmoSize);
 
 private:
-	void BroadcastNewTRSGizmoChange() const;
-	void BroadcastGizmosParametersChange() const;
+	UE_API void BroadcastNewTRSGizmoChange() const;
+	UE_API void BroadcastGizmosParametersChange() const;
 
-	void OnLegacySettingChanged(FName InPropertyName);
+	UE_API void OnLegacySettingChanged(FName InPropertyName);
 };
+
+#undef UE_API
diff --git a/Engine/Source/Editor/GameProjectGeneration/Classes/DefaultTemplateProjectDefs.h b/Engine/Source/Editor/GameProjectGeneration/Classes/DefaultTemplateProjectDefs.h
--- a/Engine/Source/Editor/GameProjectGeneration/Classes/DefaultTemplateProjectDefs.h
+++ b/Engine/Source/Editor/GameProjectGeneration/Classes/DefaultTemplateProjectDefs.h
@@ -11,16 +11,20 @@
 
 #include "DefaultTemplateProjectDefs.generated.h"
 
+#define UE_API GAMEPROJECTGENERATION_API
+
 class UObject;
 
-UCLASS()
-class GAMEPROJECTGENERATION_API UDefaultTemplateProjectDefs : public UTemplateProjectDefs
+UCLASS(MinimalAPI)
+class UDefaultTemplateProjectDefs : public UTemplateProjectDefs
 {
 	GENERATED_UCLASS_BODY()
 
-	virtual bool GeneratesCode(const FString& ProjectTemplatePath) const override;
+	UE_API virtual bool GeneratesCode(const FString& ProjectTemplatePath) const override;
 
-	virtual bool IsClassRename(const FString& DestFilename, const FString& SrcFilename, const FString& FileExtension) const override;
+	UE_API virtual bool IsClassRename(const FString& DestFilename, const FString& SrcFilename, const FString& FileExtension) const override;
 
-	virtual void AddConfigValues(TArray<FTemplateConfigValue>& ConfigValuesToSet, const FString& TemplateName, const FString& ProjectName, bool bShouldGenerateCode) const override;
+	UE_API virtual void AddConfigValues(TArray<FTemplateConfigValue>& ConfigValuesToSet, const FString& TemplateName, const FString& ProjectName, bool bShouldGenerateCode) const override;
 };
+
+#undef UE_API
diff --git a/Engine/Source/Editor/GameProjectGeneration/Public/ClassTemplateEditorSubsystem.h b/Engine/Source/Editor/GameProjectGeneration/Public/ClassTemplateEditorSubsystem.h
--- a/Engine/Source/Editor/GameProjectGeneration/Public/ClassTemplateEditorSubsystem.h
+++ b/Engine/Source/Editor/GameProjectGeneration/Public/ClassTemplateEditorSubsystem.h
@@ -14,47 +14,49 @@
 
 #include "ClassTemplateEditorSubsystem.generated.h"
 
+#define UE_API GAMEPROJECTGENERATION_API
+
 // Forward Declarations
 class FSubsystemCollectionBase;
 class FText;
 class UClass;
 
 
-UCLASS(Abstract)
-class GAMEPROJECTGENERATION_API UClassTemplate : public UObject
+UCLASS(MinimalAPI, Abstract)
+class UClassTemplate : public UObject
 {
 	GENERATED_BODY()
 
 public:
-	virtual void BeginDestroy() override;
+	UE_API virtual void BeginDestroy() override;
 
 	// Returns the directory containing the text file template for
 	// the given base generated class.
-	virtual FString GetDirectory() const;
+	UE_API virtual FString GetDirectory() const;
 
 	// Reads the header template text from disk.  If failure to read from
 	// disk, returns false and provides text reason.
-	bool ReadHeader(FString& OutHeaderFileText, FText& OutFailReason) const;
+	UE_API bool ReadHeader(FString& OutHeaderFileText, FText& OutFailReason) const;
 
 	// Reads the source template text from disk.  If failure to read from
 	// disk, returns false and provides text reason.
-	bool ReadSource(FString& OutSourceFileText, FText& OutFailReason) const;
+	UE_API bool ReadSource(FString& OutSourceFileText, FText& OutFailReason) const;
 
-	const UClass* GetGeneratedBaseClass() const;
+	UE_API const UClass* GetGeneratedBaseClass() const;
 
 protected:
 	// Sets the generated base class associated with the given template
-	void SetGeneratedBaseClass(UClass* InClass);
+	UE_API void SetGeneratedBaseClass(UClass* InClass);
 
 	// Returns the filename associated with the provided class template
 	// without an extension.  Defaults to class name.
-	virtual FString GetFilename() const;
+	UE_API virtual FString GetFilename() const;
 
 	// Returns full header filename including '.h.template' extension
-	FString GetHeaderFilename() const;
+	UE_API FString GetHeaderFilename() const;
 
 	// Returns full sourcefilename including '.cpp.template' extension
-	FString GetSourceFilename() const;
+	UE_API FString GetSourceFilename() const;
 
 private:
 	// Base UClass of which template class corresponds.
@@ -62,23 +64,23 @@ private:
 	TObjectPtr<UClass> GeneratedBaseClass;
 };
 
-UCLASS(Abstract)
-class GAMEPROJECTGENERATION_API UPluginClassTemplate : public UClassTemplate
+UCLASS(MinimalAPI, Abstract)
+class UPluginClassTemplate : public UClassTemplate
 {
 	GENERATED_BODY()
 
 public:
 	// Returns the directory containing the text file template for
 	// the given base generated class.
-	virtual FString GetDirectory() const override;
+	UE_API virtual FString GetDirectory() const override;
 
 protected:
 	UPROPERTY(Transient)
 	FString PluginName;
 };
 
-UCLASS()
-class GAMEPROJECTGENERATION_API UClassTemplateEditorSubsystem : public UEditorSubsystem
+UCLASS(MinimalAPI)
+class UClassTemplateEditorSubsystem : public UEditorSubsystem
 {
 	GENERATED_BODY()
 
@@ -87,29 +89,31 @@ private:
 	FTemplateRegistry TemplateRegistry;
 
 public:
-	UClassTemplateEditorSubsystem();
+	UE_API UClassTemplateEditorSubsystem();
 
-	virtual void Initialize(FSubsystemCollectionBase& Collection) override;
-	virtual void Deinitialize() override;
+	UE_API virtual void Initialize(FSubsystemCollectionBase& Collection) override;
+	UE_API virtual void Deinitialize() override;
 
 	// Registers all currently loaded template classes with the internal registry.
-	void RegisterTemplates();
+	UE_API void RegisterTemplates();
 
 	// Returns path to the directory containing all engine class templates.
-	static FString GetEngineTemplateDirectory();
+	static UE_API FString GetEngineTemplateDirectory();
 
 	// Returns whether or not class has registered template
-	bool ContainsClassTemplate(const UClass* InClass) const;
+	UE_API bool ContainsClassTemplate(const UClass* InClass) const;
 
 	// Returns class template if one is registered.
-	const UClassTemplate* FindClassTemplate(const UClass* InClass) const;
+	UE_API const UClassTemplate* FindClassTemplate(const UClass* InClass) const;
 
 	friend class UClassTemplate;
 
 private:
 	// Registers a template class with the subsystem
-	void Register(const UClassTemplate* InClassTemplate);
+	UE_API void Register(const UClassTemplate* InClassTemplate);
 
 	// Unregisters a template class with the subsystem
-	bool Unregister(const UClassTemplate* InClassTemplate);
+	UE_API bool Unregister(const UClassTemplate* InClassTemplate);
 };
+
+#undef UE_API
