[Editor ITF] Added DrawHUD support to UEditorTransformGizmo, with show flags guard

diff --git a/Engine/Source/Editor/Experimental/EditorInteractiveToolsFramework/Private/EditorGizmos/EditorTransformGizmo.cpp b/Engine/Source/Editor/Experimental/EditorInteractiveToolsFramework/Private/EditorGizmos/EditorTransformGizmo.cpp
--- a/Engine/Source/Editor/Experimental/EditorInteractiveToolsFramework/Private/EditorGizmos/EditorTransformGizmo.cpp
+++ b/Engine/Source/Editor/Experimental/EditorInteractiveToolsFramework/Private/EditorGizmos/EditorTransformGizmo.cpp
@@ -30,6 +30,22 @@ void UEditorTransformGizmo::Render(IToolsContextRenderAPI* RenderAPI)
 	}
 }
 
+void UEditorTransformGizmo::DrawHUD(FCanvas* Canvas, IToolsContextRenderAPI* RenderAPI)
+{
+	const FSceneView* SceneView = RenderAPI ? RenderAPI->GetSceneView() : nullptr;
+	const FSceneViewFamily* SceneViewFamily = SceneView ? SceneView->Family : nullptr;
+	if (!SceneViewFamily)
+	{
+		return;
+	}
+
+	const bool bEngineShowFlagsModeWidget = SceneViewFamily->EngineShowFlags.ModeWidgets && !SceneViewFamily->EngineShowFlags.Game;
+	if (bEngineShowFlagsModeWidget)
+	{
+		Super::DrawHUD(Canvas, RenderAPI);
+	}
+}
+
 void UEditorTransformGizmo::ApplyTranslateDelta(const FVector& InTranslateDelta)
 {
 	check(ActiveTarget);
diff --git a/Engine/Source/Editor/Experimental/EditorInteractiveToolsFramework/Public/EditorGizmos/EditorTransformGizmo.h b/Engine/Source/Editor/Experimental/EditorInteractiveToolsFramework/Public/EditorGizmos/EditorTransformGizmo.h
--- a/Engine/Source/Editor/Experimental/EditorInteractiveToolsFramework/Public/EditorGizmos/EditorTransformGizmo.h
+++ b/Engine/Source/Editor/Experimental/EditorInteractiveToolsFramework/Public/EditorGizmos/EditorTransformGizmo.h
@@ -25,6 +25,7 @@ class UEditorTransformGizmo : public UTransformGizmo
 
 	/**  UTransformGizmo override */
 	UE_API virtual void Render(IToolsContextRenderAPI* RenderAPI) override;
+	UE_API virtual void DrawHUD(FCanvas* Canvas, IToolsContextRenderAPI* RenderAPI) override;
 	
 protected:
 
