[Editor ITF] Remove level editor dependency from editor interactive tools framework module

diff --git a/Engine/Source/Editor/Experimental/EditorInteractiveToolsFramework/EditorInteractiveToolsFramework.Build.cs b/Engine/Source/Editor/Experimental/EditorInteractiveToolsFramework/EditorInteractiveToolsFramework.Build.cs
--- a/Engine/Source/Editor/Experimental/EditorInteractiveToolsFramework/EditorInteractiveToolsFramework.Build.cs
+++ b/Engine/Source/Editor/Experimental/EditorInteractiveToolsFramework/EditorInteractiveToolsFramework.Build.cs
@@ -15,7 +15,9 @@ public class EditorInteractiveToolsFramework : ModuleRules
 			new string[]
 			{
 				"Core",
-				"TypedElementFramework",
+				"InteractiveToolsFramework",
+				"TypedElementFramework", 
+				"UnrealEd",
 				// ... add other public dependencies that you statically link with here ...
 			}
 			);			
@@ -31,11 +33,8 @@ public class EditorInteractiveToolsFramework : ModuleRules
                 "SlateCore",
                 "InputCore",
 				"EditorFramework",
-				"UnrealEd",
                 "ContentBrowser",
-                "LevelEditor",
                 "ApplicationCore",
-                "InteractiveToolsFramework",
 				"MeshDescription",
 				"StaticMeshDescription",
                 "EditorSubsystem",
diff --git a/Engine/Source/Editor/Experimental/EditorInteractiveToolsFramework/Private/EditorDragTools/EditorBoxSelect.cpp b/Engine/Source/Editor/Experimental/EditorInteractiveToolsFramework/Private/EditorDragTools/EditorBoxSelect.cpp
--- a/Engine/Source/Editor/Experimental/EditorInteractiveToolsFramework/Private/EditorDragTools/EditorBoxSelect.cpp
+++ b/Engine/Source/Editor/Experimental/EditorInteractiveToolsFramework/Private/EditorDragTools/EditorBoxSelect.cpp
@@ -15,7 +15,6 @@
 #include "Engine/LevelStreaming.h"
 #include "EngineUtils.h"
 #include "GameFramework/Volume.h"
-#include "LevelEditorSubsystem.h"
 #include "LevelEditorViewport.h"
 #include "Model.h"
 #include "SceneView.h"
@@ -160,7 +159,7 @@ void FEditorBoxSelect::OnClickDrag(const FInputDeviceRay& InDragPos)
 	{
 		const bool bStrictDragSelection = GetDefault<ULevelEditorViewportSettings>()->bStrictBoxSelection;
 
-		UTypedElementSelectionSet* SelectionSet = GEditor->GetEditorSubsystem<ULevelEditorSubsystem>()->GetSelectionSet();
+		UTypedElementSelectionSet* SelectionSet = ModeTools->GetEditorSelectionSet();
 
 		UBrushEditingSubsystem* BrushSubsystem = GEditor->GetEditorSubsystem<UBrushEditingSubsystem>();
 		const bool bGeometryMode = BrushSubsystem ? BrushSubsystem->IsGeometryEditorModeActive() : false;
@@ -299,8 +298,7 @@ void FEditorBoxSelect::OnClickRelease(const FInputDeviceRay& InReleasePos)
 			{
 				const bool bMustEncompassEntireElement = IsWindowSelection();
 
-				UTypedElementSelectionSet* SelectionSet =
-					GEditor->GetEditorSubsystem<ULevelEditorSubsystem>()->GetSelectionSet();
+				UTypedElementSelectionSet* SelectionSet = ModeTools->GetEditorSelectionSet();
 				FTypedElementSelectionOptions ElementSelectionOption;
 				if (!bControlDown && !bShiftDown)
 				{
diff --git a/Engine/Source/Editor/Experimental/EditorInteractiveToolsFramework/Private/EditorDragTools/EditorFrustumSelect.cpp b/Engine/Source/Editor/Experimental/EditorInteractiveToolsFramework/Private/EditorDragTools/EditorFrustumSelect.cpp
--- a/Engine/Source/Editor/Experimental/EditorInteractiveToolsFramework/Private/EditorDragTools/EditorFrustumSelect.cpp
+++ b/Engine/Source/Editor/Experimental/EditorInteractiveToolsFramework/Private/EditorDragTools/EditorFrustumSelect.cpp
@@ -1,4 +1,4 @@
-﻿// Copyright Epic Games, Inc. All Rights Reserved.
+// Copyright Epic Games, Inc. All Rights Reserved.
 
 #include "EditorDragTools/EditorFrustumSelect.h"
 
@@ -15,7 +15,6 @@
 #include "Elements/Interfaces/TypedElementWorldInterface.h"
 #include "EngineUtils.h"
 #include "GameFramework/Volume.h"
-#include "LevelEditorSubsystem.h"
 #include "LevelEditorViewport.h"
 #include "Model.h"
 #include "SceneView.h"
@@ -232,7 +231,7 @@ void FEditorFrustumSelect::OnClickRelease(const FInputDeviceRay& InReleasePos)
 
 	if (!bEditorModeHandledSelection && !bComponentVisHandledSelection)
 	{
-		UTypedElementSelectionSet* SelectionSet = GEditor->GetEditorSubsystem<ULevelEditorSubsystem>()->GetSelectionSet();
+		UTypedElementSelectionSet* SelectionSet = ModeTools->GetEditorSelectionSet();
 		SelectionSet->Modify();
 
 		FWorldSelectionElementArgs SeletionArgs{ SelectionSet,
diff --git a/Engine/Source/Editor/Experimental/EditorInteractiveToolsFramework/Private/EditorDragTools/EditorMoveCameraWithObject.cpp b/Engine/Source/Editor/Experimental/EditorInteractiveToolsFramework/Private/EditorDragTools/EditorMoveCameraWithObject.cpp
--- a/Engine/Source/Editor/Experimental/EditorInteractiveToolsFramework/Private/EditorDragTools/EditorMoveCameraWithObject.cpp
+++ b/Engine/Source/Editor/Experimental/EditorInteractiveToolsFramework/Private/EditorDragTools/EditorMoveCameraWithObject.cpp
@@ -1,4 +1,4 @@
-﻿// Copyright Epic Games, Inc. All Rights Reserved.
+// Copyright Epic Games, Inc. All Rights Reserved.
 
 #include "EditorDragTools/EditorMoveCameraWithObject.h"
 
@@ -7,12 +7,10 @@
 #include "EditorGizmos/EditorTransformGizmo.h"
 #include "EditorGizmos/EditorTransformGizmoUtil.h"
 #include "EditorInteractiveGizmoManager.h"
-#include "ILevelEditor.h"
 #include "InteractiveGizmo.h"
 #include "InteractiveGizmoManager.h"
 #include "InteractiveToolManager.h"
 #include "LevelEditorViewport.h"
-#include "SLevelViewport.h"
 #include "Tools/EdModeInteractiveToolsContext.h"
 
 FEditorMoveCameraWithObject::FEditorMoveCameraWithObject(UInteractiveToolManager* InToolManager)
@@ -144,20 +142,8 @@ void FEditorMoveCameraWithObject::OnGizmoMovementEnd(UTransformProxy* InTransfor
 
 void FEditorMoveCameraWithObject::OnGizmoTransformChanged(UTransformProxy* InTransformProxy, UE::Math::TTransform<double> InTransform)
 {
-	if (bGizmoIsDragged && bModifierKeyIsPressed)
+	if (bGizmoIsDragged && bModifierKeyIsPressed && GCurrentLevelEditingViewportClient)
 	{
-		if (GCurrentLevelEditingViewportClient)
-		{
-			const FVector DragDelta = InTransform.GetLocation();
-
-			if (TSharedPtr<ILevelEditor> LevelEditor = GCurrentLevelEditingViewportClient->ParentLevelEditor.Pin())
-			{
-				if (TSharedPtr<SLevelViewport> ViewportInterface = LevelEditor->GetActiveViewportInterface())
-				{
-					FLevelEditorViewportClient& LevelViewportClient = ViewportInterface->GetLevelViewportClient();
-					LevelViewportClient.MoveViewportCamera(DragDelta, FRotator::ZeroRotator);
-				}
-			}
-		}
+		GCurrentLevelEditingViewportClient->MoveViewportCamera(InTransform.GetLocation(), FRotator::ZeroRotator);
 	}
 }
diff --git a/Engine/Source/Editor/Experimental/EditorInteractiveToolsFramework/Private/EditorGizmos/EditorTransformGizmoUtil.cpp b/Engine/Source/Editor/Experimental/EditorInteractiveToolsFramework/Private/EditorGizmos/EditorTransformGizmoUtil.cpp
--- a/Engine/Source/Editor/Experimental/EditorInteractiveToolsFramework/Private/EditorGizmos/EditorTransformGizmoUtil.cpp
+++ b/Engine/Source/Editor/Experimental/EditorInteractiveToolsFramework/Private/EditorGizmos/EditorTransformGizmoUtil.cpp
@@ -13,7 +13,7 @@
 #include "EditorGizmos/EditorTransformGizmoDataBinder.h"
 #include "EditorViewportClient.h"
 #include "EditorModes.h"
-#include "LevelEditor.h"
+#include "Editor.h"
 #include "Modules/ModuleManager.h"
 #include "Tools/DefaultEdMode.h"
 
@@ -403,16 +403,22 @@ bool UEditorTransformGizmoContextObject::SwapDefaultMode(const FEditorModeID InC
 
 	auto IsLevelEditorModeManager = [this]()
 	{
-		const FLevelEditorModule& LevelEditorModule = FModuleManager::GetModuleChecked<FLevelEditorModule>("LevelEditor");
-		const TSharedPtr<ILevelEditor> FirstLevelEditor = LevelEditorModule.GetFirstLevelEditor();
-		return FirstLevelEditor.IsValid() ? ModeTools == &FirstLevelEditor->GetEditorModeManager() : false;
+		if (FEditorViewportClient* FocusedViewportClient = ModeTools->GetFocusedViewportClient())
+		{
+			if (FocusedViewportClient->IsLevelEditorClient())
+			{
+				return FocusedViewportClient ? ModeTools == FocusedViewportClient->GetModeTools(): false;
+			}
+		}
+
+		return false;
 	};
-	
+
 	if (IsLevelEditorModeManager())
 	{
 		return false;
 	}
-	
+
 	// replace default mode
 	ModeTools->RemoveDefaultMode(InCurrentDefaultMode);
 	ModeTools->AddDefaultMode(InNewDefaultMode);
