[Editor ITF] Added ability to register/override the EditorTransformGizmo builder in UEditorInteractiveGizmoSubsystem

diff --git a/Engine/Source/Editor/Experimental/EditorInteractiveToolsFramework/Private/EditorInteractiveGizmoSubsystem.cpp b/Engine/Source/Editor/Experimental/EditorInteractiveToolsFramework/Private/EditorInteractiveGizmoSubsystem.cpp
--- a/Engine/Source/Editor/Experimental/EditorInteractiveToolsFramework/Private/EditorInteractiveGizmoSubsystem.cpp
+++ b/Engine/Source/Editor/Experimental/EditorInteractiveToolsFramework/Private/EditorInteractiveGizmoSubsystem.cpp
@@ -46,7 +46,7 @@ void UEditorInteractiveGizmoSubsystem::Deinitialize()
 void UEditorInteractiveGizmoSubsystem::RegisterBuiltinEditorGizmoTypes()
 {
 	// Setup gizmo transform builder
-	TransformGizmoBuilder = NewObject<UEditorTransformGizmoBuilder>();
+	RegisterTransformGizmoBuilder(NewObject<UEditorTransformGizmoBuilder>());
 
 	// Register built-in gizmo types here
 
@@ -72,6 +72,11 @@ void UEditorInteractiveGizmoSubsystem::GetQualifiedGlobalEditorGizmoBuilders(EEd
 	Registry->GetQualifiedEditorGizmoBuilders(InGizmoCategory, InToolBuilderState, InFoundBuilders);
 }
 
+void UEditorInteractiveGizmoSubsystem::RegisterTransformGizmoBuilder(UInteractiveGizmoBuilder* InTransformGizmoBuilder)
+{
+	TransformGizmoBuilder = InTransformGizmoBuilder;
+}
+
 void UEditorInteractiveGizmoSubsystem::DeregisterGlobalEditorGizmoType(EEditorGizmoCategory InGizmoCategory, UInteractiveGizmoBuilder* InGizmoBuilder)
 {
 	check(Registry);
diff --git a/Engine/Source/Editor/Experimental/EditorInteractiveToolsFramework/Public/EditorInteractiveGizmoSubsystem.h b/Engine/Source/Editor/Experimental/EditorInteractiveToolsFramework/Public/EditorInteractiveGizmoSubsystem.h
--- a/Engine/Source/Editor/Experimental/EditorInteractiveToolsFramework/Public/EditorInteractiveGizmoSubsystem.h
+++ b/Engine/Source/Editor/Experimental/EditorInteractiveToolsFramework/Public/EditorInteractiveGizmoSubsystem.h
@@ -21,6 +21,9 @@ class UInteractiveGizmoBuilder;
 class UObject;
 struct FToolBuilderState;
 
+// @NOTE: The following changes should be made (reverted) prior to 5.7 Release:
+// - Remove RegisterTransformGizmoBuilder
+
 /**
  * The InteractiveGizmoSubsystem provides methods for registering and unregistering 
  * selection-based gizmo builders. This subsystem will be queried for qualified 
@@ -100,6 +103,11 @@ public:
 	 */
 	UE_API void GetQualifiedGlobalEditorGizmoBuilders(EEditorGizmoCategory InGizmoCategory, const FToolBuilderState& InToolBuilderState, TArray<UInteractiveGizmoBuilder*>& InFoundBuilders);
 
+	/**
+	 * Register a transform gizmo builder used to build the Level Editor TRS gizmo, allowing the replacement of the default TRS gizmo builder.
+	 */
+ 	UE_API void RegisterTransformGizmoBuilder(UInteractiveGizmoBuilder* InTransformGizmoBuilder);
+ 	
 	/** 
 	 * Get transform gizmo builder used to build the Level Editor TRS gizmo. 
 	 */
