[Editor ITF] TransformGizmo tweaks:
- Virtualized some functions (for overriding)
- Added DPI support
- Replaced copy/paste function with actual (was since made public)
- Made overridable function for translation-only planar movement

diff --git a/Engine/Source/Editor/Experimental/EditorInteractiveToolsFramework/Private/EditorGizmos/TransformGizmo.cpp b/Engine/Source/Editor/Experimental/EditorInteractiveToolsFramework/Private/EditorGizmos/TransformGizmo.cpp
--- a/Engine/Source/Editor/Experimental/EditorInteractiveToolsFramework/Private/EditorGizmos/TransformGizmo.cpp
+++ b/Engine/Source/Editor/Experimental/EditorInteractiveToolsFramework/Private/EditorGizmos/TransformGizmo.cpp
@@ -62,35 +62,6 @@ static FAutoConsoleVariableRef CVarProjectIndirect(
 	ProjectIndirect,
 	TEXT("Project to the nearest point of the curve when handling indirect rotation.")
 	);
-
-//cf. GeometryCore implementation
-void MakePerpVectors(const FVector& Normal, FVector& OutPerp1, FVector& OutPerp2)
-{
-	// Duff et al method, from https://graphics.pixar.com/library/OrthonormalB/paper.pdf
-	if (Normal.Z < 0.0)
-	{
-		double A = 1.0 / (1.0 - Normal.Z);
-		double B = Normal.X * Normal.Y * A;
-		OutPerp1.X = 1.0 - Normal.X * Normal.X * A;
-		OutPerp1.Y = -B;
-		OutPerp1.Z = Normal.X;
-		OutPerp2.X = B;
-		OutPerp2.Y = Normal.Y * Normal.Y * A - 1.0;
-		OutPerp2.Z = -Normal.Y;
-	}
-	else
-	{
-		double A = 1.0 / (1.0 + Normal.Z);
-		double B = -Normal.X * Normal.Y * A;
-		OutPerp1.X = 1.0 - Normal.X * Normal.X * A;
-		OutPerp1.Y = B;
-		OutPerp1.Z = -Normal.X;
-		OutPerp2.X = B;
-		OutPerp2.Y = 1.0 - Normal.Y * Normal.Y * A;
-		OutPerp2.Z = -Normal.Y;
-	}
-}
-	
 }
 
 namespace UE::EditorInteractiveToolsFramework::Private
@@ -513,7 +484,7 @@ void UTransformGizmo::Render(IToolsContextRenderAPI* RenderAPI)
 		CurrentTransform = ActiveTarget->GetTransform();
 
 		UGizmoElementBase::FRenderTraversalState RenderState;
-		RenderState.Initialize(RenderAPI->GetSceneView(), GetGizmoTransform());
+		RenderState.Initialize(RenderAPI->GetSceneView(), GetGizmoTransform(), RenderAPI->GetCameraState().DPIScale);
 		GizmoElementRoot->Render(RenderAPI, RenderState);
 
 		RenderDebug(RenderAPI);
@@ -2499,7 +2470,7 @@ void UTransformGizmo::OnClickDragTranslateAxis(const FInputDeviceRay& DragPos)
 		const FVector2D YAxisDir = GetScreenProjectedAxis(GizmoViewContext, FVector::YAxisVector, CurrentTransform);
 		const FVector2D ZAxisDir = GetScreenProjectedAxis(GizmoViewContext, FVector::ZAxisVector, CurrentTransform);
 
-		const float PixelToWorldRatio = UE::GizmoRenderingUtil::CalculateLocalPixelToWorldScale(GizmoViewContext, CurrentTransform.GetLocation());
+		const float PixelToWorldRatio = UE::GizmoRenderingUtil::CalculateLocalPixelToWorldScale(GizmoViewContext, CurrentTransform.GetLocation(), true);
 		
 		FVector Delta((InteractionAxisList == EAxisList::X) ? PixelToWorldRatio * FVector2D::DotProduct(XAxisDir, DragDir) : 0.0,
 					  (InteractionAxisList == EAxisList::Y) ? PixelToWorldRatio * FVector2D::DotProduct(YAxisDir, DragDir) : 0.0,
@@ -2549,7 +2520,7 @@ void UTransformGizmo::OnClickPressTranslateXYPlanar(const FInputDeviceRay& Press
 	InteractionPlanarAxisX = GetWorldAxis(FVector::XAxisVector);
 	InteractionPlanarAxisY = GetWorldAxis(FVector::YAxisVector);
 	InteractionAxisList = EAxisList::XY;
-	OnClickPressPlanar(PressPos);
+	OnClickPressTranslatePlanar(PressPos);
 }
 
 void UTransformGizmo::OnClickPressTranslateYZPlanar(const FInputDeviceRay& PressPos)
@@ -2559,7 +2530,7 @@ void UTransformGizmo::OnClickPressTranslateYZPlanar(const FInputDeviceRay& Press
 	InteractionPlanarAxisX = GetWorldAxis(FVector::YAxisVector);
 	InteractionPlanarAxisY = GetWorldAxis(FVector::ZAxisVector);
 	InteractionAxisList = EAxisList::YZ;
-	OnClickPressPlanar(PressPos);
+	OnClickPressTranslatePlanar(PressPos);
 }
 
 void UTransformGizmo::OnClickPressTranslateXZPlanar(const FInputDeviceRay& PressPos)
@@ -2569,7 +2540,7 @@ void UTransformGizmo::OnClickPressTranslateXZPlanar(const FInputDeviceRay& Press
 	InteractionPlanarAxisX = GetWorldAxis(FVector::ZAxisVector);
 	InteractionPlanarAxisY = GetWorldAxis(FVector::XAxisVector);
 	InteractionAxisList = EAxisList::XZ;
-	OnClickPressPlanar(PressPos);
+	OnClickPressTranslatePlanar(PressPos);
 }
 
 void UTransformGizmo::OnClickPressPlanar(const FInputDeviceRay& PressPos)
@@ -2594,6 +2565,11 @@ void UTransformGizmo::OnClickPressPlanar(const FInputDeviceRay& PressPos)
 	}
 }
 
+void UTransformGizmo::OnClickPressTranslatePlanar(const FInputDeviceRay& PressPos)
+{
+	OnClickPressPlanar(PressPos);
+}
+
 void UTransformGizmo::OnClickDragTranslatePlanar(const FInputDeviceRay& DragPos)
 {
 	FVector::FReal HitDepth;
@@ -3091,6 +3067,7 @@ void UTransformGizmo::OnClickDragRotateAxis(const FInputDeviceRay& DragPos)
 		break;
 	}
 	case EAxisRotateMode::Arc:
+	case EAxisRotateMode::ScreenArc:
 	{
 		FVector::FReal HitDepth;
 		if (GetRayParamIntersectionWithInteractionPlane(DragPos, HitDepth))
@@ -3246,7 +3223,7 @@ void UTransformGizmo::OnClickPressGimbalRotateAxis(const FInputDeviceRay& InPres
 		const FVector WorldPlaneNormal = GetGimbalRotationAxis(RotateID);
 
 		FVector WorldPlaneAxis1, WorldPlaneAxis2;
-		GizmoLocals::MakePerpVectors(WorldPlaneNormal, WorldPlaneAxis1, WorldPlaneAxis2);
+		UE::Geometry::VectorUtil::MakePerpVectors(WorldPlaneNormal, WorldPlaneAxis1, WorldPlaneAxis2);
 		
 		const bool bCanRotateArc = OnClickPressRotateArc(InPressPos, WorldPlaneNormal, WorldPlaneAxis1, WorldPlaneAxis2);
 		if (bCanRotateArc)
@@ -3299,6 +3276,7 @@ void UTransformGizmo::OnClickDragGimbalRotateAxis(const FInputDeviceRay& InDragP
 			break;
 		}
 	case EAxisRotateMode::Arc:
+	case EAxisRotateMode::ScreenArc:
 		{
 			FVector::FReal HitDepth;
 			if (GetRayParamIntersectionWithInteractionPlane(InDragPos, HitDepth))
@@ -3637,7 +3615,7 @@ void UTransformGizmo::OnClickReleaseArcBallRotate(const FInputDeviceRay& Release
 
 float UTransformGizmo::GetWorldRadius(const float InRadius) const
 {
-	const float PixelToWorldScale = UE::GizmoRenderingUtil::CalculateLocalPixelToWorldScale(GizmoViewContext, CurrentTransform.GetLocation());
+	const float PixelToWorldScale = UE::GizmoRenderingUtil::CalculateLocalPixelToWorldScale(GizmoViewContext, CurrentTransform.GetLocation(), true);
 	const float GizmoScale = TransformGizmoSource ? TransformGizmoSource->GetGizmoScale() : 1.0f;
 	return InRadius * GetSizeCoefficient() * PixelToWorldScale * GizmoScale;
 }
diff --git a/Engine/Source/Editor/Experimental/EditorInteractiveToolsFramework/Public/EditorGizmos/TransformGizmo.h b/Engine/Source/Editor/Experimental/EditorInteractiveToolsFramework/Public/EditorGizmos/TransformGizmo.h
--- a/Engine/Source/Editor/Experimental/EditorInteractiveToolsFramework/Public/EditorGizmos/TransformGizmo.h
+++ b/Engine/Source/Editor/Experimental/EditorInteractiveToolsFramework/Public/EditorGizmos/TransformGizmo.h
@@ -57,6 +57,19 @@ class UObject;
 class UTransformProxy;
 class UGizmoElementGimbal;
 
+// @NOTE: The following changes should be made prior to 5.7 Release:
+// De-virtualize:
+//  HandleWidgetModeChanged
+//  OnParametersChanged
+//  SetupIndirectBehaviors
+//  UpdateRotationMode
+//  GetGimbalRotationAxis
+//  ResetHoverStates
+//  ResetInteractingStates
+//  OnClickPressRotateAxis
+//  ComputeGimbalRotateDelta
+//  OnClickPressGimbalRotateAxis
+
 //
 // Interface for the Transform gizmo.
 //
@@ -105,8 +118,9 @@ namespace EAxisRotateMode
 {
 	enum Type : uint8
 	{
-		Pull,
-		Arc
+		Pull = 0,
+		Arc = 1,
+		ScreenArc = 2, // Like Arc, but in screen space, rather than aligned to the associated axis
 	};
 }
 
@@ -296,12 +310,12 @@ public:
 	/**
 	 * Handle widget mode changed.
 	 */
-	UE_API void HandleWidgetModeChanged(UE::Widget::EWidgetMode InWidgetMode);
+	UE_API virtual void HandleWidgetModeChanged(UE::Widget::EWidgetMode InWidgetMode);
 
 	/**
 	 * Handle user parameters changes
 	 */
-	UE_API void OnParametersChanged(const FGizmosParameters& InParameters);
+	UE_API virtual void OnParametersChanged(const FGizmosParameters& InParameters);
 	
 public:
 
@@ -498,7 +512,7 @@ protected:
 	UE_API virtual void SetupBehaviors();
 
 	/** Setup indirect behaviors */
-	UE_API void SetupIndirectBehaviors();
+	UE_API virtual void SetupIndirectBehaviors();
 
 	/** Setup materials */
 	UE_API virtual void SetupMaterials();
@@ -510,7 +524,7 @@ protected:
 	UE_API virtual void UpdateMode();
 
 	/** Update current gizmo rotation mode (default or gimbal)*/
-	UE_API void UpdateRotationMode();
+	UE_API virtual void UpdateRotationMode();
 
 	/** Enable the given mode with the specified axes, EAxisList::Type::None will hide objects associated with mode */
 	UE_API virtual void EnableMode(EGizmoTransformMode InGizmoMode, EAxisList::Type InAxisListToDraw);
@@ -608,7 +622,7 @@ protected:
 	UE_API virtual FVector GetWorldAxis(const FVector& InAxis) const;
 
 	/** Get gimbal rotation axis (InAxis is expected to be in [0,2] range) */
-	UE_API FVector GetGimbalRotationAxis(const int32 InAxis) const;
+	UE_API virtual FVector GetGimbalRotationAxis(const int32 InAxis) const;
 
 	/** Get nearest param along input ray to the current interaction axis */
 	UE_API virtual float GetNearestRayParamToInteractionAxis(const FInputDeviceRay& InRay);
@@ -620,13 +634,13 @@ protected:
 	UE_API virtual void UpdateHoverState(const bool bInHover, const ETransformGizmoPartIdentifier InPartId);
 
 	/** Reset all hover states related to the transform mode to false */
-	UE_API void ResetHoverStates(const EGizmoTransformMode InMode);
+	UE_API virtual void ResetHoverStates(const EGizmoTransformMode InMode);
 
 	/** Update interacting state for given part id */
 	UE_API virtual void UpdateInteractingState(const bool bInInteracting, const ETransformGizmoPartIdentifier InPartId, const bool bIdOnly = false);
 
 	/** Reset all interacting states related to the transform mode to false */
-	UE_API void ResetInteractingStates(const EGizmoTransformMode InMode);
+	UE_API virtual void ResetInteractingStates(const EGizmoTransformMode InMode);
 
 	/** Called at the start of a sequence of gizmo transform edits */
 	UE_API virtual void BeginTransformEditSequence();
@@ -672,6 +686,9 @@ protected:
 	/** Handle click press for generic planar */
 	UE_API virtual void OnClickPressPlanar(const FInputDeviceRay& PressPos);
 
+	/** Handle click press for translate planar */
+	UE_API virtual void OnClickPressTranslatePlanar(const FInputDeviceRay& PressPos);
+
 	/** Handle click drag for translate planar */
 	UE_API virtual void OnClickDragTranslatePlanar(const FInputDeviceRay& DragPos);
 
@@ -711,7 +728,7 @@ protected:
 	UE_API virtual void OnClickPressRotateZAxis(const FInputDeviceRay& InPressPos);
 
 	/** Handle click press for any rotate axis */
-	UE_API void OnClickPressRotateAxis(const FInputDeviceRay& InPressPos);
+	UE_API virtual void OnClickPressRotateAxis(const FInputDeviceRay& InPressPos);
 
 	/** Handle click drag for rotate axis */
 	UE_API virtual void OnClickDragRotateAxis(const FInputDeviceRay& InDragPos);
@@ -735,7 +752,7 @@ protected:
 	UE_API virtual FQuat ComputeAxisRotateDelta(const double& InDeltaAngle);
 
 	/** Compute gimbal rotate delta based on screen-space start/end positions */
-	UE_API FQuat ComputeGimbalRotateDelta(const FVector2D& InStartPos, const FVector2D& InEndPos);
+	UE_API virtual FQuat ComputeGimbalRotateDelta(const FVector2D& InStartPos, const FVector2D& InEndPos);
 
 	/** Prepares data for arc rotation. This will return false if this is not possible (the rotate handle is perpendicular to the view) */
 	UE_API bool OnClickPressRotateArc( const FInputDeviceRay& InPressPos,
@@ -747,7 +764,7 @@ protected:
 	 */
 
 	/** Handle click press for any gimbal rotate axis */
-	UE_API void OnClickPressGimbalRotateAxis(const FInputDeviceRay& InPressPos);
+	UE_API virtual void OnClickPressGimbalRotateAxis(const FInputDeviceRay& InPressPos);
 
 	/** Handle click drag for gimbal rotate axis */
 	UE_API virtual void OnClickDragGimbalRotateAxis(const FInputDeviceRay& InDragPos);
