[Editor ITF] Added GetAxisElement to GizmoElementGimbal

diff --git a/Engine/Source/Editor/Experimental/EditorInteractiveToolsFramework/Public/EditorGizmos/GizmoElementGimbal.h b/Engine/Source/Editor/Experimental/EditorInteractiveToolsFramework/Public/EditorGizmos/GizmoElementGimbal.h
--- a/Engine/Source/Editor/Experimental/EditorInteractiveToolsFramework/Public/EditorGizmos/GizmoElementGimbal.h
+++ b/Engine/Source/Editor/Experimental/EditorInteractiveToolsFramework/Public/EditorGizmos/GizmoElementGimbal.h
@@ -7,6 +7,8 @@
 
 #include "GizmoElementGimbal.generated.h"
 
+#define UE_API EDITORINTERACTIVETOOLSFRAMEWORK_API
+
 /**
  * Gimbal rotation group object intended to be used as part of 3D Rotation Gizmos
  * This group expects three rotation sub-elements
@@ -19,13 +21,28 @@ class UGizmoElementGimbal : public UGizmoElementGroup
 public:
 
 	//~ Begin UGizmoElementBase Interface.
-	EDITORINTERACTIVETOOLSFRAMEWORK_API virtual void Render(IToolsContextRenderAPI* InRenderAPI, const FRenderTraversalState& InRenderState) override;
-	EDITORINTERACTIVETOOLSFRAMEWORK_API virtual FInputRayHit LineTrace(const UGizmoViewContext* InViewContext, const FLineTraceTraversalState& InLineTraceState, const FVector& InRayOrigin, const FVector& InRayDirection) override;
+	UE_API virtual void Render(IToolsContextRenderAPI* InRenderAPI, const FRenderTraversalState& InRenderState) override;
+	UE_API virtual FInputRayHit LineTrace(const UGizmoViewContext* InViewContext, const FLineTraceTraversalState& InLineTraceState, const FVector& InRayOrigin, const FVector& InRayDirection) override;
 	//~ End UGizmoElementBase Interface.
 
 	// Add object to group.
-	EDITORINTERACTIVETOOLSFRAMEWORK_API virtual void Add(UGizmoElementBase* InElement) override;
-	
+	UE_API virtual void Add(UGizmoElementBase* InElement) override;
+
+	template <class ElementType UE_REQUIRES(std::is_base_of_v<UGizmoElementBase, ElementType>)>
+	ElementType* GetAxisElement(const EAxis::Type InAxis) const
+	{
+		// @note: this makes the huge assumption that the elements are in X, Y, Z order.
+		const int32 AxisIndex = static_cast<int32>(InAxis) - 1;
+		if (ensure(Elements.IsValidIndex(AxisIndex)))
+		{
+			return Cast<ElementType>(Elements[AxisIndex]);
+		}
+
+		return nullptr;
+	}
+
 	UPROPERTY()
 	FRotationContext RotationContext;
 };
+
+#undef UE_API
