[Editor ITF] Exposed IEditorViewportClientProxy to external callers

diff --git a/Engine/Source/Editor/Experimental/EditorInteractiveToolsFramework/Public/EditorDragTools/EditorViewportClientProxy.h b/Engine/Source/Editor/Experimental/EditorInteractiveToolsFramework/Public/EditorDragTools/EditorViewportClientProxy.h
--- a/Engine/Source/Editor/Experimental/EditorInteractiveToolsFramework/Public/EditorDragTools/EditorViewportClientProxy.h
+++ b/Engine/Source/Editor/Experimental/EditorInteractiveToolsFramework/Public/EditorDragTools/EditorViewportClientProxy.h
@@ -9,11 +9,13 @@ class FEditorViewportClient;
 class FLevelEditorViewportClient;
 struct FWorldSelectionElementArgs;
 
+#define UE_API EDITORINTERACTIVETOOLSFRAMEWORK_API
+
 class IEditorViewportClientProxy : public TSharedFromThis<IEditorViewportClientProxy>
 {
 public:
 	virtual ~IEditorViewportClientProxy() = default;
-	static IEditorViewportClientProxy* CreateViewportClientProxy(FEditorViewportClient* InViewportClient);
+	UE_API static IEditorViewportClientProxy* CreateViewportClientProxy(FEditorViewportClient* InViewportClient);
 	static bool IsMouseOnWidgetAxis(IEditorViewportClientProxy* InEditorViewportClientProxy);
 
 
@@ -50,3 +52,4 @@ private:
 	TSharedPtr<FLevelEditorViewportClient> LevelEditorViewportClient;
 };
 
+#undef UE_API
