[Editor ITF] Exposed some GizmoRotationUtil functionality for external callers

diff --git a/Engine/Source/Editor/Experimental/EditorInteractiveToolsFramework/Public/EditorGizmos/GizmoRotationUtil.h b/Engine/Source/Editor/Experimental/EditorInteractiveToolsFramework/Public/EditorGizmos/GizmoRotationUtil.h
--- a/Engine/Source/Editor/Experimental/EditorInteractiveToolsFramework/Public/EditorGizmos/GizmoRotationUtil.h
+++ b/Engine/Source/Editor/Experimental/EditorInteractiveToolsFramework/Public/EditorGizmos/GizmoRotationUtil.h
@@ -5,6 +5,8 @@
 #include "Math/Transform.h"
 #include "Elements/Interfaces/TypedElementWorldInterface.h"
 
+#define UE_API EDITORINTERACTIVETOOLSFRAMEWORK_API
+
 struct FRotationContext;
 struct FEulerTransform;
 
@@ -25,12 +27,12 @@ namespace UE::GizmoRotationUtil
 	/**
 	 * Decomposes InRotationContext.Rotation in respect to the rotation order to output the explicit per axis rotations.
 	 */
-	void DecomposeRotations(const FTransform& InTransform, const FRotationContext& InRotationContext, FRotationDecomposition& OutDecomposition);
+	UE_API void DecomposeRotations(const FTransform& InTransform, const FRotationContext& InRotationContext, FRotationDecomposition& OutDecomposition);
 
 	/**
 	 * Returns the explicit rotation axis in respect to InRotationContext.Rotation & InRotationContext.RotationOrder
 	 */
-	FVector GetRotationAxis(const FTransform& InTransform, const FRotationContext& InRotationContext, const int32 InAxis);
+	UE_API FVector GetRotationAxis(const FTransform& InTransform, const FRotationContext& InRotationContext, const int32 InAxis);
 }
 
 namespace UE::GizmoRotationUtil
@@ -109,7 +111,7 @@ namespace UE::GizmoRotationUtil
 	/**
  * Returns the relative transform of an object by storing its explicit relative rotation in the provided FRotationContext
  */
-	EDITORINTERACTIVETOOLSFRAMEWORK_API bool GetRelativeTransform(const TTypedElement<ITypedElementWorldInterface>& InElement, FTransform& OutTransform, FRotationContext& InOutRotationContext);
+	UE_API bool GetRelativeTransform(const TTypedElement<ITypedElementWorldInterface>& InElement, FTransform& OutTransform, FRotationContext& InOutRotationContext);
 }
 
-
+#undef UE_API
