[Editor ITF] Exposed some functionality in MultiButtonClickDragBehavior

diff --git a/Engine/Source/Editor/Experimental/EditorInteractiveToolsFramework/Public/Behaviors/MultiButtonClickDragBehavior.h b/Engine/Source/Editor/Experimental/EditorInteractiveToolsFramework/Public/Behaviors/MultiButtonClickDragBehavior.h
--- a/Engine/Source/Editor/Experimental/EditorInteractiveToolsFramework/Public/Behaviors/MultiButtonClickDragBehavior.h
+++ b/Engine/Source/Editor/Experimental/EditorInteractiveToolsFramework/Public/Behaviors/MultiButtonClickDragBehavior.h
@@ -8,6 +8,8 @@
 
 class UObject;
 
+#define UE_API EDITORINTERACTIVETOOLSFRAMEWORK_API
+
 /**
  * MultiButtonClickDragBehavior is an implementation of UClickDragInputBehavior that can manage several
  * mouse buttons at once and also implements IClickDragBehaviorTarget directly, via a set of local lambda functions.
@@ -25,7 +27,7 @@ protected:
 public:
 
 	// UClickDragInputBehavior  implementation
-	void Initialize();
+	UE_API void Initialize();
 
 	virtual FInputRayHit CanBeginClickDragSequence(const FInputDeviceRay& InPressPos) override;
 	virtual void OnClickPress(const FInputDeviceRay& InPressPos) override;
@@ -58,10 +60,10 @@ public:
 	TUniqueFunction<void(const FInputDeviceState& Input)> OnStateUpdated = [](const FInputDeviceState&) {};
 
 	/** enable InButton so that in can be captured */
-	void EnableButton(const FKey& InButton);
+	UE_API void EnableButton(const FKey& InButton);
 	
 	/** disable InButton so that it's not captured */
-	void DisableButton(const FKey& InButton);
+	UE_API void DisableButton(const FKey& InButton);
 
 protected:
 
@@ -83,3 +85,5 @@ protected:
 	/** mouse buttons that need to be captured. Left=0, Middle=1, Right=2*/
 	uint8 CapturedInputs[3] = {0 ,0, 0};
 };
+
+#undef UE_API
