[Transform Gizmos] Gizmo is shown even when nothing is selected

diff --git a/Engine/Source/Editor/Experimental/EditorInteractiveToolsFramework/Private/EditorGizmos/EditorTransformGizmoUtil.cpp b/Engine/Source/Editor/Experimental/EditorInteractiveToolsFramework/Private/EditorGizmos/EditorTransformGizmoUtil.cpp
--- a/Engine/Source/Editor/Experimental/EditorInteractiveToolsFramework/Private/EditorGizmos/EditorTransformGizmoUtil.cpp
+++ b/Engine/Source/Editor/Experimental/EditorInteractiveToolsFramework/Private/EditorGizmos/EditorTransformGizmoUtil.cpp
@@ -14,6 +14,7 @@
 #include "EditorViewportClient.h"
 #include "EditorModes.h"
 #include "Editor.h"
+#include "LevelEditorViewport.h"
 #include "Modules/ModuleManager.h"
 #include "Tools/DefaultEdMode.h"
 
@@ -363,6 +364,39 @@ void UEditorTransformGizmoContextObject::InitializeViewportsBinding()
     
 		// initialize
         OnViewportClientsChanged();
+
+		auto OnLevelViewportClientsChanged = [this]()
+		{
+			if (!ModeTools)
+			{
+				return;
+			}
+
+			const bool bIsLevelViewportClient =
+			GEditor->GetLevelViewportClients().ContainsByPredicate([this](const FLevelEditorViewportClient* ViewportClient)
+			{
+				return ViewportClient && ViewportClient->GetModeTools() == ModeTools;
+			});
+
+			if (bIsLevelViewportClient)
+			{
+				// ensure level default mode
+				ModeTools->RemoveDefaultMode(FAssetEdModes::EM_AssetDefault);
+				ModeTools->AddDefaultMode(FBuiltinEditorModes::EM_Default);
+
+				if (ModeTools->IsModeActive(FAssetEdModes::EM_AssetDefault))
+				{
+					ModeTools->DeactivateMode(FAssetEdModes::EM_AssetDefault);
+					ModeTools->ActivateMode(FBuiltinEditorModes::EM_Default);
+				}
+			}
+		};
+
+		if (!LevelViewportClientsChangedHandle.IsValid())
+		{
+			LevelViewportClientsChangedHandle = GEditor->OnLevelViewportClientListChanged().AddLambda(OnLevelViewportClientsChanged);
+		}
+		OnLevelViewportClientsChanged();
 	}
 }
 
@@ -391,6 +425,15 @@ void UEditorTransformGizmoContextObject::RemoveViewportsBinding()
 		}
 		ViewportClientsChangedHandle.Reset();
 	}
+
+	if (LevelViewportClientsChangedHandle.IsValid())
+	{
+		if (ensure(GEditor))
+		{
+			GEditor->OnLevelViewportClientListChanged().Remove(LevelViewportClientsChangedHandle);
+		}
+		LevelViewportClientsChangedHandle.Reset();
+	}
 }
 
 bool UEditorTransformGizmoContextObject::SwapDefaultMode(const FEditorModeID InCurrentDefaultMode, const FEditorModeID InNewDefaultMode) const
@@ -401,24 +444,6 @@ bool UEditorTransformGizmoContextObject::SwapDefaultMode(const FEditorModeID InC
 		return false;
 	}
 
-	auto IsLevelEditorModeManager = [this]()
-	{
-		if (FEditorViewportClient* FocusedViewportClient = ModeTools->GetFocusedViewportClient())
-		{
-			if (FocusedViewportClient->IsLevelEditorClient())
-			{
-				return FocusedViewportClient ? ModeTools == FocusedViewportClient->GetModeTools(): false;
-			}
-		}
-
-		return false;
-	};
-
-	if (IsLevelEditorModeManager())
-	{
-		return false;
-	}
-
 	// replace default mode
 	ModeTools->RemoveDefaultMode(InCurrentDefaultMode);
 	ModeTools->AddDefaultMode(InNewDefaultMode);
diff --git a/Engine/Source/Editor/Experimental/EditorInteractiveToolsFramework/Public/EditorGizmos/EditorTransformGizmoUtil.h b/Engine/Source/Editor/Experimental/EditorInteractiveToolsFramework/Public/EditorGizmos/EditorTransformGizmoUtil.h
--- a/Engine/Source/Editor/Experimental/EditorInteractiveToolsFramework/Public/EditorGizmos/EditorTransformGizmoUtil.h
+++ b/Engine/Source/Editor/Experimental/EditorInteractiveToolsFramework/Public/EditorGizmos/EditorTransformGizmoUtil.h
@@ -112,6 +112,7 @@ private:
 	
 	FEditorModeTools* ModeTools = nullptr;
 	FDelegateHandle ViewportClientsChangedHandle;
+	FDelegateHandle LevelViewportClientsChangedHandle;
 	FDelegateHandle UseNewGizmosChangedHandled;
 	FDelegateHandle GizmosParametersChangedHandle;
 	bool bGizmosRegistered = false;
