[ITF Viewport Interactions] Have fly mode keyboard inputs only captured when a mouse button is pressed

diff --git a/Engine/Source/Editor/Experimental/EditorInteractiveToolsFramework/Private/ViewportInteractions/ViewportCameraRotateInteraction.cpp b/Engine/Source/Editor/Experimental/EditorInteractiveToolsFramework/Private/ViewportInteractions/ViewportCameraRotateInteraction.cpp
--- a/Engine/Source/Editor/Experimental/EditorInteractiveToolsFramework/Private/ViewportInteractions/ViewportCameraRotateInteraction.cpp
+++ b/Engine/Source/Editor/Experimental/EditorInteractiveToolsFramework/Private/ViewportInteractions/ViewportCameraRotateInteraction.cpp
@@ -16,6 +16,30 @@ UViewportCameraRotateInteraction::UViewportCameraRotateInteraction()
 	KeyInputBehavior->Initialize(this, GetKeys());
 	KeyInputBehavior->SetDefaultPriority(UE::Editor::ViewportInteractions::VIEWPORT_INTERACTIONS_DEFAULT_PRIORITY);
 	KeyInputBehavior->bRequireAllKeys = false;
+
+	KeyInputBehavior->ModifierCheckFunc = [this](const FInputDeviceState& InputDeviceState)
+	{
+		bool bResult = true;
+
+		if (InputDeviceState.InputDevice == EInputDevices::Keyboard)
+		{
+			using namespace UE::Editor::ViewportInteractions;
+
+			const FKey& ActiveKey = InputDeviceState.Keyboard.ActiveKey.Button;
+
+			// These commands should only be capture if a mouse button is down (not mapped by default)
+			if (CommandMatchesKey(FViewportNavigationCommands::Get().RotateUp, ActiveKey)
+				|| CommandMatchesKey(FViewportNavigationCommands::Get().RotateDown, ActiveKey)
+				|| CommandMatchesKey(FViewportNavigationCommands::Get().RotateLeft, ActiveKey)
+				|| CommandMatchesKey(FViewportNavigationCommands::Get().RotateRight, ActiveKey))
+			{
+				bResult = IsAnyMouseButtonDown();
+			}
+		}
+
+		return bResult;
+	};
+
 	KeyInputBehaviorWeak = KeyInputBehavior;
 
 	RegisterInputBehavior(KeyInputBehavior);
diff --git a/Engine/Source/Editor/Experimental/EditorInteractiveToolsFramework/Private/ViewportInteractions/ViewportCameraTranslateInteraction.cpp b/Engine/Source/Editor/Experimental/EditorInteractiveToolsFramework/Private/ViewportInteractions/ViewportCameraTranslateInteraction.cpp
--- a/Engine/Source/Editor/Experimental/EditorInteractiveToolsFramework/Private/ViewportInteractions/ViewportCameraTranslateInteraction.cpp
+++ b/Engine/Source/Editor/Experimental/EditorInteractiveToolsFramework/Private/ViewportInteractions/ViewportCameraTranslateInteraction.cpp
@@ -17,6 +17,31 @@ UViewportCameraTranslateInteraction::UViewportCameraTranslateInteraction()
 	KeyInputBehavior->Initialize(this, GetKeys());
 	KeyInputBehavior->SetDefaultPriority(UE::Editor::ViewportInteractions::VIEWPORT_INTERACTIONS_DEFAULT_PRIORITY);
 	KeyInputBehavior->bRequireAllKeys = false;
+	KeyInputBehavior->ModifierCheckFunc = [this](const FInputDeviceState& InputDeviceState) {
+		bool bResult = true;
+
+		if (InputDeviceState.InputDevice == EInputDevices::Keyboard)
+		{
+			using namespace UE::Editor::ViewportInteractions;
+
+			const FKey& ActiveKey = InputDeviceState.Keyboard.ActiveKey.Button;
+
+			// These commands should only be capture if a mouse button is down (defaults are W/A/S/D/R/F/Q/E)
+			if (CommandMatchesKey(FViewportNavigationCommands::Get().Forward, ActiveKey) ||
+				CommandMatchesKey(FViewportNavigationCommands::Get().Backward, ActiveKey) ||
+				CommandMatchesKey(FViewportNavigationCommands::Get().Right, ActiveKey) ||
+				CommandMatchesKey(FViewportNavigationCommands::Get().Left, ActiveKey) ||
+				CommandMatchesKey(FViewportNavigationCommands::Get().WorldUp, ActiveKey) ||
+				CommandMatchesKey(FViewportNavigationCommands::Get().WorldDown, ActiveKey) ||
+				CommandMatchesKey(FViewportNavigationCommands::Get().LocalUp, ActiveKey) ||
+				CommandMatchesKey(FViewportNavigationCommands::Get().LocalDown, ActiveKey))
+			{
+				bResult = IsAnyMouseButtonDown(); 
+			}
+		}
+
+		return bResult;
+	};
 	KeyInputBehaviorWeak = KeyInputBehavior;
 
 	RegisterInputBehavior(KeyInputBehavior);
diff --git a/Engine/Source/Editor/Experimental/EditorInteractiveToolsFramework/Private/ViewportInteractions/ViewportFOVInteraction.cpp b/Engine/Source/Editor/Experimental/EditorInteractiveToolsFramework/Private/ViewportInteractions/ViewportFOVInteraction.cpp
--- a/Engine/Source/Editor/Experimental/EditorInteractiveToolsFramework/Private/ViewportInteractions/ViewportFOVInteraction.cpp
+++ b/Engine/Source/Editor/Experimental/EditorInteractiveToolsFramework/Private/ViewportInteractions/ViewportFOVInteraction.cpp
@@ -20,15 +20,28 @@ UViewportFOVInteraction::UViewportFOVInteraction()
 	KeyInputBehavior->Initialize(this, GetKeys());
 	KeyInputBehavior->SetDefaultPriority(UE::Editor::ViewportInteractions::VIEWPORT_INTERACTIONS_DEFAULT_PRIORITY);
 	KeyInputBehavior->bRequireAllKeys = false;
-	KeyInputBehaviorWeak = KeyInputBehavior;
+	KeyInputBehavior->ModifierCheckFunc = [this](const FInputDeviceState& InputDeviceState) {
+		bool bResult = true;
+
+		if (InputDeviceState.InputDevice == EInputDevices::Keyboard)
+		{
+			using namespace UE::Editor::ViewportInteractions;
 
-	// Always on FOV controls (currently numpad based)
-	UKeyInputBehavior* NumpadKeyInputBehavior = NewObject<UKeyInputBehavior>();
-	NumpadKeyInputBehavior->Initialize(this, GetNumpadKeys());
-	NumpadKeyInputBehavior->SetDefaultPriority(UE::Editor::ViewportInteractions::VIEWPORT_INTERACTIONS_DEFAULT_PRIORITY);
-	NumpadKeyInputBehavior->bRequireAllKeys = false;
+			const FKey& ActiveKey = InputDeviceState.Keyboard.ActiveKey.Button;
+
+			// These commands should only be capture if a mouse button is down (defaults are Z/C)
+			if (CommandMatchesKey(FViewportNavigationCommands::Get().FovZoomOut, ActiveKey) ||
+				CommandMatchesKey(FViewportNavigationCommands::Get().FovZoomIn, ActiveKey))
+			{
+				bResult = IsAnyMouseButtonDown(); 
+			}
+		}
+
+		return bResult;
+	};
+	KeyInputBehaviorWeak = KeyInputBehavior;
 
-	RegisterInputBehaviors({ KeyInputBehavior, NumpadKeyInputBehavior });
+	RegisterInputBehavior(KeyInputBehavior);
 }
 
 UViewportFOVInteraction::~UViewportFOVInteraction()
@@ -117,6 +130,7 @@ TArray<FKey> UViewportFOVInteraction::GetKeys() const
 
 	// Hardcoded
 	Keys.Append({ EKeys::Add, EKeys::Subtract });
+	Keys.Append(GetNumpadKeys());
 
 	return Keys;
 }
