Expose polling current key state in UKeyInputBehavior.

Update ViewportCameraTranslateInteraction to a poll-on-input behavior, matching original functionality.

diff --git a/Engine/Source/Editor/Experimental/EditorInteractiveToolsFramework/Private/ViewportInteractions/ViewportCameraTranslateInteraction.cpp b/Engine/Source/Editor/Experimental/EditorInteractiveToolsFramework/Private/ViewportInteractions/ViewportCameraTranslateInteraction.cpp
--- a/Engine/Source/Editor/Experimental/EditorInteractiveToolsFramework/Private/ViewportInteractions/ViewportCameraTranslateInteraction.cpp
+++ b/Engine/Source/Editor/Experimental/EditorInteractiveToolsFramework/Private/ViewportInteractions/ViewportCameraTranslateInteraction.cpp
@@ -7,6 +7,23 @@
 #include "ViewportClientNavigationHelper.h"
 #include "ViewportInteractions/ViewportInteractionsBehaviorSource.h"
 
+namespace UE::Editor::ViewportInteractions::Private
+{
+	bool KeyBehaviorHasCommand(const UKeyInputBehavior* KeyBehavior, const TSharedPtr<FUICommandInfo>& Command)
+	{
+		for (int32 i = 0; i < static_cast<uint8>(EMultipleKeyBindingIndex::NumChords); ++i)
+		{
+			EMultipleKeyBindingIndex ChordIndex = static_cast<EMultipleKeyBindingIndex>(i);
+			if (KeyBehavior->IsKeyPressed(Command->GetActiveChord(ChordIndex)->Key))
+			{
+				return true;
+			}
+		}
+
+		return false;
+	}
+}
+
 UViewportCameraTranslateInteraction::UViewportCameraTranslateInteraction()
 {
 	ToolType = UE::Editor::ViewportInteractions::CameraTranslate;
@@ -36,7 +53,7 @@ UViewportCameraTranslateInteraction::UViewportCameraTranslateInteraction()
 				CommandMatchesKey(FViewportNavigationCommands::Get().LocalUp, ActiveKey) ||
 				CommandMatchesKey(FViewportNavigationCommands::Get().LocalDown, ActiveKey))
 			{
-				bResult = IsAnyMouseButtonDown(); 
+				bResult = IsMouseLooking(); 
 			}
 		}
 
@@ -47,33 +64,71 @@ UViewportCameraTranslateInteraction::UViewportCameraTranslateInteraction()
 	RegisterInputBehavior(KeyInputBehavior);
 }
 
-void UViewportCameraTranslateInteraction::OnKeyPressed(const FKey& InKeyID)
+void UViewportCameraTranslateInteraction::OnForceEndCapture()
 {
-	if (IsEnabled())
-	{
-		constexpr bool bIsPressed = true;
-		UpdateKeyState(InKeyID, bIsPressed);
-	}
 }
 
-void UViewportCameraTranslateInteraction::OnKeyReleased(const FKey& InKeyID)
+void UViewportCameraTranslateInteraction::Tick(float InDeltaTime) const
 {
-	if (IsEnabled())
+	using namespace UE::Editor::ViewportInteractions::Private;
+
+	constexpr float ImpulseValue = 1.0f;
+
+	const bool bIsMouseLooking = IsMouseLooking();
+	
+	float ForwardBackwardImpulse = 0.0f;
+	float RightLeftImpulse = 0.0f;
+	float WorldUpDownImpulse = 0.0f;
+	float LocalUpDownImpulse = 0.0f;
+	
+	// TODO: This could be better handled by either a bespoke behavior or a behavior that can track and
+	// expose multiple sets of keys at once (perhaps in a similar fashion as IModifierToggleBehaviorTarget
+	if (UKeyInputBehavior* KeyBehavior = KeyInputBehaviorWeak.Get())
 	{
-		constexpr bool bIsPressed = false;
-		UpdateKeyState(InKeyID, bIsPressed);
+		// Forward/Backward
+		if ((bIsMouseLooking && KeyBehaviorHasCommand(KeyBehavior, FViewportNavigationCommands::Get().Forward))
+			|| KeyBehavior->IsKeyPressed(EKeys::NumPadEight) || KeyBehavior->IsKeyPressed(EKeys::Up))
+		{
+			ForwardBackwardImpulse += ImpulseValue;
+		}
+		if ((bIsMouseLooking && KeyBehaviorHasCommand(KeyBehavior, FViewportNavigationCommands::Get().Backward))
+			|| KeyBehavior->IsKeyPressed(EKeys::NumPadTwo) || KeyBehavior->IsKeyPressed(EKeys::Down))
+		{
+			ForwardBackwardImpulse -= ImpulseValue;
+		}
+		// Right/Left
+		if ((KeyBehaviorHasCommand(KeyBehavior, FViewportNavigationCommands::Get().Right))
+			|| KeyBehavior->IsKeyPressed(EKeys::NumPadSix) || KeyBehavior->IsKeyPressed(EKeys::Right))
+		{
+			RightLeftImpulse += ImpulseValue;
+		}
+		if ((KeyBehaviorHasCommand(KeyBehavior, FViewportNavigationCommands::Get().Left))
+			|| KeyBehavior->IsKeyPressed(EKeys::NumPadFour) || KeyBehavior->IsKeyPressed(EKeys::Left))
+		{
+			RightLeftImpulse -= ImpulseValue;
+		}
+		// World Up/Down
+		if ((KeyBehaviorHasCommand(KeyBehavior, FViewportNavigationCommands::Get().WorldUp))
+			|| KeyBehavior->IsKeyPressed(EKeys::NumPadNine))
+		{
+			WorldUpDownImpulse += ImpulseValue;
+		}
+		if ((KeyBehaviorHasCommand(KeyBehavior, FViewportNavigationCommands::Get().WorldDown))
+			|| KeyBehavior->IsKeyPressed(EKeys::NumPadSeven))
+		{
+			WorldUpDownImpulse -= ImpulseValue;
+		}
+		// Local Up/Down
+		if (KeyBehaviorHasCommand(KeyBehavior, FViewportNavigationCommands::Get().LocalUp))
+		{
+			LocalUpDownImpulse += ImpulseValue;
+		}
+		if (KeyBehaviorHasCommand(KeyBehavior, FViewportNavigationCommands::Get().LocalDown))
+		{
+			LocalUpDownImpulse -= ImpulseValue;
+		}
 	}
-}
-void UViewportCameraTranslateInteraction::OnForceEndCapture()
-{
-	ForwardBackwardImpulse = 0.0f;
-	RightLeftImpulse = 0.0f;
-	WorldUpDownImpulse = 0.0f;
-	LocalUpDownImpulse = 0.0f;
-}
 
-void UViewportCameraTranslateInteraction::Tick(float InDeltaTime) const
-{
 	if (ForwardBackwardImpulse != 0 || RightLeftImpulse != 0 || WorldUpDownImpulse != 0 || LocalUpDownImpulse != 0)
 	{
 		if (FEditorViewportClient* EditorViewportClient = GetEditorViewportClient())
@@ -90,58 +145,6 @@ void UViewportCameraTranslateInteraction::Tick(float InDeltaTime) const
 	}
 }
 
-void UViewportCameraTranslateInteraction::UpdateKeyState(const FKey& InKeyID, bool bInIsPressed)
-{
-	using namespace UE::Editor::ViewportInteractions;
-
-	constexpr float ImpulseValue = 1.0f;
-
-	const bool bMouseButtonDownOrRelease = IsAnyMouseButtonDown() || !bInIsPressed;
-
-	// Forward/Backward
-	if ((CommandMatchesKey(FViewportNavigationCommands::Get().Forward, InKeyID) && bMouseButtonDownOrRelease)
-		|| InKeyID == EKeys::NumPadEight || InKeyID == EKeys::Up)
-	{
-		ForwardBackwardImpulse = bInIsPressed ? ImpulseValue : 0.0f;
-	}
-	else if ((CommandMatchesKey(FViewportNavigationCommands::Get().Backward, InKeyID) && bMouseButtonDownOrRelease)
-			 || InKeyID == EKeys::NumPadTwo || InKeyID == EKeys::Down)
-	{
-		ForwardBackwardImpulse = bInIsPressed ? -ImpulseValue : 0.0f;
-	}
-	// Right/Left
-	else if ((CommandMatchesKey(FViewportNavigationCommands::Get().Right, InKeyID) && bMouseButtonDownOrRelease)
-			 || InKeyID == EKeys::NumPadSix || InKeyID == EKeys::Right)
-	{
-		RightLeftImpulse = bInIsPressed ? ImpulseValue : 0.0f;
-	}
-	else if ((CommandMatchesKey(FViewportNavigationCommands::Get().Left, InKeyID) && bMouseButtonDownOrRelease)
-			 || InKeyID == EKeys::NumPadFour || InKeyID == EKeys::Left)
-	{
-		RightLeftImpulse = bInIsPressed ? -ImpulseValue : 0.0f;
-	}
-	// World Up/Down
-	else if ((CommandMatchesKey(FViewportNavigationCommands::Get().WorldUp, InKeyID) && bMouseButtonDownOrRelease)
-			 || InKeyID == EKeys::NumPadNine)
-	{
-		WorldUpDownImpulse = bInIsPressed ? ImpulseValue : 0.0f;
-	}
-	else if ((CommandMatchesKey(FViewportNavigationCommands::Get().WorldDown, InKeyID) && bMouseButtonDownOrRelease)
-			 || InKeyID == EKeys::NumPadSeven)
-	{
-		WorldUpDownImpulse = bInIsPressed ? -ImpulseValue : 0.0f;
-	}
-	// Local Up/Down
-	else if (CommandMatchesKey(FViewportNavigationCommands::Get().LocalUp, InKeyID) && bMouseButtonDownOrRelease)
-	{
-		LocalUpDownImpulse = bInIsPressed ? ImpulseValue : 0.0f;
-	}
-	else if (CommandMatchesKey(FViewportNavigationCommands::Get().LocalDown, InKeyID) && bMouseButtonDownOrRelease)
-	{
-		LocalUpDownImpulse = bInIsPressed ? -ImpulseValue : 0.0f;
-	}
-}
-
 TArray<FKey> UViewportCameraTranslateInteraction::GetKeys() const
 {
 	TArray<FKey> Keys = { EKeys::Add, EKeys::Subtract };
@@ -151,15 +154,12 @@ TArray<FKey> UViewportCameraTranslateInteraction::GetKeys() const
 	for (int32 i = 0; i < static_cast<uint8>(EMultipleKeyBindingIndex::NumChords); ++i)
 	{
 		EMultipleKeyBindingIndex ChordIndex = static_cast<EMultipleKeyBindingIndex>(i);
-		TArray<FKey> KeysFromChords;
 
 		// Retrieve command based keys
 		for (const TSharedPtr<FUICommandInfo>& Command : GetCommands())
 		{
-			KeysFromChords.Add(Command->GetActiveChord(ChordIndex)->Key);
+			Keys.Add(Command->GetActiveChord(ChordIndex)->Key);
 		}
-
-		Keys.Append(KeysFromChords);
 	}
 
 	// Hardcoded numpad keys
diff --git a/Engine/Source/Editor/Experimental/EditorInteractiveToolsFramework/Public/ViewportInteractions/ViewportCameraTranslateInteraction.h b/Engine/Source/Editor/Experimental/EditorInteractiveToolsFramework/Public/ViewportInteractions/ViewportCameraTranslateInteraction.h
--- a/Engine/Source/Editor/Experimental/EditorInteractiveToolsFramework/Public/ViewportInteractions/ViewportCameraTranslateInteraction.h
+++ b/Engine/Source/Editor/Experimental/EditorInteractiveToolsFramework/Public/ViewportInteractions/ViewportCameraTranslateInteraction.h
@@ -27,8 +27,8 @@ public:
 	UE_API UViewportCameraTranslateInteraction();
 
 	//~ Begin IKeyInputBehaviorTarget
-	UE_API virtual void OnKeyPressed(const FKey& InKeyID) override;
-	UE_API virtual void OnKeyReleased(const FKey& InKeyID) override;
+	UE_API virtual void OnKeyPressed(const FKey& InKeyID) override {}
+	UE_API virtual void OnKeyReleased(const FKey& InKeyID) override {}
 	UE_API virtual void OnForceEndCapture() override;
 	//~ End IKeyInputBehaviorTarget
 
@@ -42,18 +42,12 @@ protected:
 	UE_API virtual TArray<TSharedPtr<FUICommandInfo>> GetCommands() const override;
 	//~ End UViewportInteractionBase
 
-	UE_API void UpdateKeyState(const FKey& InKeyID, bool bInIsPressed);
 	UE_API TArray<FKey> GetKeys() const;
 
 private:
 	TWeakObjectPtr<UKeyInputBehavior> KeyInputBehaviorWeak;
 
 	bool bUseNumpadKey = true;
-
-	float ForwardBackwardImpulse = 0.0f;
-	float RightLeftImpulse = 0.0f;
-	float WorldUpDownImpulse = 0.0f;
-	float LocalUpDownImpulse = 0.0f;
 };
 
 #undef UE_API
diff --git a/Engine/Source/Runtime/InteractiveToolsFramework/Private/BaseBehaviors/KeyInputBehavior.cpp b/Engine/Source/Runtime/InteractiveToolsFramework/Private/BaseBehaviors/KeyInputBehavior.cpp
--- a/Engine/Source/Runtime/InteractiveToolsFramework/Private/BaseBehaviors/KeyInputBehavior.cpp
+++ b/Engine/Source/Runtime/InteractiveToolsFramework/Private/BaseBehaviors/KeyInputBehavior.cpp
@@ -95,6 +95,16 @@ void UKeyInputBehavior::ForceEndCapture(const FInputCaptureData& Data)
 	}
 }
 
+bool UKeyInputBehavior::IsKeyPressed(FKey InKey) const
+{
+	const int32 KeyIndex = TargetKeys.Find(InKey);
+	if (KeyIndex > -1)
+	{
+		return KeyActivations[KeyIndex];
+	}
+	return false;
+}
+
 void UKeyInputBehavior::InitializeKeyActivations()
 {
 	KeyActivations.SetNum(TargetKeys.Num());
diff --git a/Engine/Source/Runtime/InteractiveToolsFramework/Public/BaseBehaviors/KeyInputBehavior.h b/Engine/Source/Runtime/InteractiveToolsFramework/Public/BaseBehaviors/KeyInputBehavior.h
--- a/Engine/Source/Runtime/InteractiveToolsFramework/Public/BaseBehaviors/KeyInputBehavior.h
+++ b/Engine/Source/Runtime/InteractiveToolsFramework/Public/BaseBehaviors/KeyInputBehavior.h
@@ -79,6 +79,10 @@ public:
 	INTERACTIVETOOLSFRAMEWORK_API virtual FInputCaptureUpdate UpdateCapture(const FInputDeviceState& Input, const FInputCaptureData& Data) override;
 	INTERACTIVETOOLSFRAMEWORK_API virtual void ForceEndCapture(const FInputCaptureData& Data) override;
 
+	/**
+	 * @return Whether the given key is detected as pressed.
+	 */
+	INTERACTIVETOOLSFRAMEWORK_API bool IsKeyPressed(FKey InKey) const;
 
 	/**
     * The modifier set for this behavior
