[UnrealEd] Added ability to get a typed elements' Scale for use with a gizmo (to FTypedElementViewportInteractionCustomization)

diff --git a/Engine/Source/Editor/EditorFramework/Private/Elements/Framework/TypedElementViewportInteraction.cpp b/Engine/Source/Editor/EditorFramework/Private/Elements/Framework/TypedElementViewportInteraction.cpp
--- a/Engine/Source/Editor/EditorFramework/Private/Elements/Framework/TypedElementViewportInteraction.cpp
+++ b/Engine/Source/Editor/EditorFramework/Private/Elements/Framework/TypedElementViewportInteraction.cpp
@@ -167,6 +167,17 @@ bool FTypedElementViewportInteractionCustomization::GetFocusBounds(const TTypedE
 	return InElementWorldHandle.GetBounds(OutBounds);
 }
 
+bool FTypedElementViewportInteractionCustomization::GetScaleForGizmo(const TTypedElement<ITypedElementWorldInterface>& InElementWorldHandle, FVector& OutScale)
+{
+	FTransform WorldTransform;
+	if (InElementWorldHandle.GetWorldTransform(WorldTransform))
+	{
+		OutScale = WorldTransform.GetScale3D();
+		return true;
+	}
+
+	return false;
+}
 
 void UTypedElementViewportInteraction::BeginGizmoManipulation(FTypedElementListConstRef InElementsToMove, const UE::Widget::EWidgetMode InWidgetMode)
 {
@@ -278,6 +289,37 @@ bool UTypedElementViewportInteraction::GetFocusBounds(FTypedElementListConstRef
 	return bAnyHaveBounds;
 }
 
+bool UTypedElementViewportInteraction::GetScaleForGizmo(FTypedElementListConstRef InElements, FVector& OutScale)
+{
+	bool bHasAnyScale = false;
+	FVector CombinedScale(ForceInitToZero);
+	InElements->ForEachElementHandle([this, &CombinedScale, &bHasAnyScale](const FTypedElementHandle& InElementToMove)
+	{
+		FTypedElementViewportInteractionElement ViewportInteractionElement = ResolveViewportInteractionElement(InElementToMove);
+		if (ViewportInteractionElement)
+		{
+			FVector ElementScale;
+			if (ViewportInteractionElement.GetScaleForGizmo(ElementScale))
+			{
+				if (!bHasAnyScale)
+				{
+					CombinedScale = ElementScale;
+				}
+				else
+				{
+					CombinedScale = FVector::Max(CombinedScale, ElementScale);
+				}
+
+				bHasAnyScale = true;
+			}
+		}
+		return true;
+	});
+
+	OutScale = CombinedScale;
+	return bHasAnyScale;
+}
+
 FTypedElementViewportInteractionElement UTypedElementViewportInteraction::ResolveViewportInteractionElement(const FTypedElementHandle& InElementHandle) const
 {
 	return InElementHandle
diff --git a/Engine/Source/Editor/EditorFramework/Public/Elements/Framework/TypedElementViewportInteraction.h b/Engine/Source/Editor/EditorFramework/Public/Elements/Framework/TypedElementViewportInteraction.h
--- a/Engine/Source/Editor/EditorFramework/Public/Elements/Framework/TypedElementViewportInteraction.h
+++ b/Engine/Source/Editor/EditorFramework/Public/Elements/Framework/TypedElementViewportInteraction.h
@@ -50,6 +50,7 @@ public:
 	UE_API virtual void PostGizmoManipulationStopped(TArrayView<const FTypedElementHandle> InElementHandles, const UE::Widget::EWidgetMode InWidgetMode);
 	UE_API virtual void MirrorElement(const TTypedElement<ITypedElementWorldInterface>& InElementWorldHandle, const FVector& InMirrorScale, const FVector& InPivotLocation);
 	UE_API virtual bool GetFocusBounds(const TTypedElement<ITypedElementWorldInterface>& InElementWorldHandle, FBoxSphereBounds& OutBounds);
+	UE_API virtual bool GetScaleForGizmo(const TTypedElement<ITypedElementWorldInterface>& InElementWorldHandle, FVector& OutScale);
 };
 
 /**
@@ -90,7 +91,8 @@ public:
 	void GizmoManipulationStopped(const UE::Widget::EWidgetMode InWidgetMode, const ETypedElementViewportInteractionGizmoManipulationType InManipulationType) const { ViewportInteractionCustomization->GizmoManipulationStopped(ElementWorldHandle, InWidgetMode, InManipulationType); }
 	void MirrorElement(const FVector& InMirrorScale, const FVector& InPivotLocation) const { ViewportInteractionCustomization->MirrorElement(ElementWorldHandle, InMirrorScale, InPivotLocation); }
 	bool GetFocusBounds(FBoxSphereBounds& OutBounds) const { return ViewportInteractionCustomization->GetFocusBounds(ElementWorldHandle, OutBounds); }
-	
+	bool GetScaleForGizmo(FVector& OutScale) const { return ViewportInteractionCustomization->GetScaleForGizmo(ElementWorldHandle, OutScale); }
+
 private:
 	TTypedElement<ITypedElementWorldInterface> ElementWorldHandle;
 	FTypedElementViewportInteractionCustomization* ViewportInteractionCustomization = nullptr;
@@ -140,6 +142,12 @@ public:
 	 *	@note This list should have been pre-normalized via UTypedElementSelectionSet::GetNormalizedSelection or UTypedElementSelectionSet::GetNormalizedElementList, or FLevelEditorViewportClient::GetElementsToManipulate.
 	 */
 	UE_API bool GetFocusBounds(FTypedElementListConstRef InElements, FBoxSphereBounds& OutBox);
+
+	/**
+	 *	Calculate the combined, signed scale of all elements in the list, to be used for scale delta modifications (before applying). If there are multiple objects of differing scale, return the total bounds.
+	 *	@note This list should have been pre-normalized via UTypedElementSelectionSet::GetNormalizedSelection or UTypedElementSelectionSet::GetNormalizedElementList, or FLevelEditorViewportClient::GetElementsToManipulate.
+	 */
+	UE_API bool GetScaleForGizmo(FTypedElementListConstRef InElements, FVector& OutScale);
 private:
 	/**
 	 * Attempt to resolve the selection interface and viewport interaction customization for the given element, if any.
diff --git a/Engine/Source/Editor/Experimental/EditorInteractiveToolsFramework/Private/EditorGizmos/EditorTransformGizmoBuilder.cpp b/Engine/Source/Editor/Experimental/EditorInteractiveToolsFramework/Private/EditorGizmos/EditorTransformGizmoBuilder.cpp
--- a/Engine/Source/Editor/Experimental/EditorInteractiveToolsFramework/Private/EditorGizmos/EditorTransformGizmoBuilder.cpp
+++ b/Engine/Source/Editor/Experimental/EditorInteractiveToolsFramework/Private/EditorGizmos/EditorTransformGizmoBuilder.cpp
@@ -20,6 +20,7 @@
 #include "UObject/Object.h"
 #include "UObject/ObjectPtr.h"
 #include "UObject/ScriptInterface.h"
+#include "Tools/AssetEditorContextInterface.h"
 
 #include UE_INLINE_GENERATED_CPP_BY_NAME(EditorTransformGizmoBuilder)
 
@@ -47,9 +48,13 @@ void UEditorTransformGizmoBuilder::UpdateGizmoForSelection(UInteractiveGizmo* Gi
 	if (UTransformGizmo* TransformGizmo = Cast<UTransformGizmo>(Gizmo))
 	{
 		const UEditorTransformGizmoContextObject* GizmoContextObject = GetTransformGizmoContext(SceneState);
+
 		FEditorModeTools* ModeTools = GizmoContextObject ? GizmoContextObject->GetModeTools() : nullptr;
 		ensure(ModeTools);
-		UEditorTransformProxy* TransformProxy = UEditorTransformProxy::CreateNew(GizmoContextObject);
+
+		IAssetEditorContextInterface* AssetEditorContext = SceneState.ToolManager->GetContextObjectStore()->FindContext<IAssetEditorContextInterface>();
+
+		UEditorTransformProxy* TransformProxy = UEditorTransformProxy::CreateNew(GizmoContextObject, AssetEditorContext);
 		TransformGizmo->SetActiveTarget(TransformProxy, nullptr, ModeTools->GetGizmoStateTarget());
 		TransformGizmo->SetVisibility(true);
 		
diff --git a/Engine/Source/Editor/Experimental/EditorInteractiveToolsFramework/Private/EditorGizmos/EditorTransformProxy.cpp b/Engine/Source/Editor/Experimental/EditorInteractiveToolsFramework/Private/EditorGizmos/EditorTransformProxy.cpp
--- a/Engine/Source/Editor/Experimental/EditorInteractiveToolsFramework/Private/EditorGizmos/EditorTransformProxy.cpp
+++ b/Engine/Source/Editor/Experimental/EditorInteractiveToolsFramework/Private/EditorGizmos/EditorTransformProxy.cpp
@@ -3,31 +3,56 @@
 #include "EditorGizmos/EditorTransformProxy.h"
 
 #include "Editor.h"
+#include "EditorGizmos/EditorTransformGizmoSource.h"
 #include "EditorModeManager.h"
 #include "EditorViewportClient.h"
 #include "EditorGizmos/EditorTransformGizmoUtil.h"
+#include "Elements/Framework/TypedElementRegistry.h"
+#include "Elements/Framework/TypedElementViewportInteraction.h"
 #include "Math/Matrix.h"
 #include "Math/Quat.h"
 #include "Math/Vector.h"
+#include "Subsystems/EditorElementSubsystem.h"
+#include "Tools/AssetEditorContextInterface.h"
 
 #include UE_INLINE_GENERATED_CPP_BY_NAME(EditorTransformProxy)
 
 
 #define LOCTEXT_NAMESPACE "UEditorTransformProxy"
 
+UEditorTransformProxy::UEditorTransformProxy()
+{
+	ViewportInteraction = NewObject<UTypedElementViewportInteraction>();
+}
+
 FTransform UEditorTransformProxy::GetTransform() const
 {
 	if (const FEditorViewportClient* ViewportClient = GetViewportClient())
 	{
 		const FQuat Rotation(ViewportClient->GetWidgetCoordSystem());
 		const FVector Location = ViewportClient->GetWidgetLocation();
-		const FVector Scale(WeakContext->GetModeTools()->GetWidgetScale());
+		const FVector Scale(
+			bUseLegacyWidgetScale
+			? FVector(WeakContext->GetModeTools()->GetWidgetScale())
+			: SharedTransform.GetScale3D());
 		return FTransform(Rotation, Location, Scale);
 	}
 	
 	return FTransform::Identity;
 }
 
+void UEditorTransformProxy::BeginTransformEditSequence()
+{
+	Super::BeginTransformEditSequence();
+
+	if (!bUseLegacyWidgetScale)
+	{
+		UpdateSharedTransform();
+	}
+
+	InitialSharedTransform = SharedTransform;
+}
+
 void UEditorTransformProxy::InputTranslateDelta(const FVector& InDeltaTranslate, EAxisList::Type InAxisList)
 {
 	if (FEditorViewportClient* ViewportClient = GetViewportClient())
@@ -52,11 +77,28 @@ void UEditorTransformProxy::InputScaleDelta(const FVector& InDeltaScale, EAxisLi
 		FRotator Rot = FRotator::ZeroRotator;
 		FVector Scale = InDeltaScale;
 
+		constexpr bool bApplyDirectionalScale = true;
+		if constexpr (bApplyDirectionalScale)
+		{
+			const FVector OriginalScale = InitialSharedTransform.GetScale3D();
+			const FVector OriginalScaleSign = OriginalScale.GetSignVector();
+
+			Scale *= OriginalScaleSign;
+
+			UE_LOG(LogTemp, Warning, TEXT("InputScaleDelta: InDeltaScale: %s, OriginalScale: %s, AdjustedScale: %s"), *InDeltaScale.ToString(), *OriginalScale.ToString(), *Scale.ToString());
+		}
+
 		// Set legacy widget axis temporarily because InputWidgetDelta validates the axis in some crashes and crashes if it is not set
 		ViewportClient->SetCurrentWidgetAxis(InAxisList);
 		ViewportClient->InputWidgetDelta(ViewportClient->Viewport, InAxisList, Translate, Rot, Scale);
 		OnTransformChanged.Broadcast(this, FTransform(Rot, Translate, Scale));
 		ViewportClient->SetCurrentWidgetAxis(EAxisList::None);
+
+		if (!bUseLegacyWidgetScale)
+		{
+			// @todo: This is a workaround for the fact that the scale is not applied to the SharedTransform in InputWidgetDelta.
+			UpdateSharedTransform();
+		}
 	}
 }
 
@@ -76,13 +118,57 @@ void UEditorTransformProxy::InputRotateDelta(const FRotator& InDeltaRotate, EAxi
 	}
 }
 
-UEditorTransformProxy* UEditorTransformProxy::CreateNew(const UEditorTransformGizmoContextObject* InContext)
+UEditorTransformProxy* UEditorTransformProxy::CreateNew(const UEditorTransformGizmoContextObject* InContext, IAssetEditorContextInterface* InAssetEditorContext)
 {
 	UEditorTransformProxy* NewProxy = NewObject<UEditorTransformProxy>();
 	NewProxy->WeakContext = InContext;
+	NewProxy->WeakAssetEditorContext = InAssetEditorContext;
+
+	if (!NewProxy->bUseLegacyWidgetScale)
+	{
+		NewProxy->UpdateSharedTransform();
+		NewProxy->InitialSharedTransform = NewProxy->SharedTransform;
+	}
+
 	return NewProxy;
 }
 
+UEditorTransformProxy* UEditorTransformProxy::CreateNew(const UEditorTransformGizmoContextObject* InContext)
+{
+	return CreateNew(InContext, nullptr);
+}
+
+void UEditorTransformProxy::UpdateSharedTransform()
+{
+	const FTypedElementListConstRef ElementsToManipulate = GetElementsToManipulate();
+
+	FVector ElementScale = FVector::OneVector;
+	ViewportInteraction->GetScaleForGizmo(ElementsToManipulate, ElementScale);
+
+	SharedTransform.SetScale3D(ElementScale);
+}
+
+FTypedElementListConstRef UEditorTransformProxy::GetElementsToManipulate() const
+{
+	if (IAssetEditorContextInterface* StrongAssetEditorContext = WeakAssetEditorContext.Get())
+	{
+		// @see: FLevelEditorViewportClient::CacheElementsToManipulate
+		FTypedElementListRef ElementsToManipulate = UEditorElementSubsystem::GetEditorNormalizedSelectionSet(*StrongAssetEditorContext->GetSelectionSet());
+		if (const FEditorViewportClient* ViewportClient = GetViewportClient())
+		{
+			const UE::Widget::EWidgetMode WidgetMode = ViewportClient->GetWidgetMode();
+			const bool bIsSimulateInEditorViewport = ViewportClient->IsSimulateInEditorViewport();
+
+			UEditorElementSubsystem::GetEditorManipulableElements(ElementsToManipulate, WidgetMode, bIsSimulateInEditorViewport ? GEditor->PlayWorld : GEditor->EditorWorld);
+		}
+
+		return ElementsToManipulate;
+	}
+
+	FTypedElementListConstRef EmptyElements = UTypedElementRegistry::GetInstance()->CreateElementList();
+	return EmptyElements;
+}
+
 FEditorViewportClient* UEditorTransformProxy::GetViewportClient() const
 {
 	if (WeakContext.IsValid() && WeakContext->GetModeTools())
@@ -100,4 +186,17 @@ void UEditorTransformProxy::SetCurrentAxis(const EAxisList::Type InAxisList) con
 	}
 }
 
+void UEditorTransformProxy::SetUseLegacyWidgetScale(const bool bInUseLegacyWidgetScale)
+{
+	if (bUseLegacyWidgetScale != bInUseLegacyWidgetScale)
+	{
+		bUseLegacyWidgetScale = bInUseLegacyWidgetScale;
+
+		if (!bUseLegacyWidgetScale)
+		{
+			UpdateSharedTransform();	
+		}
+	}
+}
+
 #undef LOCTEXT_NAMESPACE
diff --git a/Engine/Source/Editor/Experimental/EditorInteractiveToolsFramework/Public/EditorGizmos/EditorTransformProxy.h b/Engine/Source/Editor/Experimental/EditorInteractiveToolsFramework/Public/EditorGizmos/EditorTransformProxy.h
--- a/Engine/Source/Editor/Experimental/EditorInteractiveToolsFramework/Public/EditorGizmos/EditorTransformProxy.h
+++ b/Engine/Source/Editor/Experimental/EditorInteractiveToolsFramework/Public/EditorGizmos/EditorTransformProxy.h
@@ -11,12 +11,15 @@
 #include "Misc/AssertionMacros.h"
 #include "UObject/ObjectMacros.h"
 #include "UObject/UObjectGlobals.h"
+#include "UObject/WeakInterfacePtr.h"
 
 #include "EditorTransformProxy.generated.h"
 
 #define UE_API EDITORINTERACTIVETOOLSFRAMEWORK_API
 
+class IAssetEditorContextInterface;
 class UEditorTransformGizmoContextObject;
+class UTypedElementViewportInteraction;
 class FEditorViewportClient;
 class UObject;
 
@@ -31,8 +34,11 @@ UCLASS(MinimalAPI, Transient)
 class UEditorTransformProxy : public UTransformProxy
 {
 	GENERATED_BODY()
+
 public:
 
+	UE_API UEditorTransformProxy();
+
 	/**
 	 * @return the stored transform for currently selected objects.
 	 */
@@ -46,6 +52,8 @@ public:
 		check(false);
 	}
 
+	UE_API virtual void BeginTransformEditSequence() override;
+
 	/** Input translate delta to be applied in world space of the current transform. */
 	UE_API virtual void InputTranslateDelta(const FVector& InDeltaTranslate, EAxisList::Type InAxisList);
 
@@ -58,13 +66,35 @@ public:
 	/** Set legacy widget axis temporarily because FEditorViewportClient overrides may expect it */
 	UE_API void SetCurrentAxis(const EAxisList::Type InAxisList) const;
 
-	static UE_API UEditorTransformProxy* CreateNew(const UEditorTransformGizmoContextObject* InContext = nullptr);
-	
+	/** Set whether to use FEditorModeTools::GetWidgetScale as the Transform scale. */
+	UE_API void SetUseLegacyWidgetScale(const bool bInUseLegacyWidgetScale);
+
+	static UE_API UEditorTransformProxy* CreateNew(const UEditorTransformGizmoContextObject* InContext = nullptr, IAssetEditorContextInterface* InAssetEditorContext = nullptr);
+
+	UE_DEPRECATED(5.7, "Use CreateNew that accepts an AssetEditorContext instead")
+	static UE_API UEditorTransformProxy* CreateNew(const UEditorTransformGizmoContextObject* InContext);
+
 private:
-	
+
+	/** Recalculate main SharedTransform based on the current selection set. */
+	UE_API virtual void UpdateSharedTransform() override;
+
+	UE_API FTypedElementListConstRef GetElementsToManipulate() const;
+
 	UE_API FEditorViewportClient* GetViewportClient() const;
 
+private:
+
+	/** Whether to use FEditorModeTools::GetWidgetScale as the Transform scale. If false, scale is derived from the underlying objects. */
+	UPROPERTY()
+	bool bUseLegacyWidgetScale = true;
+
+	/** The viewport interaction instance, providing handlers for selected Typed Elements. */
+	UPROPERTY()
+	TObjectPtr<UTypedElementViewportInteraction> ViewportInteraction;
+
 	TWeakObjectPtr<const UEditorTransformGizmoContextObject> WeakContext;
+	TWeakInterfacePtr<IAssetEditorContextInterface> WeakAssetEditorContext;
 };
 
 #undef UE_API
diff --git a/Engine/Source/Editor/UnrealEd/Private/Elements/Component/ComponentElementEditorViewportInteractionCustomization.cpp b/Engine/Source/Editor/UnrealEd/Private/Elements/Component/ComponentElementEditorViewportInteractionCustomization.cpp
--- a/Engine/Source/Editor/UnrealEd/Private/Elements/Component/ComponentElementEditorViewportInteractionCustomization.cpp
+++ b/Engine/Source/Editor/UnrealEd/Private/Elements/Component/ComponentElementEditorViewportInteractionCustomization.cpp
@@ -134,6 +134,25 @@ bool FComponentElementEditorViewportInteractionCustomization::GetFocusBounds(con
 	return false;
 }
 
+bool FComponentElementEditorViewportInteractionCustomization::GetScaleForGizmo(const TTypedElement<ITypedElementWorldInterface>& InElementWorldHandle, FVector& OutScale)
+{
+	if (UActorComponent* Component = ComponentElementDataUtil::GetComponentFromHandle(InElementWorldHandle))
+	{
+		if (const UPrimitiveComponent* PrimitiveComponent = Cast<UPrimitiveComponent>(Component))
+		{
+			if (!PrimitiveComponent->IsRegistered())
+			{
+				return false;
+			}
+
+			OutScale = PrimitiveComponent->GetComponentScale();
+			return true;
+		}
+	}
+
+	return false;
+}
+
 void FComponentElementEditorViewportInteractionCustomization::ApplyDeltaToComponent(USceneComponent* InComponent, const bool InIsDelta, const FVector* InDeltaTranslationPtr, const FRotator* InDeltaRotationPtr, const FVector* InDeltaScalePtr, const FVector& InPivotLocation, const FInputDeviceState& InInputState)
 {
 	if (GEditor->IsDeltaModificationEnabled())
diff --git a/Engine/Source/Editor/UnrealEd/Public/Elements/Component/ComponentElementEditorViewportInteractionCustomization.h b/Engine/Source/Editor/UnrealEd/Public/Elements/Component/ComponentElementEditorViewportInteractionCustomization.h
--- a/Engine/Source/Editor/UnrealEd/Public/Elements/Component/ComponentElementEditorViewportInteractionCustomization.h
+++ b/Engine/Source/Editor/UnrealEd/Public/Elements/Component/ComponentElementEditorViewportInteractionCustomization.h
@@ -14,6 +14,7 @@ public:
 	UNREALED_API virtual bool GetGizmoPivotLocation(const TTypedElement<ITypedElementWorldInterface>& InElementWorldHandle, const UE::Widget::EWidgetMode InWidgetMode, FVector& OutPivotLocation) override;
 	UNREALED_API virtual void GizmoManipulationDeltaUpdate(const TTypedElement<ITypedElementWorldInterface>& InElementWorldHandle, const UE::Widget::EWidgetMode InWidgetMode, const EAxisList::Type InDragAxis, const FInputDeviceState& InInputState, const FTransform& InDeltaTransform, const FVector& InPivotLocation) override;
 	UNREALED_API virtual bool GetFocusBounds(const TTypedElement<ITypedElementWorldInterface>& InElementWorldHandle, FBoxSphereBounds& OutBounds) override;
+	UNREALED_API virtual bool GetScaleForGizmo(const TTypedElement<ITypedElementWorldInterface>& InElementWorldHandle, FVector& OutScale) override;
 	
 	static UNREALED_API void ApplyDeltaToComponent(USceneComponent* InComponent, const bool InIsDelta, const FVector* InDeltaTranslationPtr, const FRotator* InDeltaRotationPtr, const FVector* InDeltaScalePtr, const FVector& InPivotLocation, const FInputDeviceState& InInputState);
 };
