[Editor ITF] Removed scale sign correction (moved to the gizmo itself).

diff --git a/Engine/Source/Editor/Experimental/EditorInteractiveToolsFramework/Private/EditorGizmos/EditorTransformProxy.cpp b/Engine/Source/Editor/Experimental/EditorInteractiveToolsFramework/Private/EditorGizmos/EditorTransformProxy.cpp
--- a/Engine/Source/Editor/Experimental/EditorInteractiveToolsFramework/Private/EditorGizmos/EditorTransformProxy.cpp
+++ b/Engine/Source/Editor/Experimental/EditorInteractiveToolsFramework/Private/EditorGizmos/EditorTransformProxy.cpp
@@ -77,17 +77,6 @@ void UEditorTransformProxy::InputScaleDelta(const FVector& InDeltaScale, EAxisLi
 		FRotator Rot = FRotator::ZeroRotator;
 		FVector Scale = InDeltaScale;
 
-		constexpr bool bApplyDirectionalScale = true;
-		if constexpr (bApplyDirectionalScale)
-		{
-			const FVector OriginalScale = InitialSharedTransform.GetScale3D();
-			const FVector OriginalScaleSign = OriginalScale.GetSignVector();
-
-			Scale *= OriginalScaleSign;
-
-			UE_LOG(LogTemp, Warning, TEXT("InputScaleDelta: InDeltaScale: %s, OriginalScale: %s, AdjustedScale: %s"), *InDeltaScale.ToString(), *OriginalScale.ToString(), *Scale.ToString());
-		}
-
 		// Set legacy widget axis temporarily because InputWidgetDelta validates the axis in some crashes and crashes if it is not set
 		ViewportClient->SetCurrentWidgetAxis(InAxisList);
 		ViewportClient->InputWidgetDelta(ViewportClient->Viewport, InAxisList, Translate, Rot, Scale);
