[ITF Drag Tools] Update ViewportClientProxy so that it stores InteractiveToolsContext instead of Viewport Client

This allows to retrieve the currently Focused Viewport Client.
Previous version was storing what was Focused at time of ViewportClientProxy creation, which could change throughout the Editor session.

diff --git a/Engine/Source/Editor/Experimental/EditorInteractiveToolsFramework/Private/EditorDragTools/EditorBoxSelect.cpp b/Engine/Source/Editor/Experimental/EditorInteractiveToolsFramework/Private/EditorDragTools/EditorBoxSelect.cpp
--- a/Engine/Source/Editor/Experimental/EditorInteractiveToolsFramework/Private/EditorDragTools/EditorBoxSelect.cpp
+++ b/Engine/Source/Editor/Experimental/EditorInteractiveToolsFramework/Private/EditorDragTools/EditorBoxSelect.cpp
@@ -71,8 +71,7 @@ FInputRayHit FEditorBoxSelect::CanBeginClickDragSequence(const FInputDeviceRay&
 	{
 		if (IsCurrentModeSupported())
 		{
-			if (EditorViewportClient->IsOrtho()
-				&& !IEditorViewportClientProxy::IsMouseOnWidgetAxis(EditorViewportClientProxy))
+			if (EditorViewportClient->IsOrtho())
 			{
 				return FInputRayHit(TNumericLimits<float>::Max());
 			}
diff --git a/Engine/Source/Editor/Experimental/EditorInteractiveToolsFramework/Private/EditorDragTools/EditorDragToolBehaviorTarget.cpp b/Engine/Source/Editor/Experimental/EditorInteractiveToolsFramework/Private/EditorDragTools/EditorDragToolBehaviorTarget.cpp
--- a/Engine/Source/Editor/Experimental/EditorInteractiveToolsFramework/Private/EditorDragTools/EditorDragToolBehaviorTarget.cpp
+++ b/Engine/Source/Editor/Experimental/EditorInteractiveToolsFramework/Private/EditorDragTools/EditorDragToolBehaviorTarget.cpp
@@ -63,7 +63,7 @@ FEditorDragToolBehaviorTarget::FEditorDragToolBehaviorTarget(FEditorViewportClie
 		ModeTools = InEditorViewportClient->GetModeTools();
 	}
 
-	EditorViewportClientProxy = IEditorViewportClientProxy::CreateViewportClientProxy(InEditorViewportClient);
+	EditorViewportClientProxy = IEditorViewportClientProxy::CreateViewportClientProxy(ModeTools->GetInteractiveToolsContext());
 }
 
 void FEditorDragToolBehaviorTarget::OnUpdateModifierState(int InModifierID, bool bInIsOn)
diff --git a/Engine/Source/Editor/Experimental/EditorInteractiveToolsFramework/Private/EditorDragTools/EditorMeasureTool.cpp b/Engine/Source/Editor/Experimental/EditorInteractiveToolsFramework/Private/EditorDragTools/EditorMeasureTool.cpp
--- a/Engine/Source/Editor/Experimental/EditorInteractiveToolsFramework/Private/EditorDragTools/EditorMeasureTool.cpp
+++ b/Engine/Source/Editor/Experimental/EditorInteractiveToolsFramework/Private/EditorDragTools/EditorMeasureTool.cpp
@@ -89,7 +89,7 @@ FInputRayHit FEditorMeasureTool::CanBeginClickDragSequence(const FInputDeviceRay
 		if (IsCurrentModeSupported())
 		{
 			// This Box Selection is for Orthographic viewports only
-			if (EditorViewportClient->IsOrtho() && ! IEditorViewportClientProxy::IsMouseOnWidgetAxis(EditorViewportClientProxy))
+			if (EditorViewportClient->IsOrtho())
 			{
 				return FInputRayHit(TNumericLimits<float>::Max());
 			}
diff --git a/Engine/Source/Editor/Experimental/EditorInteractiveToolsFramework/Private/EditorDragTools/EditorViewportClientProxy.cpp b/Engine/Source/Editor/Experimental/EditorInteractiveToolsFramework/Private/EditorDragTools/EditorViewportClientProxy.cpp
--- a/Engine/Source/Editor/Experimental/EditorInteractiveToolsFramework/Private/EditorDragTools/EditorViewportClientProxy.cpp
+++ b/Engine/Source/Editor/Experimental/EditorInteractiveToolsFramework/Private/EditorDragTools/EditorViewportClientProxy.cpp
@@ -1,54 +1,39 @@
 ﻿// Copyright Epic Games, Inc. All Rights Reserved.
 
 #include "EditorDragTools/EditorViewportClientProxy.h"
+#include "EditorModeManager.h"
 #include "GameFramework/Volume.h"
 #include "LevelEditorViewport.h"
+#include "Tools/EdModeInteractiveToolsContext.h"
 #include "UnrealWidget.h"
 
 IEditorViewportClientProxy* IEditorViewportClientProxy::CreateViewportClientProxy(FEditorViewportClient* InViewportClient)
 {
-	if (InViewportClient->IsLevelEditorClient())
+	if (FEditorModeTools* ModeTool = InViewportClient->GetModeTools())
 	{
-		if (FLevelEditorViewportClient* LevelEditorViewportClient = static_cast<FLevelEditorViewportClient*>(InViewportClient))
-		{
-			return new FLevelEditorViewportClientProxy(LevelEditorViewportClient);
-		}
+		return CreateViewportClientProxy(ModeTool->GetInteractiveToolsContext());
 	}
-
-	return new FEditorViewportClientProxy(InViewportClient);
+	return nullptr;
 }
 
-bool IEditorViewportClientProxy::IsMouseOnWidgetAxis(IEditorViewportClientProxy* InEditorViewportClientProxy)
+IEditorViewportClientProxy* IEditorViewportClientProxy::CreateViewportClientProxy(UModeManagerInteractiveToolsContext* InInteractiveToolsContext)
 {
-	if (InEditorViewportClientProxy)
+	if (InInteractiveToolsContext)
 	{
-		if (FEditorViewportClient* const EditorViewportClient = InEditorViewportClientProxy->GetEditorViewportClient())
+		if (FEditorViewportClient* EditorViewportClient = InInteractiveToolsContext->GetFocusedViewportClient())
 		{
-			if (FViewport* const Viewport = EditorViewportClient->Viewport)
+			if (EditorViewportClient->IsLevelEditorClient())
 			{
-				FIntPoint MousePos;
-				EditorViewportClient->Viewport->GetMousePos(MousePos);
-
-				// Since some drag tool does not involve any keyboard modifier (Alt, Shift, Ctrl),
-				// in some cases we need to make sure the user is not hovering over widget axis proxies.
-				// Intercepting input in that case prevents TRS gizmos from working.
-				if (const HHitProxy* const Proxy = Viewport->GetHitProxy(MousePos.X, MousePos.Y))
-				{
-					// We hit a widget axis, so we don't start the drag sequence
-					if (Proxy->IsA(HWidgetAxis::StaticGetType()))
-					{
-						return true;
-					}
-				}
+				return new FLevelEditorViewportClientProxy(InInteractiveToolsContext);
 			}
 		}
 	}
 
-	return false;
+	return new FEditorViewportClientProxy(InInteractiveToolsContext);
 }
 
-FEditorViewportClientProxy::FEditorViewportClientProxy(FEditorViewportClient* InEditorViewportClient)
-	: EditorViewportClient(InEditorViewportClient)
+FEditorViewportClientProxy::FEditorViewportClientProxy(UModeManagerInteractiveToolsContext* InInteractiveToolsContext)
+    : InteractiveToolsContextWeak(InInteractiveToolsContext)
 {
 }
 
@@ -67,21 +52,34 @@ const TArray<FName> FEditorViewportClientProxy::GetHiddenLayers() const
 	return TArray<FName>();
 }
 
-FEditorViewportClient* FEditorViewportClientProxy::GetEditorViewportClient()
+FEditorViewportClient* FEditorViewportClientProxy::GetEditorViewportClient() const
 {
-	return EditorViewportClient.Get();
+	if (TStrongObjectPtr<UModeManagerInteractiveToolsContext> InteractiveToolsContext = InteractiveToolsContextWeak.Pin())
+	{
+		return InteractiveToolsContext->GetFocusedViewportClient();
+	}
+
+	return nullptr;
 }
 
-FLevelEditorViewportClientProxy::FLevelEditorViewportClientProxy(FLevelEditorViewportClient* InLevelEditorViewportClient)
-	:LevelEditorViewportClient(InLevelEditorViewportClient)
+FLevelEditorViewportClientProxy::FLevelEditorViewportClientProxy(UModeManagerInteractiveToolsContext* InInteractiveToolsContext)
+	: InteractiveToolsContextWeak(InInteractiveToolsContext)
 {
 }
 
 bool FLevelEditorViewportClientProxy::IsActorVisible(const AActor* InActor) const
 {
-	if (InActor && LevelEditorViewportClient)
+	if (InActor)
 	{
-		return !InActor->IsA(AVolume::StaticClass()) || !LevelEditorViewportClient->IsVolumeVisibleInViewport(*InActor);
+		if (!InActor->IsA(AVolume::StaticClass()))
+		{
+			return true;
+		}
+
+		if (FLevelEditorViewportClient* LevelEditorViewportClient = GetLevelEditorViewportClient())
+		{
+			return !LevelEditorViewportClient->IsVolumeVisibleInViewport(*InActor);
+		}
 	}
 
 	return false;
@@ -89,7 +87,7 @@ bool FLevelEditorViewportClientProxy::IsActorVisible(const AActor* InActor) cons
 
 const TArray<FName> FLevelEditorViewportClientProxy::GetHiddenLayers() const
 {
-	if (LevelEditorViewportClient)
+	if (FLevelEditorViewportClient* LevelEditorViewportClient = GetLevelEditorViewportClient())
 	{
 		return LevelEditorViewportClient->ViewHiddenLayers;
 	}
@@ -97,7 +95,17 @@ const TArray<FName> FLevelEditorViewportClientProxy::GetHiddenLayers() const
 	return TArray<FName>();
 }
 
-FEditorViewportClient* FLevelEditorViewportClientProxy::GetEditorViewportClient()
+FEditorViewportClient* FLevelEditorViewportClientProxy::GetEditorViewportClient() const
+{
+	if (TStrongObjectPtr<UModeManagerInteractiveToolsContext> InteractiveToolsContext = InteractiveToolsContextWeak.Pin())
+	{
+		return InteractiveToolsContext->GetFocusedViewportClient();
+	}
+
+	return nullptr;
+}
+
+FLevelEditorViewportClient* FLevelEditorViewportClientProxy::GetLevelEditorViewportClient() const
 {
-	return LevelEditorViewportClient.Get();
+	return static_cast<FLevelEditorViewportClient*>(GetEditorViewportClient());
 }
diff --git a/Engine/Source/Editor/Experimental/EditorInteractiveToolsFramework/Public/EditorDragTools/DragToolsBehaviourSource.h b/Engine/Source/Editor/Experimental/EditorInteractiveToolsFramework/Public/EditorDragTools/DragToolsBehaviourSource.h
--- a/Engine/Source/Editor/Experimental/EditorInteractiveToolsFramework/Public/EditorDragTools/DragToolsBehaviourSource.h
+++ b/Engine/Source/Editor/Experimental/EditorInteractiveToolsFramework/Public/EditorDragTools/DragToolsBehaviourSource.h
@@ -30,7 +30,7 @@ static constexpr int DUPLICATE_SELECTION_PRIORITY = DRAG_TOOLS_DEFAULT_PRIORITY;
 static constexpr int FRUSTUM_SELECT_PRIORITY = DRAG_TOOLS_DEFAULT_PRIORITY;
 static constexpr int VIEWPORT_CHANGE_PRIORITY = DRAG_TOOLS_DEFAULT_PRIORITY;
 static constexpr int BOX_SELECT_PRIORITY = HANDLED_AFTER_GIZMO_PRIORITY;
-static constexpr int MEASURE_PRIORITY = DRAG_TOOLS_LOW_PRIORITY;
+static constexpr int MEASURE_PRIORITY = HANDLED_AFTER_GIZMO_PRIORITY;
 
 }
 
diff --git a/Engine/Source/Editor/Experimental/EditorInteractiveToolsFramework/Public/EditorDragTools/EditorViewportClientProxy.h b/Engine/Source/Editor/Experimental/EditorInteractiveToolsFramework/Public/EditorDragTools/EditorViewportClientProxy.h
--- a/Engine/Source/Editor/Experimental/EditorInteractiveToolsFramework/Public/EditorDragTools/EditorViewportClientProxy.h
+++ b/Engine/Source/Editor/Experimental/EditorInteractiveToolsFramework/Public/EditorDragTools/EditorViewportClientProxy.h
@@ -3,10 +3,12 @@
 #pragma once
 
 #include "Templates/SharedPointer.h"
+#include "UObject/WeakObjectPtr.h"
 
 class AActor;
 class FEditorViewportClient;
 class FLevelEditorViewportClient;
+class UModeManagerInteractiveToolsContext;
 struct FWorldSelectionElementArgs;
 
 #define UE_API EDITORINTERACTIVETOOLSFRAMEWORK_API
@@ -15,41 +17,43 @@ class IEditorViewportClientProxy : public TSharedFromThis<IEditorViewportClientP
 {
 public:
 	virtual ~IEditorViewportClientProxy() = default;
-	UE_API static IEditorViewportClientProxy* CreateViewportClientProxy(FEditorViewportClient* InViewportClient);
-	static bool IsMouseOnWidgetAxis(IEditorViewportClientProxy* InEditorViewportClientProxy);
 
+	UE_DEPRECATED(5.7, "Please use the version taking UModeManagerInteractiveToolsContext* as argument")
+	UE_API static IEditorViewportClientProxy* CreateViewportClientProxy(FEditorViewportClient* InViewportClient);
+	UE_API static IEditorViewportClientProxy* CreateViewportClientProxy(UModeManagerInteractiveToolsContext* InInteractiveToolsContext);
 
 	virtual bool IsActorVisible(const AActor* InActor) const = 0;
 
 	virtual const TArray<FName> GetHiddenLayers() const = 0;
 
-	virtual FEditorViewportClient* GetEditorViewportClient() = 0;
+	virtual FEditorViewportClient* GetEditorViewportClient() const = 0;
 };
 
 class FEditorViewportClientProxy : public IEditorViewportClientProxy
 {
 public:
-	FEditorViewportClientProxy(FEditorViewportClient* InEditorViewportClient);
-	
+	FEditorViewportClientProxy(UModeManagerInteractiveToolsContext* InInteractiveToolsContext);
+
 	virtual bool IsActorVisible(const AActor* InActor) const override;
 	virtual const TArray<FName> GetHiddenLayers() const override;
-	virtual FEditorViewportClient* GetEditorViewportClient() override;
+	virtual FEditorViewportClient* GetEditorViewportClient() const override;
 
 private:
-	TSharedPtr<FEditorViewportClient> EditorViewportClient;
+	TWeakObjectPtr<UModeManagerInteractiveToolsContext> InteractiveToolsContextWeak;
 };
 
 class FLevelEditorViewportClientProxy : public IEditorViewportClientProxy
 {
 public:
-	FLevelEditorViewportClientProxy(FLevelEditorViewportClient *InLevelEditorViewportClient);
+	FLevelEditorViewportClientProxy(UModeManagerInteractiveToolsContext* InInteractiveToolsContext);
 
 	virtual bool IsActorVisible(const AActor* InActor) const override;
 	virtual const TArray<FName> GetHiddenLayers() const override;
-	virtual FEditorViewportClient* GetEditorViewportClient() override;
+	virtual FEditorViewportClient* GetEditorViewportClient() const override;
+	virtual FLevelEditorViewportClient* GetLevelEditorViewportClient() const;
 
 private:
-	TSharedPtr<FLevelEditorViewportClient> LevelEditorViewportClient;
+	TWeakObjectPtr<UModeManagerInteractiveToolsContext> InteractiveToolsContextWeak;
 };
 
 #undef UE_API
diff --git a/Engine/Source/Editor/UnrealEd/Private/Tools/EdModeInteractiveToolsContext.cpp b/Engine/Source/Editor/UnrealEd/Private/Tools/EdModeInteractiveToolsContext.cpp
--- a/Engine/Source/Editor/UnrealEd/Private/Tools/EdModeInteractiveToolsContext.cpp
+++ b/Engine/Source/Editor/UnrealEd/Private/Tools/EdModeInteractiveToolsContext.cpp
@@ -1256,6 +1256,58 @@ FRay UModeManagerInteractiveToolsContext::GetLastWorldRay() const
 	return CurrentMouseState.Mouse.WorldRay;
 }
 
+FEditorViewportClient* UModeManagerInteractiveToolsContext::GetHoveredViewportClient() const
+{
+	FEditorViewportClient* HoveredViewportClient = nullptr;
+
+	if (const FEditorModeTools* const ModeManager = GetParentEditorModeManager())
+	{
+		HoveredViewportClient = ModeManager->GetHoveredViewportClient();
+	}
+
+	if (!HoveredViewportClient)
+	{
+		if (ToolManager)
+		{
+			if (const IToolsContextQueriesAPI* const ContextQueriesAPI = ToolManager->GetContextQueriesAPI())
+			{
+				if (const FViewport* const Viewport = ContextQueriesAPI->GetHoveredViewport())
+				{
+					HoveredViewportClient = static_cast<FEditorViewportClient*>(Viewport->GetClient());
+				}
+			}
+		}
+	}
+
+	return HoveredViewportClient;
+}
+
+FEditorViewportClient* UModeManagerInteractiveToolsContext::GetFocusedViewportClient() const
+{
+	FEditorViewportClient* FocusedViewportClient = nullptr;
+
+	if (const FEditorModeTools* const ModeManager = GetParentEditorModeManager())
+	{
+		FocusedViewportClient = ModeManager->GetFocusedViewportClient();
+	}
+
+	if (!FocusedViewportClient)
+	{
+		if (ToolManager)
+		{
+			if (const IToolsContextQueriesAPI* const ContextQueriesAPI = ToolManager->GetContextQueriesAPI())
+			{
+				if (const FViewport* const Viewport = ContextQueriesAPI->GetFocusedViewport())
+				{
+					FocusedViewportClient = static_cast<FEditorViewportClient*>(Viewport->GetClient());
+				}
+			}
+		}
+	}
+
+	return FocusedViewportClient;
+}
+
 void UModeManagerInteractiveToolsContext::Initialize(IToolsContextQueriesAPI* QueriesAPIIn, IToolsContextTransactionsAPI* TransactionsAPIIn)
 {
 	Super::Initialize(QueriesAPIIn, TransactionsAPIIn);
diff --git a/Engine/Source/Editor/UnrealEd/Public/Tools/EdModeInteractiveToolsContext.h b/Engine/Source/Editor/UnrealEd/Public/Tools/EdModeInteractiveToolsContext.h
--- a/Engine/Source/Editor/UnrealEd/Public/Tools/EdModeInteractiveToolsContext.h
+++ b/Engine/Source/Editor/UnrealEd/Public/Tools/EdModeInteractiveToolsContext.h
@@ -315,6 +315,12 @@ public:
 	/** @return Ray into 3D scene at last mouse event */
 	UE_API virtual FRay GetLastWorldRay() const override;
 
+	/** @return currently Hovered viewport client */
+	UE_API FEditorViewportClient* GetHoveredViewportClient() const;
+
+	/** @return currently Focused viewport client */
+	UE_API FEditorViewportClient* GetFocusedViewportClient() const;
+
 protected:
 	UE_API virtual void DeactivateAllActiveTools(EToolShutdownType ShutdownType) override;
 	UE_API virtual void Initialize(IToolsContextQueriesAPI* QueriesAPIIn, IToolsContextTransactionsAPI* TransactionsAPIIn) override;
