[ITF Drag Tools] Add ModeTools validity check

diff --git a/Engine/Source/Editor/Experimental/EditorInteractiveToolsFramework/Private/EditorDragTools/EditorDragToolBehaviorTarget.cpp b/Engine/Source/Editor/Experimental/EditorInteractiveToolsFramework/Private/EditorDragTools/EditorDragToolBehaviorTarget.cpp
--- a/Engine/Source/Editor/Experimental/EditorInteractiveToolsFramework/Private/EditorDragTools/EditorDragToolBehaviorTarget.cpp
+++ b/Engine/Source/Editor/Experimental/EditorInteractiveToolsFramework/Private/EditorDragTools/EditorDragToolBehaviorTarget.cpp
@@ -61,9 +61,12 @@ FEditorDragToolBehaviorTarget::FEditorDragToolBehaviorTarget(FEditorViewportClie
 	if (InEditorViewportClient)
 	{
 		ModeTools = InEditorViewportClient->GetModeTools();
-	}
 
-	EditorViewportClientProxy = IEditorViewportClientProxy::CreateViewportClientProxy(ModeTools->GetInteractiveToolsContext());
+		if (ModeTools)
+		{
+			EditorViewportClientProxy = IEditorViewportClientProxy::CreateViewportClientProxy(ModeTools->GetInteractiveToolsContext());
+		}
+	}
 }
 
 void FEditorDragToolBehaviorTarget::OnUpdateModifierState(int InModifierID, bool bInIsOn)
