[ITF Drag Tools] Marquee Select and Measure tool are rendered on multiple viewports

DragTools render method now also takes the current interaction state as argument (Focused/Hovered/None).
This allows more control on Drag Tools rendering.

diff --git a/Engine/Source/Editor/Experimental/EditorInteractiveToolsFramework/Private/EditorDragTools/DragToolsBehaviourSource.cpp b/Engine/Source/Editor/Experimental/EditorInteractiveToolsFramework/Private/EditorDragTools/DragToolsBehaviourSource.cpp
--- a/Engine/Source/Editor/Experimental/EditorInteractiveToolsFramework/Private/EditorDragTools/DragToolsBehaviourSource.cpp
+++ b/Engine/Source/Editor/Experimental/EditorInteractiveToolsFramework/Private/EditorDragTools/DragToolsBehaviourSource.cpp
@@ -60,7 +60,7 @@ void UDragToolsBehaviorSource::Initialize(const UEditorInteractiveToolsContext*
 		return;
 	}
 
-	FEditorViewportClient* EditorViewportClient = GetEditorViewportClient();
+	FEditorViewportClient* EditorViewportClient = GetFocusedEditorViewportClient();
 
 	BehaviorSet = NewObject<UInputBehaviorSet>();
 
@@ -199,7 +199,7 @@ FEditorDragToolBehaviorTarget* UDragToolsBehaviorSource::GetActiveTool() const
 	return nullptr;
 }
 
-FEditorViewportClient* UDragToolsBehaviorSource::GetEditorViewportClient() const
+FEditorViewportClient* UDragToolsBehaviorSource::GetFocusedEditorViewportClient() const
 {
 	FEditorViewportClient* EditorViewportClient = nullptr;
 
@@ -228,6 +228,36 @@ FEditorViewportClient* UDragToolsBehaviorSource::GetEditorViewportClient() const
 	return EditorViewportClient;
 }
 
+FEditorViewportClient* UDragToolsBehaviorSource::GetHoveredEditorViewportClient() const
+{
+	FEditorViewportClient* EditorViewportClient = nullptr;
+
+	if (const UEditorInteractiveToolsContext* const InteractiveToolsContext = EditorInteractiveToolsContextWeak.Get())
+	{
+		if (const FEditorModeTools* const ModeManager = InteractiveToolsContext->GetParentEditorModeManager())
+		{
+			EditorViewportClient = ModeManager->GetHoveredViewportClient();
+		}
+
+		if (!EditorViewportClient)
+		{
+			if (UInteractiveToolManager* const ToolManager = InteractiveToolsContext->ToolManager)
+			{
+				if (const IToolsContextQueriesAPI* const ContextQueriesAPI = ToolManager->GetContextQueriesAPI())
+				{
+					if (const FViewport* const Viewport = ContextQueriesAPI->GetHoveredViewport())
+					{
+						EditorViewportClient = static_cast<FEditorViewportClient*>(Viewport->GetClient());
+					}
+				}
+			}
+		}
+	}
+
+	return EditorViewportClient;
+}
+
+
 bool UDragToolsBehaviorSource::IsViewportChangeToolEnabled()
 {
 	return UE::Editor::DragTools::UseViewportChangeTool == 1;
@@ -283,11 +313,24 @@ void UDragToolsBehaviorSource::OnDeactivateTool(EDragTool::Type InDragToolType)
 	}
 }
 
-void UDragToolsBehaviorSource::RenderTools(const FSceneView* InSceneView, FCanvas* InCanvas) const
+void UDragToolsBehaviorSource::RenderTools(FViewportClient* InViewportClient, const FSceneView* InSceneView, FCanvas* InCanvas) const
 {
+	const FViewportClient* Focused = GetFocusedEditorViewportClient();
+	const FViewportClient* Hovered = GetHoveredEditorViewportClient();
+
+	EViewInteractionState InteractionState = EViewInteractionState::None;
+	if (InViewportClient == Focused)
+	{
+		InteractionState = EViewInteractionState::Focused;
+	}
+	else if (InViewportClient == Hovered)
+	{
+		InteractionState = EViewInteractionState::Hovered;
+	}
+
 	if (FEditorDragToolBehaviorTarget* ActiveDragTool = GetActiveTool())
 	{
-		ActiveDragTool->Render(InSceneView, InCanvas);
+		ActiveDragTool->Render(InSceneView, InCanvas, InteractionState);
 	}
 }
 
@@ -302,7 +345,7 @@ void UDragToolsBehaviorSource::ActivateViewportChangeTool()
 		return;
 	}
 
-	FEditorViewportClient* EditorViewportClient = GetEditorViewportClient();
+	FEditorViewportClient* EditorViewportClient = GetFocusedEditorViewportClient();
 	if (!EditorViewportClient)
 	{
 		return;
diff --git a/Engine/Source/Editor/Experimental/EditorInteractiveToolsFramework/Private/EditorDragTools/EditorMarqueeSelect.cpp b/Engine/Source/Editor/Experimental/EditorInteractiveToolsFramework/Private/EditorDragTools/EditorMarqueeSelect.cpp
--- a/Engine/Source/Editor/Experimental/EditorInteractiveToolsFramework/Private/EditorDragTools/EditorMarqueeSelect.cpp
+++ b/Engine/Source/Editor/Experimental/EditorInteractiveToolsFramework/Private/EditorDragTools/EditorMarqueeSelect.cpp
@@ -6,8 +6,13 @@
 #include "CanvasTypes.h"
 #include "Settings/LevelEditorViewportSettings.h"
 
-void FEditorMarqueeSelect::Render(const FSceneView* View, FCanvas* Canvas)
+void FEditorMarqueeSelect::Render(const FSceneView* View, FCanvas* Canvas, EViewInteractionState InInteractionState)
 {
+	if (InInteractionState != EViewInteractionState::Focused)
+	{
+		return;
+	}
+
 	if (!Canvas || !View)
 	{
 		return;
diff --git a/Engine/Source/Editor/Experimental/EditorInteractiveToolsFramework/Private/EditorDragTools/EditorMeasureTool.cpp b/Engine/Source/Editor/Experimental/EditorInteractiveToolsFramework/Private/EditorDragTools/EditorMeasureTool.cpp
--- a/Engine/Source/Editor/Experimental/EditorInteractiveToolsFramework/Private/EditorDragTools/EditorMeasureTool.cpp
+++ b/Engine/Source/Editor/Experimental/EditorInteractiveToolsFramework/Private/EditorDragTools/EditorMeasureTool.cpp
@@ -21,8 +21,13 @@ FEditorMeasureTool::FEditorMeasureTool(FEditorViewportClient* InViewportClient)
 	bConvertDelta = false;
 }
 
-void FEditorMeasureTool::Render(const FSceneView* InView, FCanvas* InCanvas)
+void FEditorMeasureTool::Render(const FSceneView* InView, FCanvas* InCanvas, EViewInteractionState InInteractionState)
 {
+	if (InInteractionState != EViewInteractionState::Focused)
+	{
+		return;
+	}
+
 	FEditorViewportClient* const EditorViewportClient = EditorViewportClientProxy->GetEditorViewportClient();
 	if (!EditorViewportClient)
 	{
diff --git a/Engine/Source/Editor/Experimental/EditorInteractiveToolsFramework/Private/EditorDragTools/EditorViewportChange.cpp b/Engine/Source/Editor/Experimental/EditorInteractiveToolsFramework/Private/EditorDragTools/EditorViewportChange.cpp
--- a/Engine/Source/Editor/Experimental/EditorInteractiveToolsFramework/Private/EditorDragTools/EditorViewportChange.cpp
+++ b/Engine/Source/Editor/Experimental/EditorInteractiveToolsFramework/Private/EditorDragTools/EditorViewportChange.cpp
@@ -17,8 +17,13 @@ FEditorViewportChange::FEditorViewportChange(FEditorViewportClient* InEditorView
 	bConvertDelta = false;
 }
 
-void FEditorViewportChange::Render(const FSceneView* View, FCanvas* Canvas)
+void FEditorViewportChange::Render(const FSceneView* View, FCanvas* Canvas, EViewInteractionState InInteractionState)
 {
+	if (InInteractionState != EViewInteractionState::Focused)
+	{
+		return;
+	}
+
 	FCanvasLineItem LineItem(Start, End);
 	Canvas->DrawItem(LineItem);
 
diff --git a/Engine/Source/Editor/Experimental/EditorInteractiveToolsFramework/Public/EditorDragTools/DragToolsBehaviourSource.h b/Engine/Source/Editor/Experimental/EditorInteractiveToolsFramework/Public/EditorDragTools/DragToolsBehaviourSource.h
--- a/Engine/Source/Editor/Experimental/EditorInteractiveToolsFramework/Public/EditorDragTools/DragToolsBehaviourSource.h
+++ b/Engine/Source/Editor/Experimental/EditorInteractiveToolsFramework/Public/EditorDragTools/DragToolsBehaviourSource.h
@@ -81,7 +81,7 @@ public:
 	/**
 	 * Renders the active tool on the specified View/Canvas
 	 */
-	EDITORINTERACTIVETOOLSFRAMEWORK_API void RenderTools(const FSceneView* InSceneView, FCanvas* InCanvas) const;
+	EDITORINTERACTIVETOOLSFRAMEWORK_API void RenderTools(FViewportClient* InViewportClient, const FSceneView* InSceneView, FCanvas* InCanvas) const;
 
 private:
 	void OnActivateTool(EDragTool::Type InDragToolType);
@@ -95,7 +95,9 @@ private:
 	 */
 	FEditorDragToolBehaviorTarget* GetActiveTool() const;
 
-	FEditorViewportClient* GetEditorViewportClient() const;
+	FEditorViewportClient* GetFocusedEditorViewportClient() const;
+
+	FEditorViewportClient* GetHoveredEditorViewportClient() const;
 
 	/**
 	 * Hosting drag tools behaviors
diff --git a/Engine/Source/Editor/Experimental/EditorInteractiveToolsFramework/Public/EditorDragTools/EditorDragToolBehaviorTarget.h b/Engine/Source/Editor/Experimental/EditorInteractiveToolsFramework/Public/EditorDragTools/EditorDragToolBehaviorTarget.h
--- a/Engine/Source/Editor/Experimental/EditorInteractiveToolsFramework/Public/EditorDragTools/EditorDragToolBehaviorTarget.h
+++ b/Engine/Source/Editor/Experimental/EditorInteractiveToolsFramework/Public/EditorDragTools/EditorDragToolBehaviorTarget.h
@@ -4,6 +4,7 @@
 
 #include "BaseBehaviors/BehaviorTargetInterfaces.h"
 #include "Editor.h"
+#include "ToolContextInterfaces.h"
 
 class IEditorViewportClientProxy;
 struct FWorldSelectionElementArgs;
@@ -48,7 +49,7 @@ class FEditorDragToolBehaviorTarget : public IClickDragBehaviorTarget
 public:
 	FEditorDragToolBehaviorTarget(FEditorViewportClient* InEditorViewportClient);
 
-	virtual void Render(const FSceneView* InView, FCanvas* InCanvas)
+	virtual void Render(const FSceneView* InView, FCanvas* InCanvas, EViewInteractionState InInteractionState)
 	{
 	}
 
diff --git a/Engine/Source/Editor/Experimental/EditorInteractiveToolsFramework/Public/EditorDragTools/EditorMarqueeSelect.h b/Engine/Source/Editor/Experimental/EditorInteractiveToolsFramework/Public/EditorDragTools/EditorMarqueeSelect.h
--- a/Engine/Source/Editor/Experimental/EditorInteractiveToolsFramework/Public/EditorDragTools/EditorMarqueeSelect.h
+++ b/Engine/Source/Editor/Experimental/EditorInteractiveToolsFramework/Public/EditorDragTools/EditorMarqueeSelect.h
@@ -19,7 +19,7 @@ public:
 	{
 	}
 
-	virtual void Render(const FSceneView* View, FCanvas* Canvas);
+	virtual void Render(const FSceneView* View, FCanvas* Canvas, EViewInteractionState InInteractionState) override;
 	
 protected:
 
diff --git a/Engine/Source/Editor/Experimental/EditorInteractiveToolsFramework/Public/EditorDragTools/EditorMeasureTool.h b/Engine/Source/Editor/Experimental/EditorInteractiveToolsFramework/Public/EditorDragTools/EditorMeasureTool.h
--- a/Engine/Source/Editor/Experimental/EditorInteractiveToolsFramework/Public/EditorDragTools/EditorMeasureTool.h
+++ b/Engine/Source/Editor/Experimental/EditorInteractiveToolsFramework/Public/EditorDragTools/EditorMeasureTool.h
@@ -11,7 +11,7 @@ class FEditorMeasureTool : public FEditorDragToolBehaviorTarget
 public:
 	explicit FEditorMeasureTool(FEditorViewportClient* InViewportClient);
 
-	virtual void Render(const FSceneView* InView,FCanvas* InCanvas) override;
+	virtual void Render(const FSceneView* InView,FCanvas* InCanvas, EViewInteractionState InInteractionState) override;
 
 	//~ Begin IClickDragBehaviorTarget
 	virtual FInputRayHit CanBeginClickDragSequence(const FInputDeviceRay& InPressPos) override;
diff --git a/Engine/Source/Editor/Experimental/EditorInteractiveToolsFramework/Public/EditorDragTools/EditorViewportChange.h b/Engine/Source/Editor/Experimental/EditorInteractiveToolsFramework/Public/EditorDragTools/EditorViewportChange.h
--- a/Engine/Source/Editor/Experimental/EditorInteractiveToolsFramework/Public/EditorDragTools/EditorViewportChange.h
+++ b/Engine/Source/Editor/Experimental/EditorInteractiveToolsFramework/Public/EditorDragTools/EditorViewportChange.h
@@ -14,7 +14,7 @@ class FEditorViewportChange : public FEditorDragToolBehaviorTarget
 {
 public:
 	explicit FEditorViewportChange(FEditorViewportClient* InEditorViewportClient);
-	virtual void Render(const FSceneView* View, FCanvas* Canvas) override;
+	virtual void Render(const FSceneView* View, FCanvas* Canvas, EViewInteractionState InInteractionState) override;
 
 	//~ Begin IClickDragBehaviorTarget
 	virtual FInputRayHit CanBeginClickDragSequence(const FInputDeviceRay& InPressPos) override;
diff --git a/Engine/Source/Editor/UnrealEd/Private/Tools/EdModeInteractiveToolsContext.cpp b/Engine/Source/Editor/UnrealEd/Private/Tools/EdModeInteractiveToolsContext.cpp
--- a/Engine/Source/Editor/UnrealEd/Private/Tools/EdModeInteractiveToolsContext.cpp
+++ b/Engine/Source/Editor/UnrealEd/Private/Tools/EdModeInteractiveToolsContext.cpp
@@ -929,7 +929,7 @@ void UModeManagerInteractiveToolsContext::DrawHUD(FViewportClient* ViewportClien
 
 	if (DragToolsBehaviorSource)
 	{
-		DragToolsBehaviorSource->RenderTools(View, Canvas);
+		DragToolsBehaviorSource->RenderTools(ViewportClient, View, Canvas);
 	}
 }
 
