[ITF Viewport Interactions] Add ViewportClickDragBehavior supporting capture stealing

A new ViewportClickDragBehavior dedicated to Viewport Interactions. Basically a MultiButtonClickDragBehavior directly supporting threshold-based capture .
This can be used when capture stealing e.g. a single click behavior

diff --git a/Engine/Source/Editor/Experimental/EditorInteractiveToolsFramework/Private/Behaviors/ViewportClickDragBehavior.cpp b/Engine/Source/Editor/Experimental/EditorInteractiveToolsFramework/Private/Behaviors/ViewportClickDragBehavior.cpp
new file mode 100644
--- /dev/null
+++ b/Engine/Source/Editor/Experimental/EditorInteractiveToolsFramework/Private/Behaviors/ViewportClickDragBehavior.cpp
@@ -0,0 +1,73 @@
+// Copyright Epic Games, Inc. All Rights Reserved.
+
+#include "Behaviors/ViewportClickDragBehavior.h"
+
+void UViewportClickDragBehavior::SetDragConfirmedPriority(FInputCapturePriority InDragConfirmedPriority)
+{
+	DragConfirmedPriority = InDragConfirmedPriority;
+}
+
+FInputCapturePriority UViewportClickDragBehavior::GetPriority()
+{
+	if (MouseTraveledDistance > ClickDistanceThreshold)
+	{
+		return DragConfirmedPriority;
+	}
+
+	return Super::GetPriority();
+}
+
+FInputCaptureRequest UViewportClickDragBehavior::WantsCapture(const FInputDeviceState& InInputState)
+{
+	bInClickDrag = false;	// should never be true here, but weird things can happen w/ focus
+
+	const bool bIsPressed = IsAnyButtonPressed(InInputState);
+
+	if (bIsPressed)
+	{
+		// Traveled distance is 0.0f if mouse was just pressed
+		MouseTraveledDistance = 0.0f;
+
+		// Cache button down location
+		MouseDownPosition = InInputState.Mouse.Position2D;
+	}
+	else
+	{
+		MouseTraveledDistance = static_cast<float>(FVector2D::Distance(InInputState.Mouse.Position2D, MouseDownPosition));
+	}
+
+	bool bAnyButtonDown = IsAnyButtonDown(InInputState);
+
+	// When stealing, we check button down state and traveled distance
+	bool bIsStealing = MouseTraveledDistance > ClickDistanceThreshold && bAnyButtonDown;
+
+	if ((bAnyButtonDown || bIsStealing) && (ModifierCheckFunc == nullptr || ModifierCheckFunc(InInputState)) )
+	{
+		const FInputRayHit HitResult = CanBeginClickDragSequence(GetDeviceRay(InInputState));
+		if (HitResult.bHit)
+		{
+			return FInputCaptureRequest::Begin(this, EInputCaptureSide::Any, HitResult.HitDepth);
+		}
+	}
+	return FInputCaptureRequest::Ignore();
+}
+
+bool UViewportClickDragBehavior::IsAnyButtonPressed(const FInputDeviceState& InInput)
+{
+	if (InInput.IsFromDevice(EInputDevices::Mouse)) 
+	{
+		ActiveDevice = EInputDevices::Mouse;
+
+		const FDeviceButtonState& Left = InInput.Mouse.Left;
+		const FDeviceButtonState& Middle = InInput.Mouse.Middle;
+		const FDeviceButtonState& Right = InInput.Mouse.Right;
+
+		const bool bLeftDown = HandlesLeftMouseButton() ? Left.bPressed : false;
+		const bool bMiddleDown = HandlesMiddleMouseButton() ? Middle.bPressed : false;
+		const bool bRightDown = HandlesRightMouseButton() ? Right.bPressed : false;
+
+		return bLeftDown || bMiddleDown || bRightDown;
+	}
+
+	return false;
+}
diff --git a/Engine/Source/Editor/Experimental/EditorInteractiveToolsFramework/Public/Behaviors/ViewportClickDragBehavior.h b/Engine/Source/Editor/Experimental/EditorInteractiveToolsFramework/Public/Behaviors/ViewportClickDragBehavior.h
new file mode 100644
--- /dev/null
+++ b/Engine/Source/Editor/Experimental/EditorInteractiveToolsFramework/Public/Behaviors/ViewportClickDragBehavior.h
@@ -0,0 +1,46 @@
+// Copyright Epic Games, Inc. All Rights Reserved.
+
+#pragma once
+
+#include "MultiButtonClickDragBehavior.h"
+#include "ViewportClickDragBehavior.generated.h"
+
+/**
+ * 
+ */
+UCLASS(MinimalAPI)
+class UViewportClickDragBehavior : public UMultiButtonClickDragBehavior
+{
+	GENERATED_BODY()
+
+public:
+
+	/**
+	 * Set the priority this Drag Behavior will return from GetPriority once the current interaction is confirmed to be a drag.
+	 * Should be higher than its DefaultPriority, and higher than the priority of any other behavior it wants to steal capture from
+	 */
+	EDITORINTERACTIVETOOLSFRAMEWORK_API void SetDragConfirmedPriority(FInputCapturePriority InDragConfirmedPriority);
+
+	EDITORINTERACTIVETOOLSFRAMEWORK_API virtual FInputCapturePriority GetPriority() override;
+
+	virtual FInputCaptureRequest WantsCapture(const FInputDeviceState& InputState) override;
+
+protected:
+
+	/** @return true if any captured button was pressed. */
+	bool IsAnyButtonPressed(const FInputDeviceState& InInput);
+
+	/** Distance after which a drag is considered such */
+	float ClickDistanceThreshold = 5.0f;
+
+	/** Cached cursor position */
+	FVector2D MouseDownPosition;
+
+	float MouseTraveledDistance = 0.0f;
+
+	/**
+	 * The priority this Drag Behavior will return from GetPriority once the current interaction is confirmed to be a drag.
+	 * Should be higher than its DefaultPriority, and higher than the priority of any other behavior it wants to steal capture from
+	 */
+	FInputCapturePriority DragConfirmedPriority;
+};
