[Transform Gizmo] Allow gizmo to be manipulated in game mode (within the editor)

diff --git a/Engine/Source/Editor/Experimental/EditorInteractiveToolsFramework/Private/EditorGizmos/EditorTransformGizmoSource.cpp b/Engine/Source/Editor/Experimental/EditorInteractiveToolsFramework/Private/EditorGizmos/EditorTransformGizmoSource.cpp
--- a/Engine/Source/Editor/Experimental/EditorInteractiveToolsFramework/Private/EditorGizmos/EditorTransformGizmoSource.cpp
+++ b/Engine/Source/Editor/Experimental/EditorInteractiveToolsFramework/Private/EditorGizmos/EditorTransformGizmoSource.cpp
@@ -76,13 +76,17 @@ bool UEditorTransformGizmoSource::GetVisible(const EViewportContext InViewportCo
 			return false;
 		}
 
+		const bool bEngineShowFlagsModeWidget = ViewportClient->EngineShowFlags.ModeWidgets;
+		if (!bEngineShowFlagsModeWidget)
+		{
+			return false;
+		}
+
 		if (InViewportContext == EViewportContext::Hovered)
 		{
-			const bool bEngineShowFlagsModeWidget = ViewportClient->EngineShowFlags.ModeWidgets && !ViewportClient->EngineShowFlags.Game;
-			if (!bEngineShowFlagsModeWidget)
-			{
-				return false;
-			}
+			// Always show for hover even in game mode
+			constexpr bool bCanHoverInGameMode = true;
+			return bCanHoverInGameMode;
 		}
 		
 		return CanInteract(InViewportContext);
