[Transform Gizmo] Visibility disabled when entering game mode, but enables on explicit selection (matches legacy behavior)

diff --git a/Engine/Source/Editor/Experimental/EditorInteractiveToolsFramework/Private/EditorGizmos/TransformGizmo.cpp b/Engine/Source/Editor/Experimental/EditorInteractiveToolsFramework/Private/EditorGizmos/TransformGizmo.cpp
--- a/Engine/Source/Editor/Experimental/EditorInteractiveToolsFramework/Private/EditorGizmos/TransformGizmo.cpp
+++ b/Engine/Source/Editor/Experimental/EditorInteractiveToolsFramework/Private/EditorGizmos/TransformGizmo.cpp
@@ -483,10 +483,15 @@ void UTransformGizmo::Render(IToolsContextRenderAPI* RenderAPI)
 	{
 		CurrentTransform = ActiveTarget->GetTransform();
 
-		UGizmoElementBase::FRenderTraversalState RenderState;
-		RenderState.Initialize(RenderAPI->GetSceneView(), GetGizmoTransform(), RenderAPI->GetCameraState().DPIScale);
-		GizmoElementRoot->Render(RenderAPI, RenderState);
+		if (IsVisible())
+		{
+			UGizmoElementBase::FRenderTraversalState RenderState;
+			RenderState.Initialize(RenderAPI->GetSceneView(), GetGizmoTransform(), RenderAPI->GetCameraState().DPIScale);
+			GizmoElementRoot->Render(RenderAPI, RenderState);
+		}
 
+
+		// Allow debug rendering regardless of actual gizmo visibility (it can still be interacted with, so we want to visualize it)
 		RenderDebug(RenderAPI);
 	}
 }
@@ -1111,7 +1116,7 @@ bool UTransformGizmo::IsVisible(const EViewportContext InViewportContext) const
 	if (TransformGizmoSource)
 	{
 		return bVisible && TransformGizmoSource->GetVisible(InViewportContext);
-	}	
+	}
 	return bVisible;
 }
 
