New TRS Gizmos: fix level's default mode being removed when opening another asset editor

#rnx

diff --git a/Engine/Source/Editor/Experimental/EditorInteractiveToolsFramework/Private/EditorGizmos/EditorTransformGizmoUtil.cpp b/Engine/Source/Editor/Experimental/EditorInteractiveToolsFramework/Private/EditorGizmos/EditorTransformGizmoUtil.cpp
--- a/Engine/Source/Editor/Experimental/EditorInteractiveToolsFramework/Private/EditorGizmos/EditorTransformGizmoUtil.cpp
+++ b/Engine/Source/Editor/Experimental/EditorInteractiveToolsFramework/Private/EditorGizmos/EditorTransformGizmoUtil.cpp
@@ -264,8 +264,20 @@ void UEditorTransformGizmoContextObject::UpdateGizmo(const TArray<FEditorViewpor
 	const UModeManagerInteractiveToolsContext* ToolsContext = ModeTools->GetInteractiveToolsContext();
 	const TObjectPtr<UInteractiveToolManager> ToolManager = ToolsContext->ToolManager;
 
-	// swap default mode
-	const bool Swapped = SwapDefaultMode(FBuiltinEditorModes::EM_Default, FAssetEdModes::EM_AssetDefault);
+	// swap default mode if not dealing with the level editor mode manager
+	const bool bIsLevelViewportClient = [&ViewportClients]()
+	{
+		if (GEditor)
+		{
+			return GEditor->GetLevelViewportClients().ContainsByPredicate([&ViewportClients](const FEditorViewportClient* LevelViewportClient)
+			{
+				return ViewportClients.Contains(LevelViewportClient);
+			});
+		}
+		return false;
+	}();
+	
+	const bool Swapped = !bIsLevelViewportClient ? SwapDefaultMode(FBuiltinEditorModes::EM_Default, FAssetEdModes::EM_AssetDefault) : false;
 	
 	if (UE::EditorTransformGizmoUtil::FindDefaultTransformGizmo(ToolManager))
 	{
