Remove `undefined()` macro from `UEditorSceneSnappingManager::ExecuteSceneSnapQuery()`.

#rnx

diff --git a/Engine/Source/Editor/Experimental/EditorInteractiveToolsFramework/Private/Snapping/EditorSnappingManager.cpp b/Engine/Source/Editor/Experimental/EditorInteractiveToolsFramework/Private/Snapping/EditorSnappingManager.cpp
--- a/Engine/Source/Editor/Experimental/EditorInteractiveToolsFramework/Private/Snapping/EditorSnappingManager.cpp
+++ b/Engine/Source/Editor/Experimental/EditorInteractiveToolsFramework/Private/Snapping/EditorSnappingManager.cpp
@@ -113,7 +113,6 @@ bool UEditorSceneSnappingManager::ExecuteSceneHitQuery(const FSceneHitQueryReque
 
 bool UEditorSceneSnappingManager::ExecuteSceneSnapQuery(const FSceneSnapQueryRequest& InRequest, TArray<FSceneSnapQueryResult>& OutResults) const
 {
-	unimplemented();
-
+	// @todo: implement when needed/supported
 	return false;
 }
