[ITF Viewport Interactions] Click based Viewport Interactions

- Add click based Viewport Interaction class (base class for click interactions)
- Update behaviors priority handling according to new ITF Capture Stealing feature (helps drag and click behavior to work better together)
- Change behavior type for Gizmo indirect interactions
- Update priority also for Gizmo indirect drag behaviors

diff --git a/Engine/Source/Editor/Experimental/EditorInteractiveToolsFramework/Private/EditorDragTools/DragToolsBehaviourSource.cpp b/Engine/Source/Editor/Experimental/EditorInteractiveToolsFramework/Private/EditorDragTools/DragToolsBehaviourSource.cpp
--- a/Engine/Source/Editor/Experimental/EditorInteractiveToolsFramework/Private/EditorDragTools/DragToolsBehaviourSource.cpp
+++ b/Engine/Source/Editor/Experimental/EditorInteractiveToolsFramework/Private/EditorDragTools/DragToolsBehaviourSource.cpp
@@ -74,6 +74,7 @@ void UDragToolsBehaviorSource::Initialize(const UEditorInteractiveToolsContext*
 		UClickDragInputBehavior* MouseBehavior = NewObject<UClickDragInputBehavior>();
 		MouseBehavior->Initialize(BoxSelectTool.Get());
 		MouseBehavior->SetDefaultPriority(BOX_SELECT_PRIORITY);
+		MouseBehavior->SetAllowsCaptureStealing(true);
 		BehaviorSet->Add(MouseBehavior);
 
 		// Register modifier listeners
diff --git a/Engine/Source/Editor/Experimental/EditorInteractiveToolsFramework/Private/EditorGizmos/TransformGizmo.cpp b/Engine/Source/Editor/Experimental/EditorInteractiveToolsFramework/Private/EditorGizmos/TransformGizmo.cpp
--- a/Engine/Source/Editor/Experimental/EditorInteractiveToolsFramework/Private/EditorGizmos/TransformGizmo.cpp
+++ b/Engine/Source/Editor/Experimental/EditorInteractiveToolsFramework/Private/EditorGizmos/TransformGizmo.cpp
@@ -18,7 +18,7 @@
 #include "Materials/Material.h"
 #include "Misc/AxisDisplayInfo.h"
 #include "UnrealEngine.h"
-#include "Behaviors/MultiButtonClickDragBehavior.h"
+#include "Behaviors/ViewportClickDragBehavior.h"
 #include "Intersection/IntersectionUtil.h"
 #include "PrimitiveDrawingUtils.h"
 #include "SceneQueries/SceneSnappingManager.h"
@@ -222,10 +222,11 @@ void UTransformGizmo::SetupIndirectBehaviors()
 	static constexpr int INDIRECT_GIZMO_PRIORITY = FInputCapturePriority::DEFAULT_GIZMO_PRIORITY-1;
 	
 	// Add middle mouse input behavior for indirect manipulation
-	ULocalClickDragInputBehavior* MiddleClickDragBehavior = NewObject<ULocalClickDragInputBehavior>();
+	UViewportClickDragBehavior* MiddleClickDragBehavior = NewObject<UViewportClickDragBehavior>();
 	MiddleClickDragBehavior->Initialize();
-	MiddleClickDragBehavior->SetDefaultPriority(INDIRECT_GIZMO_PRIORITY);
-	MiddleClickDragBehavior->SetUseMiddleMouseButton();
+	MiddleClickDragBehavior->SetDefaultPriority(FInputCapturePriority::DEFAULT_GIZMO_PRIORITY + 5);
+	MiddleClickDragBehavior->SetDragConfirmedPriority(INDIRECT_GIZMO_PRIORITY);
+	MiddleClickDragBehavior->EnableButton(EKeys::MiddleMouseButton);
 	MiddleClickDragBehavior->CanBeginClickDragFunc = [this](const FInputDeviceRay&)
 	{
 		static const FInputRayHit InvalidRayHit;
@@ -264,10 +265,11 @@ void UTransformGizmo::SetupIndirectBehaviors()
 	AddInputBehavior(MiddleClickDragBehavior);
 
 	// Add left/right mouse input behavior for indirect manipulation
-	MultiIndirectClickDragBehavior = NewObject<UMultiButtonClickDragBehavior>();
+	MultiIndirectClickDragBehavior = NewObject<UViewportClickDragBehavior>();
 	MultiIndirectClickDragBehavior->Initialize();
-	MultiIndirectClickDragBehavior->SetDefaultPriority(INDIRECT_GIZMO_PRIORITY);
 	MultiIndirectClickDragBehavior->EnableButton(EKeys::LeftMouseButton);
+	MultiIndirectClickDragBehavior->SetDefaultPriority(FInputCapturePriority::DEFAULT_GIZMO_PRIORITY + 5);
+	MultiIndirectClickDragBehavior->SetDragConfirmedPriority(INDIRECT_GIZMO_PRIORITY);
 	if (bCtrlMiddleDoesY)
 	{
 		MultiIndirectClickDragBehavior->EnableButton(EKeys::MiddleMouseButton);
diff --git a/Engine/Source/Editor/Experimental/EditorInteractiveToolsFramework/Private/ViewportInteractions/ViewportCameraRotateInteraction.cpp b/Engine/Source/Editor/Experimental/EditorInteractiveToolsFramework/Private/ViewportInteractions/ViewportCameraRotateInteraction.cpp
--- a/Engine/Source/Editor/Experimental/EditorInteractiveToolsFramework/Private/ViewportInteractions/ViewportCameraRotateInteraction.cpp
+++ b/Engine/Source/Editor/Experimental/EditorInteractiveToolsFramework/Private/ViewportInteractions/ViewportCameraRotateInteraction.cpp
@@ -34,6 +34,16 @@ UViewportCameraRotateInteraction::UViewportCameraRotateInteraction()
 				|| CommandMatchesKey(FViewportNavigationCommands::Get().RotateRight, ActiveKey))
 			{
 				bResult = IsAnyMouseButtonDown();
+				// If a mouse button is down, but mouse hasn't been dragged yet, we need to tell the behavior source that
+				// User is performing mouse look or at least camera movement. This also prevents regular clicks from performing their action on button release
+				// e.g. we don't want a RMB release to summon a context menu after mouse looking
+				if (bResult)
+				{
+					if (UViewportInteractionsBehaviorSource* BehaviorSource = GetViewportInteractionsBehaviorSource())
+					{
+						BehaviorSource->SetIsMouseLooking(true);
+					}
+				}
 			}
 		}
 
@@ -84,23 +94,23 @@ void UViewportCameraRotateInteraction::UpdateKeyState(const FKey& InKeyID, bool
 
 	constexpr float ImpulseValue = 1.0f;
 
-	const bool bMouseButtonDownOrRelease = IsAnyMouseButtonDown() || !bInIsPressed;
+	const bool bMouseLookingOrRelease = IsMouseLooking() || !bInIsPressed;
 
 	// Rotate Up/Down
-	if (CommandMatchesKey(FViewportNavigationCommands::Get().RotateUp, InKeyID) && bMouseButtonDownOrRelease)
+	if (CommandMatchesKey(FViewportNavigationCommands::Get().RotateUp, InKeyID) && bMouseLookingOrRelease)
 	{
 		RotatePitchImpulse = bInIsPressed ? ImpulseValue : 0.0f;
 	}
-	else if (CommandMatchesKey(FViewportNavigationCommands::Get().RotateDown, InKeyID) && bMouseButtonDownOrRelease)
+	else if (CommandMatchesKey(FViewportNavigationCommands::Get().RotateDown, InKeyID) && bMouseLookingOrRelease)
 	{
 		RotatePitchImpulse = bInIsPressed ? -ImpulseValue : 0.0f;
 	}
 	// Rotate Left/Right
-	else if (CommandMatchesKey(FViewportNavigationCommands::Get().RotateLeft, InKeyID) && bMouseButtonDownOrRelease)
+	else if (CommandMatchesKey(FViewportNavigationCommands::Get().RotateLeft, InKeyID) && bMouseLookingOrRelease)
 	{
 		RotateYawImpulse = bInIsPressed ? -ImpulseValue : 0.0f;
 	}
-	else if (CommandMatchesKey(FViewportNavigationCommands::Get().RotateRight, InKeyID) && bMouseButtonDownOrRelease)
+	else if (CommandMatchesKey(FViewportNavigationCommands::Get().RotateRight, InKeyID) && bMouseLookingOrRelease)
 	{
 		RotateYawImpulse = bInIsPressed ? ImpulseValue : 0.0f;
 	}
diff --git a/Engine/Source/Editor/Experimental/EditorInteractiveToolsFramework/Private/ViewportInteractions/ViewportCameraSpeedMouseWheelInteraction.cpp b/Engine/Source/Editor/Experimental/EditorInteractiveToolsFramework/Private/ViewportInteractions/ViewportCameraSpeedMouseWheelInteraction.cpp
--- a/Engine/Source/Editor/Experimental/EditorInteractiveToolsFramework/Private/ViewportInteractions/ViewportCameraSpeedMouseWheelInteraction.cpp
+++ b/Engine/Source/Editor/Experimental/EditorInteractiveToolsFramework/Private/ViewportInteractions/ViewportCameraSpeedMouseWheelInteraction.cpp
@@ -1,4 +1,4 @@
-﻿// Copyright Epic Games, Inc. All Rights Reserved.
+// Copyright Epic Games, Inc. All Rights Reserved.
 
 #include "ViewportInteractions/ViewportCameraSpeedMouseWheelInteraction.h"
 
@@ -14,7 +14,7 @@ UViewportCameraSpeedMouseWheelInteraction::UViewportCameraSpeedMouseWheelInterac
 	// Set camera speed using mouse wheel
 	UMouseWheelInputBehavior* MouseWheelInputBehavior = NewObject<UMouseWheelModifierInputBehavior>();
 	MouseWheelInputBehavior->Initialize(this);
-	MouseWheelInputBehavior->SetDefaultPriority(UE::Editor::ViewportInteractions::VIEWPORT_INTERACTIONS_HIGH_PRIORITY);
+	MouseWheelInputBehavior->SetDefaultPriority(UE::Editor::ViewportInteractions::VIEWPORT_INTERACTIONS_DEFAULT_PRIORITY);
 
 	RegisterInputBehavior(MouseWheelInputBehavior);
 }
diff --git a/Engine/Source/Editor/Experimental/EditorInteractiveToolsFramework/Private/ViewportInteractions/ViewportCameraTranslateInteraction.cpp b/Engine/Source/Editor/Experimental/EditorInteractiveToolsFramework/Private/ViewportInteractions/ViewportCameraTranslateInteraction.cpp
--- a/Engine/Source/Editor/Experimental/EditorInteractiveToolsFramework/Private/ViewportInteractions/ViewportCameraTranslateInteraction.cpp
+++ b/Engine/Source/Editor/Experimental/EditorInteractiveToolsFramework/Private/ViewportInteractions/ViewportCameraTranslateInteraction.cpp
@@ -53,7 +53,17 @@ UViewportCameraTranslateInteraction::UViewportCameraTranslateInteraction()
 				CommandMatchesKey(FViewportNavigationCommands::Get().LocalUp, ActiveKey) ||
 				CommandMatchesKey(FViewportNavigationCommands::Get().LocalDown, ActiveKey))
 			{
-				bResult = IsMouseLooking(); 
+				bResult = IsAnyMouseButtonDown();
+				// If a mouse button is down, but mouse hasn't been dragged yet, we need to tell the behavior source that
+				// User is performing mouse look or at least camera movement. This also prevents regular clicks from performing their action on button release
+				// e.g. we don't want a RMB release to summon a context menu after mouse looking
+				if (bResult)
+				{
+					if (UViewportInteractionsBehaviorSource* BehaviorSource = GetViewportInteractionsBehaviorSource())
+					{
+						BehaviorSource->SetIsMouseLooking(true);
+					}
+				}
 			}
 		}
 
diff --git a/Engine/Source/Editor/Experimental/EditorInteractiveToolsFramework/Private/ViewportInteractions/ViewportClickInteraction.cpp b/Engine/Source/Editor/Experimental/EditorInteractiveToolsFramework/Private/ViewportInteractions/ViewportClickInteraction.cpp
new file mode 100644
--- /dev/null
+++ b/Engine/Source/Editor/Experimental/EditorInteractiveToolsFramework/Private/ViewportInteractions/ViewportClickInteraction.cpp
@@ -0,0 +1,81 @@
+// Copyright Epic Games, Inc. All Rights Reserved.
+
+#include "ViewportInteractions/ViewportClickInteraction.h"
+
+#include "BaseBehaviors/DoubleClickBehavior.h"
+#include "BaseBehaviors/SingleClickBehavior.h"
+#include "EditorViewportClient.h"
+#include "SceneView.h"
+#include "ViewportInteractions/ViewportInteractionsBehaviorSource.h"
+
+UViewportClickInteraction::UViewportClickInteraction()
+{
+	USingleClickInputBehavior* SingleClickBehavior = NewObject<USingleClickInputBehavior>();
+	SingleClickBehavior->Initialize(this);
+
+	SingleClickBehavior->SetDefaultPriority(UE::Editor::ViewportInteractions::CLICK_PRIORITY);
+	SingleClickBehavior->SetAllowsCaptureStealing(true);
+	RegisterInputBehavior(SingleClickBehavior);
+
+	SingleClickInputBehaviorWeak = SingleClickBehavior;
+}
+
+FInputRayHit UViewportClickInteraction::IsHitByClick(const FInputDeviceRay& InClickPos)
+{
+	return FInputRayHit(0.0f);
+}
+
+void UViewportClickInteraction::OnClicked(const FInputDeviceRay& InClickPos)
+{
+	// Make sure this release is not happening while user is moving the viewport camera - in that case, this is not considered a click.
+	if (IsMouseLooking())
+	{
+		if (UViewportInteractionsBehaviorSource* const ViewportInteractionBehaviorSource = GetViewportInteractionsBehaviorSource())
+		{
+			ViewportInteractionBehaviorSource->SetIsMouseLooking(false);
+		}
+
+		return;
+	}
+
+	const int32 CurrentClickPosX = static_cast<int32>(InClickPos.ScreenPosition.X);
+	const int32 CurrentClickPosY = static_cast<int32>(InClickPos.ScreenPosition.Y);
+
+	if (FEditorViewportClient* EditorViewportClient = GetEditorViewportClient())
+	{
+		if (FViewport* Viewport = EditorViewportClient->Viewport)
+		{
+			TRefCountPtr<HHitProxy> HitProxy = Viewport->GetHitProxy(CurrentClickPosX, CurrentClickPosY);
+
+			// Compute a view.
+			FSceneViewFamilyContext ViewFamily(
+				FSceneViewFamily::ConstructionValues(Viewport, EditorViewportClient->GetScene(), EditorViewportClient->EngineShowFlags)
+					.SetRealtimeUpdate(EditorViewportClient->IsRealtime())
+			);
+
+			if (FSceneView* View = EditorViewportClient->CalcSceneView(&ViewFamily))
+			{
+				const FViewportClick ViewportClick(
+					View, EditorViewportClient, GetMouseKey(), GetMouseEvent(), CurrentClickPosX, CurrentClickPosY
+				);
+
+				ProcessClick_Internal(ViewportClick, HitProxy, *View);
+			}
+		}
+	}
+}
+
+FDeviceButtonState UViewportClickInteraction::GetActiveMouseButtonState(const FInputDeviceState& Input)
+{
+	return Input.Mouse.Left;
+}
+
+FKey UViewportClickInteraction::GetMouseKey() const
+{
+	return MouseButton;
+}
+
+EInputEvent UViewportClickInteraction::GetMouseEvent() const
+{
+	return IE_Released;
+}
diff --git a/Engine/Source/Editor/Experimental/EditorInteractiveToolsFramework/Private/ViewportInteractions/ViewportCommandsInteraction.cpp b/Engine/Source/Editor/Experimental/EditorInteractiveToolsFramework/Private/ViewportInteractions/ViewportCommandsInteraction.cpp
--- a/Engine/Source/Editor/Experimental/EditorInteractiveToolsFramework/Private/ViewportInteractions/ViewportCommandsInteraction.cpp
+++ b/Engine/Source/Editor/Experimental/EditorInteractiveToolsFramework/Private/ViewportInteractions/ViewportCommandsInteraction.cpp
@@ -12,7 +12,7 @@ UViewportCommandsInteraction::UViewportCommandsInteraction()
 	UKeyInputBehavior* KeyInputBehavior = NewObject<UKeyInputBehavior>();
 	KeyInputBehavior->Initialize(this, GetKeys());
 
-	KeyInputBehavior->SetDefaultPriority(UE::Editor::ViewportInteractions::VIEWPORT_INTERACTIONS_HIGH_PRIORITY);
+	KeyInputBehavior->SetDefaultPriority(UE::Editor::ViewportInteractions::VIEWPORT_INTERACTIONS_DEFAULT_PRIORITY);
 	KeyInputBehavior->bRequireAllKeys = false;
 	KeyInputBehavior->ModifierCheckFunc = [this](const FInputDeviceState& InputDeviceState)
 	{
diff --git a/Engine/Source/Editor/Experimental/EditorInteractiveToolsFramework/Private/ViewportInteractions/ViewportDragInteraction.cpp b/Engine/Source/Editor/Experimental/EditorInteractiveToolsFramework/Private/ViewportInteractions/ViewportDragInteraction.cpp
--- a/Engine/Source/Editor/Experimental/EditorInteractiveToolsFramework/Private/ViewportInteractions/ViewportDragInteraction.cpp
+++ b/Engine/Source/Editor/Experimental/EditorInteractiveToolsFramework/Private/ViewportInteractions/ViewportDragInteraction.cpp
@@ -3,6 +3,7 @@
 #include "ViewportInteractions/ViewportDragInteraction.h"
 #include "BaseBehaviors/ClickDragBehavior.h"
 #include "BaseBehaviors/SingleClickOrDragBehavior.h"
+#include "Behaviors/ViewportClickDragBehavior.h"
 #include "ViewportInteractions/ViewportInteractionsBehaviorSource.h"
 
 UViewportDragInteraction::UViewportDragInteraction()
@@ -11,13 +12,36 @@ UViewportDragInteraction::UViewportDragInteraction()
 	MouseX = 0;
 	MouseY = 0;
 
-	USingleClickOrDragInputBehavior* ClickDragBehavior = NewObject<USingleClickOrDragInputBehavior>();
+	UViewportClickDragBehavior* ClickDragBehavior = NewObject<UViewportClickDragBehavior>();
+	ClickDragBehavior->Initialize();
 
-	// Click Target can be specified separately if needed
-	ClickDragBehavior->Initialize(nullptr, this);
-	ClickDragBehavior->SetDefaultPriority(UE::Editor::ViewportInteractions::VIEWPORT_INTERACTIONS_DEFAULT_PRIORITY);
+	ClickDragBehavior->CanBeginClickDragFunc = [this](const FInputDeviceRay& InPressPos)
+	{
+		return CanBeginClickDragSequence(InPressPos);
+	};
+
+	ClickDragBehavior->OnClickPressFunc = [this](const FInputDeviceRay& InPressPos)
+	{
+		OnClickPress(InPressPos);
+	};
+
+	ClickDragBehavior->OnClickReleaseFunc = [this](const FInputDeviceRay& InReleasePos)
+	{
+		OnClickRelease(InReleasePos);
+	};
+
+	ClickDragBehavior->OnClickDragFunc = [this](const FInputDeviceRay& InDragPos)
+	{
+		OnClickDrag(InDragPos);
+	};
+
+	ClickDragBehavior->SetDefaultPriority(UE::Editor::ViewportInteractions::DRAG_PRIORITY);
+	ClickDragBehavior->SetDragConfirmedPriority(UE::Editor::ViewportInteractions::DRAG_CONFIRMED_PRIORITY);
+
+	// Other higher priority drag interactions might steal capture (e.g. Gizmo Indirect Manipulation)
+	ClickDragBehavior->SetAllowsCaptureStealing(true);
 
-	ClickOrDragInputBehavior = ClickDragBehavior;
+	ViewportClickDragBehavior = ClickDragBehavior;
 
 	RegisterInputBehavior(ClickDragBehavior);
 }
@@ -79,37 +103,29 @@ void UViewportDragInteraction::Initialize(UViewportInteractionsBehaviorSource* c
 {
 	UViewportInteraction::Initialize(InViewportInteractionsBehaviorSource);
 
-	if (ClickOrDragInputBehavior.IsValid())
+	if (ViewportClickDragBehavior.IsValid())
 	{
-		ClickOrDragInputBehavior->SetUseCustomMouseButton([this](const FInputDeviceState& InInputDeviceState)
+		ViewportClickDragBehavior->SetUseCustomMouseButton([this](const FInputDeviceState& InInputDeviceState)
 		{
 			return GetActiveMouseButtonState(InInputDeviceState);
 		});
 
-		ClickOrDragInputBehavior->ModifierCheckFunc = [this](const FInputDeviceState& InInputDeviceState)
+		ViewportClickDragBehavior->ModifierCheckFunc = [this](const FInputDeviceState& InInputDeviceState)
 		{
 			return CanBeActivated(InInputDeviceState);
 		};
 	}
 }
 
-void UViewportDragInteraction::SetClickTarget(IClickBehaviorTarget* InClickBehaviorTarget)
+FDeviceButtonState UViewportDragInteraction::GetActiveMouseButtonState(const FInputDeviceState& Input)
 {
-	if (TStrongObjectPtr<USingleClickOrDragInputBehavior> ClickOrDragBehavior = ClickOrDragInputBehavior.Pin())
-	{
-		ClickOrDragBehavior->Initialize(InClickBehaviorTarget, this);
-	}
+	return Input.Mouse.Left;
 }
 
-void UViewportDragInteraction::SetClickOrDragPriority(int InPriority) const
+void UViewportDragInteraction::EnableButton(const FKey& InButton) const
 {
-	if (TStrongObjectPtr<USingleClickOrDragInputBehavior> ClickOrDragBehavior = ClickOrDragInputBehavior.Pin())
+	if (TStrongObjectPtr<UViewportClickDragBehavior> ClickDragBehavior = ViewportClickDragBehavior.Pin())
 	{
-		ClickOrDragBehavior->SetDefaultPriority(InPriority);
+		ClickDragBehavior->EnableButton(InButton);
 	}
 }
-
-FDeviceButtonState UViewportDragInteraction::GetActiveMouseButtonState(const FInputDeviceState& Input)
-{
-	return Input.Mouse.Left;
-}
diff --git a/Engine/Source/Editor/Experimental/EditorInteractiveToolsFramework/Private/ViewportInteractions/ViewportFOVInteraction.cpp b/Engine/Source/Editor/Experimental/EditorInteractiveToolsFramework/Private/ViewportInteractions/ViewportFOVInteraction.cpp
--- a/Engine/Source/Editor/Experimental/EditorInteractiveToolsFramework/Private/ViewportInteractions/ViewportFOVInteraction.cpp
+++ b/Engine/Source/Editor/Experimental/EditorInteractiveToolsFramework/Private/ViewportInteractions/ViewportFOVInteraction.cpp
@@ -33,7 +33,17 @@ UViewportFOVInteraction::UViewportFOVInteraction()
 			if (CommandMatchesKey(FViewportNavigationCommands::Get().FovZoomOut, ActiveKey) ||
 				CommandMatchesKey(FViewportNavigationCommands::Get().FovZoomIn, ActiveKey))
 			{
-				bResult = IsAnyMouseButtonDown(); 
+				bResult = IsAnyMouseButtonDown();
+				// If a mouse button is down, but mouse hasn't been dragged yet, we need to tell the behavior source that
+				// User is performing mouse look or at least camera movement. This also prevents regular clicks from performing their action on button release
+				// e.g. we don't want a RMB release to summon a context menu after mouse looking
+				if (bResult)
+				{
+					if (UViewportInteractionsBehaviorSource* BehaviorSource = GetViewportInteractionsBehaviorSource())
+					{
+						BehaviorSource->SetIsMouseLooking(true);
+					}
+				}
 			}
 		}
 
diff --git a/Engine/Source/Editor/Experimental/EditorInteractiveToolsFramework/Private/ViewportInteractions/ViewportInteractionsBehaviorSource.cpp b/Engine/Source/Editor/Experimental/EditorInteractiveToolsFramework/Private/ViewportInteractions/ViewportInteractionsBehaviorSource.cpp
--- a/Engine/Source/Editor/Experimental/EditorInteractiveToolsFramework/Private/ViewportInteractions/ViewportInteractionsBehaviorSource.cpp
+++ b/Engine/Source/Editor/Experimental/EditorInteractiveToolsFramework/Private/ViewportInteractions/ViewportInteractionsBehaviorSource.cpp
@@ -95,7 +95,7 @@ void LOG(const TCHAR* InMessage)
 	}
 }
 
-static constexpr int MODIFIERS_PRIORITY = VIEWPORT_INTERACTIONS_HIGH_PRIORITY;
+static constexpr int MODIFIERS_PRIORITY = VIEWPORT_INTERACTIONS_DEFAULT_PRIORITY;
 
 // We use this string to group Behavior Inputs so that we can selectively remove them if needed, without touching other inputs we might add (e.g. the ones added by Initialize)
 static const FString ViewportInteractionsGroupName = TEXT("ViewportInteractions");
diff --git a/Engine/Source/Editor/Experimental/EditorInteractiveToolsFramework/Private/ViewportInteractions/ViewportMoveYawInteraction.cpp b/Engine/Source/Editor/Experimental/EditorInteractiveToolsFramework/Private/ViewportInteractions/ViewportMoveYawInteraction.cpp
--- a/Engine/Source/Editor/Experimental/EditorInteractiveToolsFramework/Private/ViewportInteractions/ViewportMoveYawInteraction.cpp
+++ b/Engine/Source/Editor/Experimental/EditorInteractiveToolsFramework/Private/ViewportInteractions/ViewportMoveYawInteraction.cpp
@@ -2,7 +2,6 @@
 
 #include "ViewportInteractions/ViewportMoveYawInteraction.h"
 #include "BaseBehaviors/ClickDragBehavior.h"
-#include "BaseBehaviors/SingleClickOrDragBehavior.h"
 #include "EditorViewportClient.h"
 #include "Settings/LevelEditorViewportSettings.h"
 #include "ViewportClientNavigationHelper.h"
@@ -15,7 +14,7 @@ UViewportMoveYawInteraction::UViewportMoveYawInteraction()
 
 	InteractionName = TEXT("Move-Yaw");
 
-	SetClickOrDragPriority(FInputCapturePriority::DEFAULT_GIZMO_PRIORITY + 1);
+	EnableButton(EKeys::LeftMouseButton);
 }
 
 FInputRayHit UViewportMoveYawInteraction::CanBeginClickDragSequence(const FInputDeviceRay& InPressPos)
diff --git a/Engine/Source/Editor/Experimental/EditorInteractiveToolsFramework/Private/ViewportInteractions/ViewportOrbitInteraction.cpp b/Engine/Source/Editor/Experimental/EditorInteractiveToolsFramework/Private/ViewportInteractions/ViewportOrbitInteraction.cpp
--- a/Engine/Source/Editor/Experimental/EditorInteractiveToolsFramework/Private/ViewportInteractions/ViewportOrbitInteraction.cpp
+++ b/Engine/Source/Editor/Experimental/EditorInteractiveToolsFramework/Private/ViewportInteractions/ViewportOrbitInteraction.cpp
@@ -4,7 +4,6 @@
 #include "BaseBehaviors/ClickDragBehavior.h"
 #include "EditorViewportClient.h"
 #include "ViewportClientNavigationHelper.h"
-#include "BaseBehaviors/SingleClickOrDragBehavior.h"
 #include "ViewportInteractions/ViewportInteractionsBehaviorSource.h"
 
 UViewportOrbitInteraction::UViewportOrbitInteraction()
@@ -13,12 +12,7 @@ UViewportOrbitInteraction::UViewportOrbitInteraction()
 
 	InteractionName = TEXT("Orbit");
 
-	if (TStrongObjectPtr<USingleClickOrDragInputBehavior> ClickOrDragBehavior = ClickOrDragInputBehavior.Pin())
-	{
-		// Priority lower than gizmos, so that we can still Alt - Duplicate.
-		// This might need adjustments once we have all the pieces together (gizmo, drag tools, viewport interactions)
-		ClickOrDragBehavior->SetDefaultPriority(FInputCapturePriority::DEFAULT_GIZMO_PRIORITY + 1);
-	}
+	EnableButton(EKeys::LeftMouseButton);
 }
 
 FInputRayHit UViewportOrbitInteraction::CanBeginClickDragSequence(const FInputDeviceRay& InPressPos)
diff --git a/Engine/Source/Editor/Experimental/EditorInteractiveToolsFramework/Private/ViewportInteractions/ViewportOrthoPanInteraction.cpp b/Engine/Source/Editor/Experimental/EditorInteractiveToolsFramework/Private/ViewportInteractions/ViewportOrthoPanInteraction.cpp
--- a/Engine/Source/Editor/Experimental/EditorInteractiveToolsFramework/Private/ViewportInteractions/ViewportOrthoPanInteraction.cpp
+++ b/Engine/Source/Editor/Experimental/EditorInteractiveToolsFramework/Private/ViewportInteractions/ViewportOrthoPanInteraction.cpp
@@ -11,6 +11,8 @@ UViewportOrthoPanInteraction::UViewportOrthoPanInteraction()
 	ToolType = UE::Editor::ViewportInteractions::CameraDrag;
 
 	InteractionName = TEXT("Orthographic Panning");
+
+	EnableButton(EKeys::RightMouseButton);
 }
 
 FInputRayHit UViewportOrthoPanInteraction::CanBeginClickDragSequence(const FInputDeviceRay& InPressPos)
@@ -100,6 +102,5 @@ FDeviceButtonState UViewportOrthoPanInteraction::GetActiveMouseButtonState(const
 }
 bool UViewportOrthoPanInteraction::CanBeActivated(const FInputDeviceState& InInputDeviceState) const
 {
-	//TODO: we might be able to remove the check for LMB
-	return IsEnabled() && !IsLeftMouseButtonDown();
+	return IsEnabled();
 }
diff --git a/Engine/Source/Editor/Experimental/EditorInteractiveToolsFramework/Private/ViewportInteractions/ViewportPanInteraction.cpp b/Engine/Source/Editor/Experimental/EditorInteractiveToolsFramework/Private/ViewportInteractions/ViewportPanInteraction.cpp
--- a/Engine/Source/Editor/Experimental/EditorInteractiveToolsFramework/Private/ViewportInteractions/ViewportPanInteraction.cpp
+++ b/Engine/Source/Editor/Experimental/EditorInteractiveToolsFramework/Private/ViewportInteractions/ViewportPanInteraction.cpp
@@ -3,6 +3,7 @@
 #include "ViewportInteractions/ViewportPanInteraction.h"
 #include "BaseBehaviors/ClickDragBehavior.h"
 #include "Behaviors/MultiButtonClickDragBehavior.h"
+#include "Behaviors/ViewportClickDragBehavior.h"
 #include "EditorViewportClient.h"
 #include "Settings/LevelEditorViewportSettings.h"
 #include "ViewportClientNavigationHelper.h"
@@ -14,8 +15,10 @@ UViewportPanInteraction::UViewportPanInteraction()
 
 	InteractionName = TEXT("Perspective Panning");
 
+	EnableButton(EKeys::MiddleMouseButton);
+
 	// LMB + RMB + Drag to Pan
-	UMultiButtonClickDragBehavior* MultiButtonClickDragBehavior = NewObject<UMultiButtonClickDragBehavior>();
+	UViewportClickDragBehavior* MultiButtonClickDragBehavior = NewObject<UViewportClickDragBehavior>();
 	MultiButtonClickDragBehavior->Initialize();
 	MultiButtonClickDragBehavior->EnableButton(EKeys::LeftMouseButton);
 	MultiButtonClickDragBehavior->EnableButton(EKeys::RightMouseButton);
@@ -52,7 +55,9 @@ UViewportPanInteraction::UViewportPanInteraction()
 		OnClickDrag(InDragPos);
 	};
 
-	MultiButtonClickDragBehavior->SetDefaultPriority(UE::Editor::ViewportInteractions::VIEWPORT_INTERACTIONS_HIGH_PRIORITY);
+	MultiButtonClickDragBehavior->SetDefaultPriority(UE::Editor::ViewportInteractions::MULTIBUTTON_DRAG_PRIORITY);
+	MultiButtonClickDragBehavior->SetDragConfirmedPriority(UE::Editor::ViewportInteractions::MULTIBUTTON_DRAG_CONFIRMED_PRIORITY);
+	MultiButtonClickDragBehavior->SetAllowsCaptureStealing(true);
 
 	RegisterInputBehavior(MultiButtonClickDragBehavior);
 }
diff --git a/Engine/Source/Editor/Experimental/EditorInteractiveToolsFramework/Private/ViewportInteractions/ViewportViewAngleInteraction.cpp b/Engine/Source/Editor/Experimental/EditorInteractiveToolsFramework/Private/ViewportInteractions/ViewportViewAngleInteraction.cpp
--- a/Engine/Source/Editor/Experimental/EditorInteractiveToolsFramework/Private/ViewportInteractions/ViewportViewAngleInteraction.cpp
+++ b/Engine/Source/Editor/Experimental/EditorInteractiveToolsFramework/Private/ViewportInteractions/ViewportViewAngleInteraction.cpp
@@ -2,7 +2,6 @@
 
 #include "ViewportInteractions/ViewportViewAngleInteraction.h"
 #include "BaseBehaviors/ClickDragBehavior.h"
-#include "BaseBehaviors/SingleClickOrDragBehavior.h"
 #include "EditorViewportClient.h"
 #include "Settings/LevelEditorViewportSettings.h"
 #include "ViewportClientNavigationHelper.h"
@@ -14,7 +13,7 @@ UViewportViewAngleInteraction::UViewportViewAngleInteraction()
 
 	InteractionName = TEXT("View Angle");
 
-	SetClickOrDragPriority(UE::Editor::ViewportInteractions::VIEWPORT_INTERACTIONS_AFTER_GIZMOS_PRIORITY);
+	EnableButton(EKeys::RightMouseButton);
 }
 
 FInputRayHit UViewportViewAngleInteraction::CanBeginClickDragSequence(const FInputDeviceRay& InPressPos)
diff --git a/Engine/Source/Editor/Experimental/EditorInteractiveToolsFramework/Private/ViewportInteractions/ViewportZoomInteraction.cpp b/Engine/Source/Editor/Experimental/EditorInteractiveToolsFramework/Private/ViewportInteractions/ViewportZoomInteraction.cpp
--- a/Engine/Source/Editor/Experimental/EditorInteractiveToolsFramework/Private/ViewportInteractions/ViewportZoomInteraction.cpp
+++ b/Engine/Source/Editor/Experimental/EditorInteractiveToolsFramework/Private/ViewportInteractions/ViewportZoomInteraction.cpp
@@ -4,6 +4,7 @@
 #include "BaseBehaviors/ClickDragBehavior.h"
 #include "BaseBehaviors/KeyInputBehavior.h"
 #include "Behaviors/MultiButtonClickDragBehavior.h"
+#include "Behaviors/ViewportClickDragBehavior.h"
 #include "EditorViewportClient.h"
 #include "Settings/LevelEditorViewportSettings.h"
 #include "Templates/Function.h"
@@ -29,7 +30,7 @@ UViewportZoomInteraction::UViewportZoomInteraction()
 	KeyInputBehavior->SetDefaultPriority(UE::Editor::ViewportInteractions::VIEWPORT_INTERACTIONS_DEFAULT_PRIORITY);
 
 	// LMB + RMB + Drag zoom.
-	UMultiButtonClickDragBehavior* MultiButtonClickDragBehavior = NewObject<UMultiButtonClickDragBehavior>();
+	UViewportClickDragBehavior* MultiButtonClickDragBehavior = NewObject<UViewportClickDragBehavior>();
 	MultiButtonClickDragBehavior->Initialize();
 	MultiButtonClickDragBehavior->EnableButton(EKeys::LeftMouseButton);
 	MultiButtonClickDragBehavior->EnableButton(EKeys::RightMouseButton);
@@ -65,7 +66,7 @@ UViewportZoomInteraction::UViewportZoomInteraction()
 		OnClickDrag(InDragPos);
 	};
 
-	MultiButtonClickDragBehavior->SetDefaultPriority(UE::Editor::ViewportInteractions::VIEWPORT_INTERACTIONS_HIGH_PRIORITY);
+	MultiButtonClickDragBehavior->SetDefaultPriority(UE::Editor::ViewportInteractions::VIEWPORT_INTERACTIONS_DEFAULT_PRIORITY);
 
 	RegisterInputBehaviors({ MouseWheelInputBehavior, KeyInputBehavior, MultiButtonClickDragBehavior });
 }
diff --git a/Engine/Source/Editor/Experimental/EditorInteractiveToolsFramework/Public/EditorDragTools/DragToolsBehaviourSource.h b/Engine/Source/Editor/Experimental/EditorInteractiveToolsFramework/Public/EditorDragTools/DragToolsBehaviourSource.h
--- a/Engine/Source/Editor/Experimental/EditorInteractiveToolsFramework/Public/EditorDragTools/DragToolsBehaviourSource.h
+++ b/Engine/Source/Editor/Experimental/EditorInteractiveToolsFramework/Public/EditorDragTools/DragToolsBehaviourSource.h
@@ -11,14 +11,8 @@
 
 namespace UE::Editor::DragTools
 {
-
-// In case New TRS Gizmos are enabled, we need Mouse Drag + Key modifier Drag Tools Behaviors to be processed before
-// the Gizmo ones, and also before viewport interactions - otherwise Drag Behaviors will never be triggered. See:
-// * UTransformGizmo::SetupBehaviors()			[DEFAULT_GIZMO_PRIORITY]
-// * UTransformGizmo::SetupIndirectBehaviors()	[DEFAULT_GIZMO_PRIORITY - 1]
-
 static constexpr int HANDLED_BEFORE_VIEWPORT_INTERACTIONS =
-	UE::Editor::ViewportInteractions::VIEWPORT_INTERACTIONS_HIGH_PRIORITY - 1;
+	UE::Editor::ViewportInteractions::VIEWPORT_INTERACTIONS_DEFAULT_PRIORITY - 1;
 
 // Conversely, tools which are activated with just mouse buttons input (no key modifiers, mouse drag behavior only) need to be processed after Gizmo Behaviors
 static constexpr int HANDLED_AFTER_GIZMO_PRIORITY = FInputCapturePriority::DEFAULT_GIZMO_PRIORITY + 1;
diff --git a/Engine/Source/Editor/Experimental/EditorInteractiveToolsFramework/Public/EditorGizmos/TransformGizmo.h b/Engine/Source/Editor/Experimental/EditorInteractiveToolsFramework/Public/EditorGizmos/TransformGizmo.h
--- a/Engine/Source/Editor/Experimental/EditorInteractiveToolsFramework/Public/EditorGizmos/TransformGizmo.h
+++ b/Engine/Source/Editor/Experimental/EditorInteractiveToolsFramework/Public/EditorGizmos/TransformGizmo.h
@@ -56,6 +56,7 @@ class UMaterialInterface;
 class UObject;
 class UTransformProxy;
 class UGizmoElementGimbal;
+class UViewportClickDragBehavior;
 
 // @NOTE: The following changes should be made prior to 5.7 Release:
 // De-virtualize:
@@ -329,7 +330,7 @@ public:
 
 	/** The multi button mouse click behavior is accessible so that it can be modified to use different mouse keys. */
 	UPROPERTY()
-	TObjectPtr<UMultiButtonClickDragBehavior> MultiIndirectClickDragBehavior;
+	TObjectPtr<UViewportClickDragBehavior> MultiIndirectClickDragBehavior;
 
 	/** Transform Gizmo Source */
 	UPROPERTY()
diff --git a/Engine/Source/Editor/Experimental/EditorInteractiveToolsFramework/Public/ViewportInteractions/ViewportClickInteraction.h b/Engine/Source/Editor/Experimental/EditorInteractiveToolsFramework/Public/ViewportInteractions/ViewportClickInteraction.h
new file mode 100644
--- /dev/null
+++ b/Engine/Source/Editor/Experimental/EditorInteractiveToolsFramework/Public/ViewportInteractions/ViewportClickInteraction.h
@@ -0,0 +1,53 @@
+// Copyright Epic Games, Inc. All Rights Reserved.
+
+#pragma once
+
+#include "ViewportInteraction.h"
+#include "BaseBehaviors/BehaviorTargetInterfaces.h"
+#include "ViewportClickInteraction.generated.h"
+
+namespace UE::InteractiveTools
+{
+enum class EViewportClickButton : uint8
+{
+	Left,
+	Middle,
+	Right
+};
+} 
+
+class HHitProxy;
+class USingleClickInputBehavior;
+struct FViewportClick;
+
+/**
+ *
+ */
+UCLASS(Abstract)
+class EDITORINTERACTIVETOOLSFRAMEWORK_API UViewportClickInteraction
+	: public UViewportInteraction
+	, public IClickBehaviorTarget
+{
+	GENERATED_BODY()
+
+public:
+	UViewportClickInteraction();
+
+	// ~Begin IClickBehaviorTarget
+	virtual FInputRayHit IsHitByClick(const FInputDeviceRay& InClickPos) override;
+	virtual void OnClicked(const FInputDeviceRay& InClickPos) override;
+	// ~End IClickBehaviorTarget
+
+	FKey MouseButton = EKeys::LeftMouseButton;
+
+protected:
+	/** Override this function to provide custom functionality */
+	virtual void ProcessClick_Internal(const FViewportClick& InViewportClick, HHitProxy* InHitProxy, FSceneView& View) {}
+
+	virtual FDeviceButtonState GetActiveMouseButtonState(const FInputDeviceState& Input);
+
+	FKey GetMouseKey() const;
+	virtual EInputEvent GetMouseEvent() const;
+
+	TWeakObjectPtr<USingleClickInputBehavior> SingleClickInputBehaviorWeak;
+};
diff --git a/Engine/Source/Editor/Experimental/EditorInteractiveToolsFramework/Public/ViewportInteractions/ViewportDragInteraction.h b/Engine/Source/Editor/Experimental/EditorInteractiveToolsFramework/Public/ViewportInteractions/ViewportDragInteraction.h
--- a/Engine/Source/Editor/Experimental/EditorInteractiveToolsFramework/Public/ViewportInteractions/ViewportDragInteraction.h
+++ b/Engine/Source/Editor/Experimental/EditorInteractiveToolsFramework/Public/ViewportInteractions/ViewportDragInteraction.h
@@ -11,6 +11,8 @@
 
 class UClickDragInputBehavior;
 class USingleClickOrDragInputBehavior;
+class USingleDoubleClickOrDragBehavior;
+class UViewportClickDragBehavior;
 
 /**
  * A base class which can be used to implement interactions based on mouse drag actions
@@ -49,22 +51,15 @@ public:
 		return IsEnabled();
 	}
 
-	/**
-	 * Provides a click behavior target to be called when drag interaction fails
-	 * @param InClickBehaviorTarget the Click behavior target responding to clicks
-	 */
-	UE_API void SetClickTarget(IClickBehaviorTarget* InClickBehaviorTarget);
-
 protected:
-
-	UE_API void SetClickOrDragPriority(int InPriority) const;
-	
 	/**
 	 * Returns the state of the mouse button currently used by this Interaction. Left mouse by default, override for other buttons.
 	 */
 	UE_API virtual FDeviceButtonState GetActiveMouseButtonState(const FInputDeviceState& Input);
 
-	TWeakObjectPtr<USingleClickOrDragInputBehavior> ClickOrDragInputBehavior;
+	UE_API void EnableButton(const FKey& InButton) const;
+
+	TWeakObjectPtr<UViewportClickDragBehavior> ViewportClickDragBehavior;
 
 	bool bIsDragging;
 	int32 MouseX;
diff --git a/Engine/Source/Editor/Experimental/EditorInteractiveToolsFramework/Public/ViewportInteractions/ViewportInteraction.h b/Engine/Source/Editor/Experimental/EditorInteractiveToolsFramework/Public/ViewportInteractions/ViewportInteraction.h
--- a/Engine/Source/Editor/Experimental/EditorInteractiveToolsFramework/Public/ViewportInteractions/ViewportInteraction.h
+++ b/Engine/Source/Editor/Experimental/EditorInteractiveToolsFramework/Public/ViewportInteractions/ViewportInteraction.h
@@ -27,10 +27,14 @@ constexpr int LeftMouseButtonMod = 4;
 constexpr int RightMouseButtonMod = 5;
 constexpr int MiddleMouseButtonMod = 6;
 
-static constexpr int VIEWPORT_INTERACTIONS_DEFAULT_PRIORITY = FInputCapturePriority::DEFAULT_GIZMO_PRIORITY - 4;
-static constexpr int VIEWPORT_INTERACTIONS_LOW_PRIORITY = FInputCapturePriority::DEFAULT_GIZMO_PRIORITY - 3;
-static constexpr int VIEWPORT_INTERACTIONS_HIGH_PRIORITY = FInputCapturePriority::DEFAULT_GIZMO_PRIORITY - 5;
-static constexpr int VIEWPORT_INTERACTIONS_AFTER_GIZMOS_PRIORITY = FInputCapturePriority::DEFAULT_GIZMO_PRIORITY + 1;
+static constexpr int VIEWPORT_INTERACTIONS_DEFAULT_PRIORITY = FInputCapturePriority::DEFAULT_GIZMO_PRIORITY - 1;
+static constexpr int DRAG_CONFIRMED_PRIORITY = FInputCapturePriority::DEFAULT_GIZMO_PRIORITY + 5;
+static constexpr int CLICK_PRIORITY = FInputCapturePriority::DEFAULT_GIZMO_PRIORITY + 6;
+static constexpr int DRAG_PRIORITY = FInputCapturePriority::DEFAULT_GIZMO_PRIORITY + 7;
+
+// Higher priority for behaviors needing multiple buttons, so that capture stealing can handle those
+static constexpr int MULTIBUTTON_DRAG_CONFIRMED_PRIORITY = DRAG_CONFIRMED_PRIORITY - 1;
+static constexpr int MULTIBUTTON_DRAG_PRIORITY = DRAG_PRIORITY - 1;
 
 /**
  * Possible types of Viewport Interactions.
@@ -45,6 +49,8 @@ static FName CameraRotate = TEXT("CameraRotate");
 static FName FOV = TEXT("FOV");
 static FName Commands = TEXT("Commands");
 static FName SingleClick = TEXT("SingleClick");
+static FName DoubleClick = TEXT("DoubleClick");
+static FName ContextMenu = TEXT("ContextMenu");
 
 } // namespace UE::Editor::ViewportInteractions
 
diff --git a/Engine/Source/Editor/UnrealEd/Private/EditorViewportClient.cpp b/Engine/Source/Editor/UnrealEd/Private/EditorViewportClient.cpp
--- a/Engine/Source/Editor/UnrealEd/Private/EditorViewportClient.cpp
+++ b/Engine/Source/Editor/UnrealEd/Private/EditorViewportClient.cpp
@@ -3497,14 +3497,19 @@ bool FEditorViewportClient::Internal_InputKey(const FInputKeyEventArgs& EventArg
 		bHasMouseMovedSinceClick = false;
 	}
 
+	const bool bIsMultiMouseButtonInteractionActive = IsViewportInteractionActive(ModeTools.Get(), UE::Editor::ViewportInteractions::Zoom) || IsViewportInteractionActive(ModeTools.Get(), UE::Editor::ViewportInteractions::CameraDrag);
 
-	// Start tracking if any mouse button is down and it was a tracking event (MouseButton/Ctrl/Shift/Alt):
-	if ( InputState.IsAnyMouseButtonDown()
-		&& (Event == IE_Pressed || Event == IE_Released)
-		&& (InputState.IsMouseButtonEvent() || InputState.IsCtrlButtonEvent() || InputState.IsAltButtonEvent() || InputState.IsShiftButtonEvent() ) )
+	// Make sure ITF viewport interactions dealing with multiple mouse buttons pressed at the same time DO NOT trigger StartTrackingDueToInput when a button is released
+	if (!bIsMultiMouseButtonInteractionActive)
 	{
-		StartTrackingDueToInput( InputState, *View );
-		return true;
+		// Start tracking if any mouse button is down and it was a tracking event (MouseButton/Ctrl/Shift/Alt):
+		if ( InputState.IsAnyMouseButtonDown()
+			&& (Event == IE_Pressed || Event == IE_Released)
+			&& (InputState.IsMouseButtonEvent() || InputState.IsCtrlButtonEvent() || InputState.IsAltButtonEvent() || InputState.IsShiftButtonEvent() ) )
+		{
+			StartTrackingDueToInput( InputState, *View );
+			return true;
+		}
 	}
 
 	const bool bIsITFSingleClickActive = IsViewportInteractionActive(ModeTools.Get(), UE::Editor::ViewportInteractions::SingleClick);
