[ITF Viewport Interactions] Rework ViewportClickDragBehavior to dynamically lower its priority when it doesn't want to be active.

#rnx

diff --git a/Engine/Source/Editor/Experimental/EditorInteractiveToolsFramework/Private/Behaviors/ViewportClickDragBehavior.cpp b/Engine/Source/Editor/Experimental/EditorInteractiveToolsFramework/Private/Behaviors/ViewportClickDragBehavior.cpp
--- a/Engine/Source/Editor/Experimental/EditorInteractiveToolsFramework/Private/Behaviors/ViewportClickDragBehavior.cpp
+++ b/Engine/Source/Editor/Experimental/EditorInteractiveToolsFramework/Private/Behaviors/ViewportClickDragBehavior.cpp
@@ -9,21 +9,49 @@ void UViewportClickDragBehavior::SetDragConfirmedPriority(FInputCapturePriority
 
 FInputCapturePriority UViewportClickDragBehavior::GetPriority()
 {
-	if (MouseTraveledDistance > ClickDistanceThreshold)
+	if (const float* TraveledDistance = MouseTraveledDistance.GetPtrOrNull())
 	{
-		return DragConfirmedPriority;
+		if (*TraveledDistance > ClickDistanceThreshold)
+		{
+			return DragConfirmedPriority;
+		}
+		return Super::GetPriority();
 	}
-
-	return Super::GetPriority();
+	return NoCapturePriority;
 }
 
 FInputCaptureRequest UViewportClickDragBehavior::WantsCapture(const FInputDeviceState& InInputState)
 {
 	bInClickDrag = false;	// should never be true here, but weird things can happen w/ focus
 
-	const bool bIsPressed = IsAnyButtonPressed(InInputState);
-
-	if (bIsPressed)
+	if (InInputState.IsFromDevice(EInputDevices::Mouse) && !HasButtonsDown(InInputState))
+	{
+		// We don't want this event, reset state
+		MouseTraveledDistance = TOptional<float>();
+		return FInputCaptureRequest::Ignore();
+	}
+	
+	const bool bAnyButtonsPressed = AreAnyButtonsPressed(InInputState);
+	const bool bModifiersValid = ModifierCheckFunc == nullptr || ModifierCheckFunc(InInputState);
+	
+	if (ModifierMode == EModifierMode::Continuous && !bModifiersValid)
+	{
+		return FInputCaptureRequest::Ignore();
+	}
+	
+	if (!MouseTraveledDistance.IsSet())
+	{
+		if (!bAnyButtonsPressed && TriggerMode == ETriggerMode::Press)
+		{
+			return FInputCaptureRequest::Ignore();
+		}
+		if (bAnyButtonsPressed && ModifierMode == EModifierMode::OnInitialPress && !bModifiersValid)
+		{
+			return FInputCaptureRequest::Ignore();
+		}
+	}
+	
+	if (!MouseTraveledDistance.IsSet() && bAnyButtonsPressed)
 	{
 		// Traveled distance is 0.0f if mouse was just pressed
 		MouseTraveledDistance = 0.0f;
@@ -36,23 +64,46 @@ FInputCaptureRequest UViewportClickDragBehavior::WantsCapture(const FInputDevice
 		MouseTraveledDistance = static_cast<float>(FVector2D::Distance(InInputState.Mouse.Position2D, MouseDownPosition));
 	}
 
-	bool bAnyButtonDown = IsAnyButtonDown(InInputState);
+	const FInputRayHit HitResult = CanBeginClickDragSequence(GetDeviceRay(InInputState));
+	if (!HitResult.bHit)
+	{
+		return FInputCaptureRequest::Ignore();
+	}
+	
+	return FInputCaptureRequest::Begin(this, EInputCaptureSide::Any, HitResult.HitDepth);
+}
 
-	// When stealing, we check button down state and traveled distance
-	bool bIsStealing = MouseTraveledDistance > ClickDistanceThreshold && bAnyButtonDown;
+void UViewportClickDragBehavior::ForceEndCapture(const FInputCaptureData& CaptureData)
+{
+	MouseTraveledDistance = TOptional<float>();
+}
 
-	if ((bAnyButtonDown || bIsStealing) && (ModifierCheckFunc == nullptr || ModifierCheckFunc(InInputState)) )
+bool UViewportClickDragBehavior::HasButtonsDown(const FInputDeviceState& InInput)
+{
+	if (InInput.IsFromDevice(EInputDevices::Mouse)) 
 	{
-		const FInputRayHit HitResult = CanBeginClickDragSequence(GetDeviceRay(InInputState));
-		if (HitResult.bHit)
+		ActiveDevice = EInputDevices::Mouse;
+
+		const FDeviceButtonState& Left = InInput.Mouse.Left;
+		const FDeviceButtonState& Middle = InInput.Mouse.Middle;
+		const FDeviceButtonState& Right = InInput.Mouse.Right;
+		
+		if (bRequireAllButtons)
 		{
-			return FInputCaptureRequest::Begin(this, EInputCaptureSide::Any, HitResult.HitDepth);
+			return (!HandlesLeftMouseButton() || Left.bDown)
+				&& (!HandlesMiddleMouseButton() || Middle.bDown)
+				&& (!HandlesRightMouseButton() || Right.bDown);
 		}
+
+		return (HandlesLeftMouseButton() && Left.bDown)
+				|| (HandlesMiddleMouseButton() && Middle.bDown)
+				|| (HandlesRightMouseButton() && Right.bDown);
 	}
-	return FInputCaptureRequest::Ignore();
+
+	return false;
 }
 
-bool UViewportClickDragBehavior::IsAnyButtonPressed(const FInputDeviceState& InInput)
+bool UViewportClickDragBehavior::AreAnyButtonsPressed(const FInputDeviceState& InInput)
 {
 	if (InInput.IsFromDevice(EInputDevices::Mouse)) 
 	{
@@ -62,11 +113,9 @@ bool UViewportClickDragBehavior::IsAnyButtonPressed(const FInputDeviceState& InI
 		const FDeviceButtonState& Middle = InInput.Mouse.Middle;
 		const FDeviceButtonState& Right = InInput.Mouse.Right;
 
-		const bool bLeftDown = HandlesLeftMouseButton() ? Left.bPressed : false;
-		const bool bMiddleDown = HandlesMiddleMouseButton() ? Middle.bPressed : false;
-		const bool bRightDown = HandlesRightMouseButton() ? Right.bPressed : false;
-
-		return bLeftDown || bMiddleDown || bRightDown;
+		return (HandlesLeftMouseButton() && Left.bPressed)
+			|| (HandlesMiddleMouseButton() && Middle.bPressed)
+			|| (HandlesRightMouseButton() && Right.bPressed);
 	}
 
 	return false;
diff --git a/Engine/Source/Editor/Experimental/EditorInteractiveToolsFramework/Private/ViewportInteractions/ViewportDragInteraction.cpp b/Engine/Source/Editor/Experimental/EditorInteractiveToolsFramework/Private/ViewportInteractions/ViewportDragInteraction.cpp
--- a/Engine/Source/Editor/Experimental/EditorInteractiveToolsFramework/Private/ViewportInteractions/ViewportDragInteraction.cpp
+++ b/Engine/Source/Editor/Experimental/EditorInteractiveToolsFramework/Private/ViewportInteractions/ViewportDragInteraction.cpp
@@ -14,6 +14,9 @@ UViewportDragInteraction::UViewportDragInteraction()
 
 	UViewportClickDragBehavior* ClickDragBehavior = NewObject<UViewportClickDragBehavior>();
 	ClickDragBehavior->Initialize();
+	ClickDragBehavior->SetTriggerMode(UViewportClickDragBehavior::ETriggerMode::Down);
+	ClickDragBehavior->SetModifierMode(UViewportClickDragBehavior::EModifierMode::Continuous);
+	ClickDragBehavior->SetRequireAllButtons(true);
 
 	ClickDragBehavior->CanBeginClickDragFunc = [this](const FInputDeviceRay& InPressPos)
 	{
diff --git a/Engine/Source/Editor/Experimental/EditorInteractiveToolsFramework/Private/ViewportInteractions/ViewportPanInteraction.cpp b/Engine/Source/Editor/Experimental/EditorInteractiveToolsFramework/Private/ViewportInteractions/ViewportPanInteraction.cpp
--- a/Engine/Source/Editor/Experimental/EditorInteractiveToolsFramework/Private/ViewportInteractions/ViewportPanInteraction.cpp
+++ b/Engine/Source/Editor/Experimental/EditorInteractiveToolsFramework/Private/ViewportInteractions/ViewportPanInteraction.cpp
@@ -22,6 +22,9 @@ UViewportPanInteraction::UViewportPanInteraction()
 	MultiButtonClickDragBehavior->Initialize();
 	MultiButtonClickDragBehavior->EnableButton(EKeys::LeftMouseButton);
 	MultiButtonClickDragBehavior->EnableButton(EKeys::RightMouseButton);
+	MultiButtonClickDragBehavior->SetRequireAllButtons(true);
+	MultiButtonClickDragBehavior->SetTriggerMode(UViewportClickDragBehavior::ETriggerMode::Down);
+	MultiButtonClickDragBehavior->SetModifierMode(UViewportClickDragBehavior::EModifierMode::Continuous);
 
 	MultiButtonClickDragBehavior->CanBeginClickDragFunc = [this](const FInputDeviceRay& InPressPos)
 	{
@@ -34,11 +37,6 @@ UViewportPanInteraction::UViewportPanInteraction()
 		return bSuccess ? FInputRayHit(TNumericLimits<float>::Max()) : FInputRayHit();
 	};
 
-	MultiButtonClickDragBehavior->ModifierCheckFunc = [this](const FInputDeviceState& InInputDeviceState)
-	{
-		return InInputDeviceState.Mouse.Left.bDown && InInputDeviceState.Mouse.Right.bDown;
-	};
-
 	MultiButtonClickDragBehavior->OnClickPressFunc = [this](const FInputDeviceRay& InPressPos)
 	{
 		OnClickPress(InPressPos);
diff --git a/Engine/Source/Editor/Experimental/EditorInteractiveToolsFramework/Public/Behaviors/ViewportClickDragBehavior.h b/Engine/Source/Editor/Experimental/EditorInteractiveToolsFramework/Public/Behaviors/ViewportClickDragBehavior.h
--- a/Engine/Source/Editor/Experimental/EditorInteractiveToolsFramework/Public/Behaviors/ViewportClickDragBehavior.h
+++ b/Engine/Source/Editor/Experimental/EditorInteractiveToolsFramework/Public/Behaviors/ViewportClickDragBehavior.h
@@ -15,6 +15,22 @@ class UViewportClickDragBehavior : public UMultiButtonClickDragBehavior
 
 public:
 
+	enum class ETriggerMode
+	{
+		/** Mouse buttons must have been clicked to start a drag */
+		Press,
+		/** Mouse buttons need only to be held down to start a drag */
+		Down
+	};
+
+	enum class EModifierMode
+	{
+		/** Modifiers are checked on the initial press */
+		OnInitialPress,
+		/** Modifiers are always checked */
+		Continuous
+	};
+
 	/**
 	 * Set the priority this Drag Behavior will return from GetPriority once the current interaction is confirmed to be a drag.
 	 * Should be higher than its DefaultPriority, and higher than the priority of any other behavior it wants to steal capture from
@@ -23,20 +39,38 @@ public:
 
 	EDITORINTERACTIVETOOLSFRAMEWORK_API virtual FInputCapturePriority GetPriority() override;
 
-	virtual FInputCaptureRequest WantsCapture(const FInputDeviceState& InputState) override;
+	EDITORINTERACTIVETOOLSFRAMEWORK_API virtual FInputCaptureRequest WantsCapture(const FInputDeviceState& InputState) override;
+	EDITORINTERACTIVETOOLSFRAMEWORK_API virtual void ForceEndCapture(const FInputCaptureData& CaptureData) override;
+
+	void SetRequireAllButtons(bool bInRequireAllButtons) { bRequireAllButtons = bInRequireAllButtons; }
+	void SetTriggerMode(ETriggerMode InTriggerMode) { TriggerMode = InTriggerMode; }
+	void SetModifierMode(EModifierMode InModifierMode) { ModifierMode = InModifierMode; }
 
 protected:
 
-	/** @return true if any captured button was pressed. */
-	bool IsAnyButtonPressed(const FInputDeviceState& InInput);
+	/** @return true if all target buttons are down. */
+	bool HasButtonsDown(const FInputDeviceState& InInput);
+	
+	/** @return true if any target button has been pressed. */
+    bool AreAnyButtonsPressed(const FInputDeviceState& InInput);
 
 	/** Distance after which a drag is considered such */
-	float ClickDistanceThreshold = 5.0f;
+	float ClickDistanceThreshold = 1.0f;
+	
+	bool bRequireAllButtons = false;
+	ETriggerMode TriggerMode = ETriggerMode::Press;
+	EModifierMode ModifierMode = EModifierMode::OnInitialPress;
 
 	/** Cached cursor position */
 	FVector2D MouseDownPosition;
 
-	float MouseTraveledDistance = 0.0f;
+	TOptional<float> MouseTraveledDistance;
+
+	/**
+	 * This is essentially a hack. When this behavior knows it doesn't want an input event, it falls back to this priority.
+	 * This works around behavioral awkwardness with how capture stealing is currently implemented.
+	 */
+	FInputCapturePriority NoCapturePriority = FInputCapturePriority(INT32_MAX);
 
 	/**
 	 * The priority this Drag Behavior will return from GetPriority once the current interaction is confirmed to be a drag.
