[ITF Viewport Interactions] Remove incorrect usages of FEditorViewportClient::ToggleOrbitCamera(). Fixes a visual bug with screen-space gizmo controls when switching into and out of orbit mode.

#rnx

diff --git a/Engine/Source/Editor/Experimental/EditorInteractiveToolsFramework/Private/ViewportInteractions/ViewportOrbitInteraction.cpp b/Engine/Source/Editor/Experimental/EditorInteractiveToolsFramework/Private/ViewportInteractions/ViewportOrbitInteraction.cpp
--- a/Engine/Source/Editor/Experimental/EditorInteractiveToolsFramework/Private/ViewportInteractions/ViewportOrbitInteraction.cpp
+++ b/Engine/Source/Editor/Experimental/EditorInteractiveToolsFramework/Private/ViewportInteractions/ViewportOrbitInteraction.cpp
@@ -32,11 +32,6 @@ void UViewportOrbitInteraction::OnClickPress(const FInputDeviceRay& InPressPos)
 {
 	UE::Editor::ViewportInteractions::LOG(TEXT("ORBIT 1"));
 
-	if (FEditorViewportClient* EditorViewportClient = GetEditorViewportClient())
-	{
-		EditorViewportClient->ToggleOrbitCamera(true);
-	}
-
 	if (UViewportInteractionsBehaviorSource* BehaviorSource = GetViewportInteractionsBehaviorSource())
 	{
 		BehaviorSource->SetIsMouseLooking(true);
@@ -49,11 +44,6 @@ void UViewportOrbitInteraction::OnClickRelease(const FInputDeviceRay& InReleaseP
 {
 	UE::Editor::ViewportInteractions::LOG(TEXT("ORBIT 0"));
 
-	if (FEditorViewportClient* EditorViewportClient = GetEditorViewportClient())
-	{
-		EditorViewportClient->ToggleOrbitCamera(false);
-	}
-
 	if (UViewportInteractionsBehaviorSource* BehaviorSource = GetViewportInteractionsBehaviorSource())
 	{
 		BehaviorSource->SetIsMouseLooking(false);
