[Gizmos] Revert behavior change of experimental gizmos.

#rnx

diff --git a/Engine/Source/Editor/Experimental/EditorInteractiveToolsFramework/Private/EditorGizmos/TransformGizmo.cpp b/Engine/Source/Editor/Experimental/EditorInteractiveToolsFramework/Private/EditorGizmos/TransformGizmo.cpp
--- a/Engine/Source/Editor/Experimental/EditorInteractiveToolsFramework/Private/EditorGizmos/TransformGizmo.cpp
+++ b/Engine/Source/Editor/Experimental/EditorInteractiveToolsFramework/Private/EditorGizmos/TransformGizmo.cpp
@@ -18,7 +18,7 @@
 #include "Materials/Material.h"
 #include "Misc/AxisDisplayInfo.h"
 #include "UnrealEngine.h"
-#include "Behaviors/ViewportClickDragBehavior.h"
+#include "Behaviors/MultiButtonClickDragBehavior.h"
 #include "Intersection/IntersectionUtil.h"
 #include "PrimitiveDrawingUtils.h"
 #include "SceneQueries/SceneSnappingManager.h"
@@ -222,11 +222,10 @@ void UTransformGizmo::SetupIndirectBehaviors()
 	static constexpr int INDIRECT_GIZMO_PRIORITY = FInputCapturePriority::DEFAULT_GIZMO_PRIORITY-1;
 	
 	// Add middle mouse input behavior for indirect manipulation
-	UViewportClickDragBehavior* MiddleClickDragBehavior = NewObject<UViewportClickDragBehavior>();
+	ULocalClickDragInputBehavior* MiddleClickDragBehavior = NewObject<ULocalClickDragInputBehavior>();
 	MiddleClickDragBehavior->Initialize();
-	MiddleClickDragBehavior->SetDefaultPriority(FInputCapturePriority::DEFAULT_GIZMO_PRIORITY + 5);
-	MiddleClickDragBehavior->SetDragConfirmedPriority(INDIRECT_GIZMO_PRIORITY);
-	MiddleClickDragBehavior->EnableButton(EKeys::MiddleMouseButton);
+	MiddleClickDragBehavior->SetDefaultPriority(INDIRECT_GIZMO_PRIORITY);
+	MiddleClickDragBehavior->SetUseMiddleMouseButton();
 	MiddleClickDragBehavior->CanBeginClickDragFunc = [this](const FInputDeviceRay&)
 	{
 		static const FInputRayHit InvalidRayHit;
@@ -265,11 +264,10 @@ void UTransformGizmo::SetupIndirectBehaviors()
 	AddInputBehavior(MiddleClickDragBehavior);
 
 	// Add left/right mouse input behavior for indirect manipulation
-	MultiIndirectClickDragBehavior = NewObject<UViewportClickDragBehavior>();
+	MultiIndirectClickDragBehavior = NewObject<UMultiButtonClickDragBehavior>();
 	MultiIndirectClickDragBehavior->Initialize();
+	MultiIndirectClickDragBehavior->SetDefaultPriority(INDIRECT_GIZMO_PRIORITY);
 	MultiIndirectClickDragBehavior->EnableButton(EKeys::LeftMouseButton);
-	MultiIndirectClickDragBehavior->SetDefaultPriority(FInputCapturePriority::DEFAULT_GIZMO_PRIORITY + 5);
-	MultiIndirectClickDragBehavior->SetDragConfirmedPriority(INDIRECT_GIZMO_PRIORITY);
 	if (bCtrlMiddleDoesY)
 	{
 		MultiIndirectClickDragBehavior->EnableButton(EKeys::MiddleMouseButton);
diff --git a/Engine/Source/Editor/Experimental/EditorInteractiveToolsFramework/Public/EditorGizmos/TransformGizmo.h b/Engine/Source/Editor/Experimental/EditorInteractiveToolsFramework/Public/EditorGizmos/TransformGizmo.h
--- a/Engine/Source/Editor/Experimental/EditorInteractiveToolsFramework/Public/EditorGizmos/TransformGizmo.h
+++ b/Engine/Source/Editor/Experimental/EditorInteractiveToolsFramework/Public/EditorGizmos/TransformGizmo.h
@@ -56,7 +56,6 @@ class UMaterialInterface;
 class UObject;
 class UTransformProxy;
 class UGizmoElementGimbal;
-class UViewportClickDragBehavior;
 
 // @NOTE: The following changes should be made prior to 5.7 Release:
 // De-virtualize:
@@ -330,7 +329,7 @@ public:
 
 	/** The multi button mouse click behavior is accessible so that it can be modified to use different mouse keys. */
 	UPROPERTY()
-	TObjectPtr<UViewportClickDragBehavior> MultiIndirectClickDragBehavior;
+	TObjectPtr<UMultiButtonClickDragBehavior> MultiIndirectClickDragBehavior;
 
 	/** Transform Gizmo Source */
 	UPROPERTY()
