[ITF Camera Movements] Alt + RMB + Drag doesn't work with ITF viewport interaction enabled

diff --git a/Engine/Source/Editor/Experimental/EditorInteractiveToolsFramework/Private/ViewportInteractions/ViewportOrbitInteraction.cpp b/Engine/Source/Editor/Experimental/EditorInteractiveToolsFramework/Private/ViewportInteractions/ViewportOrbitInteraction.cpp
--- a/Engine/Source/Editor/Experimental/EditorInteractiveToolsFramework/Private/ViewportInteractions/ViewportOrbitInteraction.cpp
+++ b/Engine/Source/Editor/Experimental/EditorInteractiveToolsFramework/Private/ViewportInteractions/ViewportOrbitInteraction.cpp
@@ -4,6 +4,7 @@
 #include "BaseBehaviors/ClickDragBehavior.h"
 #include "EditorViewportClient.h"
 #include "ViewportClientNavigationHelper.h"
+#include "Behaviors/ViewportClickDragBehavior.h"
 #include "ViewportInteractions/ViewportInteractionsBehaviorSource.h"
 
 UViewportOrbitInteraction::UViewportOrbitInteraction()
@@ -13,6 +14,13 @@ UViewportOrbitInteraction::UViewportOrbitInteraction()
 	InteractionName = TEXT("Orbit");
 
 	EnableButton(EKeys::LeftMouseButton);
+	EnableButton(EKeys::RightMouseButton);
+
+	// LMB and RMB work separately
+	if (const TStrongObjectPtr<UViewportClickDragBehavior> ClickDragBehavior = ViewportClickDragBehavior.Pin())
+	{
+		ClickDragBehavior->SetRequireAllButtons(false);
+	}
 }
 
 FInputRayHit UViewportOrbitInteraction::CanBeginClickDragSequence(const FInputDeviceRay& InPressPos)
