[ITF] Fix wrong logic with Viewport Click Drag Behavior

If a button is not mapped to e.g. a click interaction, and just through this ViewportClickDragBehavior, the target will receive OnClickPress even if no drag has happened yet.
Makes sense, since the press did indeed happen, but the target might need/want to know both about the Drag starting, and/or the Click happening.

Add a dedicated Target interface which allows targets to handle Click and Drag events separately, and updates existing gizmo and viewport interactions where required.

diff --git a/Engine/Source/Editor/Experimental/EditorInteractiveToolsFramework/Private/Behaviors/ViewportClickDragBehavior.cpp b/Engine/Source/Editor/Experimental/EditorInteractiveToolsFramework/Private/Behaviors/ViewportClickDragBehavior.cpp
--- a/Engine/Source/Editor/Experimental/EditorInteractiveToolsFramework/Private/Behaviors/ViewportClickDragBehavior.cpp
+++ b/Engine/Source/Editor/Experimental/EditorInteractiveToolsFramework/Private/Behaviors/ViewportClickDragBehavior.cpp
@@ -2,6 +2,12 @@
 
 #include "Behaviors/ViewportClickDragBehavior.h"
 
+void UViewportClickDragBehavior::Initialize(IViewportClickDragBehaviorTarget* InTarget)
+{
+	check(InTarget != nullptr);
+	Target = InTarget;
+}
+
 void UViewportClickDragBehavior::SetDragConfirmedPriority(FInputCapturePriority InDragConfirmedPriority)
 {
 	DragConfirmedPriority = InDragConfirmedPriority;
@@ -22,8 +28,6 @@ FInputCapturePriority UViewportClickDragBehavior::GetPriority()
 
 FInputCaptureRequest UViewportClickDragBehavior::WantsCapture(const FInputDeviceState& InInputState)
 {
-	bInClickDrag = false;	// should never be true here, but weird things can happen w/ focus
-
 	if (InInputState.IsFromDevice(EInputDevices::Mouse) && !HasButtonsDown(InInputState))
 	{
 		// We don't want this event, reset state
@@ -32,7 +36,7 @@ FInputCaptureRequest UViewportClickDragBehavior::WantsCapture(const FInputDevice
 	}
 	
 	const bool bAnyButtonsPressed = AreAnyButtonsPressed(InInputState);
-	const bool bModifiersValid = ModifierCheckFunc == nullptr || ModifierCheckFunc(InInputState);
+	const bool bModifiersValid = IsModifiersStateValid(InInputState);
 	
 	if (ModifierMode == EModifierMode::Continuous && !bModifiersValid)
 	{
@@ -53,6 +57,9 @@ FInputCaptureRequest UViewportClickDragBehavior::WantsCapture(const FInputDevice
 	
 	if (!MouseTraveledDistance.IsSet() && bAnyButtonsPressed)
 	{
+		// Make sure behavior knows we are not dragging yet
+		bDragging = false;
+
 		// Traveled distance is 0.0f if mouse was just pressed
 		MouseTraveledDistance = 0.0f;
 
@@ -64,7 +71,7 @@ FInputCaptureRequest UViewportClickDragBehavior::WantsCapture(const FInputDevice
 		MouseTraveledDistance = static_cast<float>(FVector2D::Distance(InInputState.Mouse.Position2D, MouseDownPosition));
 	}
 
-	const FInputRayHit HitResult = CanBeginClickDragSequence(GetDeviceRay(InInputState));
+	const FInputRayHit HitResult = Target->CanBeginClickDragSequence(GetDeviceRay(InInputState));
 	if (!HitResult.bHit)
 	{
 		return FInputCaptureRequest::Ignore();
@@ -73,31 +80,140 @@ FInputCaptureRequest UViewportClickDragBehavior::WantsCapture(const FInputDevice
 	return FInputCaptureRequest::Begin(this, EInputCaptureSide::Any, HitResult.HitDepth);
 }
 
+FInputCaptureUpdate UViewportClickDragBehavior::BeginCapture(const FInputDeviceState& InputState, EInputCaptureSide eSide)
+{
+	OnStateUpdatedInternal(InputState);
+	return FInputCaptureUpdate::Begin(this, EInputCaptureSide::Any);
+}
+
+FInputCaptureUpdate UViewportClickDragBehavior::UpdateCapture(const FInputDeviceState& InputState, const FInputCaptureData& CaptureData)
+{
+	MouseTraveledDistance = static_cast<float>(FVector2D::Distance(InputState.Mouse.Position2D, MouseDownPosition));
+
+	if (!bDragging)
+	{
+		if (const float* TraveledDistance = MouseTraveledDistance.GetPtrOrNull())
+		{
+			if (*TraveledDistance > ClickDistanceThreshold)
+			{
+				OnDragStartInternal(InputState);
+
+				// Do not start drag until ClickDistanceThreshold has been reached
+				bDragging = true;
+			}
+		}
+	}
+
+	if (bDragging)
+	{
+		OnClickDragInternal(InputState, CaptureData);
+	}
+
+	// Make sure modifiers are still in a valid state if mode is continuous
+	const bool bPassesModifiersCheck = ModifierMode == EModifierMode::Continuous ? IsModifiersStateValid(InputState) : true;
+	const bool bPassesMouseButtonsCheck = InputState.InputDevice == EInputDevices::Mouse ? HasButtonsDown(InputState) : true;
+
+	if (!bPassesMouseButtonsCheck || !bPassesModifiersCheck)
+	{
+		OnDragEndInternal(InputState, CaptureData);
+		MouseTraveledDistance = TOptional<float>();
+		bDragging = false;
+		return FInputCaptureUpdate::End();
+	}
+
+	if (DidAnyButtonChangeState(InputState))
+	{
+		OnStateUpdatedInternal(InputState);
+	}
+
+	return FInputCaptureUpdate::Continue();
+}
+
 void UViewportClickDragBehavior::ForceEndCapture(const FInputCaptureData& CaptureData)
 {
 	MouseTraveledDistance = TOptional<float>();
+
+	if (bDragging)
+	{
+		bDragging = false;
+		OnDragEndInternal(FInputDeviceState(), CaptureData);
+	}
+}
+
+void UViewportClickDragBehavior::EnableButton(const FKey& InButton)
+{
+	if (InButton.IsMouseButton() && !ActiveMouseButtons.Contains(InButton))
+	{
+		ActiveMouseButtons.Add(InButton);
+	}
+}
+
+void UViewportClickDragBehavior::DisableButton(const FKey& InButton)
+{
+	if (InButton.IsMouseButton() && ActiveMouseButtons.Contains(InButton))
+	{
+		ActiveMouseButtons.Remove(InButton);
+	}
+}
+
+bool UViewportClickDragBehavior::AreAnyButtonsReleased(const FInputDeviceState& InInput)
+{
+	if (InInput.IsFromDevice(EInputDevices::Mouse))
+	{
+		const FDeviceButtonState& Left = InInput.Mouse.Left;
+		const FDeviceButtonState& Middle = InInput.Mouse.Middle;
+		const FDeviceButtonState& Right = InInput.Mouse.Right;
+
+		const bool bLeftReleased = IsLeftButtonActive() && Left.bReleased;
+		const bool bMiddleReleased = IsMiddleButtonActive() && Middle.bReleased;
+		const bool bRightReleased = IsRightButtonActive() && Right.bReleased;
+
+		return bLeftReleased || bMiddleReleased || bRightReleased;
+	}
+
+	return false;
+}
+
+bool UViewportClickDragBehavior::DidAnyButtonChangeState(const FInputDeviceState& InInput) const
+{
+	if (InInput.IsFromDevice(EInputDevices::Mouse))
+	{
+		const FDeviceButtonState& Left = InInput.Mouse.Left;
+		const FDeviceButtonState& Middle = InInput.Mouse.Middle;
+		const FDeviceButtonState& Right = InInput.Mouse.Right;
+
+		const bool bLeftChanged = IsLeftButtonActive() && (Left.bPressed || Left.bReleased);
+		const bool bMiddleChanged = IsMiddleButtonActive() && (Middle.bPressed || Middle.bReleased);
+		const bool bRightChanged = IsRightButtonActive() && (Right.bPressed || Right.bReleased);
+
+		return bLeftChanged || bMiddleChanged || bRightChanged;
+	}
+	return false;
+}
+
+bool UViewportClickDragBehavior::IsModifiersStateValid(const FInputDeviceState& InInputState) const
+{
+	return ModifierCheckFunc == nullptr || ModifierCheckFunc(InInputState);
 }
 
 bool UViewportClickDragBehavior::HasButtonsDown(const FInputDeviceState& InInput)
 {
 	if (InInput.IsFromDevice(EInputDevices::Mouse)) 
 	{
-		ActiveDevice = EInputDevices::Mouse;
-
 		const FDeviceButtonState& Left = InInput.Mouse.Left;
 		const FDeviceButtonState& Middle = InInput.Mouse.Middle;
 		const FDeviceButtonState& Right = InInput.Mouse.Right;
 		
 		if (bRequireAllButtons)
 		{
-			return (!HandlesLeftMouseButton() || Left.bDown)
-				&& (!HandlesMiddleMouseButton() || Middle.bDown)
-				&& (!HandlesRightMouseButton() || Right.bDown);
+			return (!IsLeftButtonActive() || Left.bDown)
+				&& (!IsMiddleButtonActive() || Middle.bDown)
+				&& (!IsRightButtonActive() || Right.bDown);
 		}
 
-		return (HandlesLeftMouseButton() && Left.bDown)
-				|| (HandlesMiddleMouseButton() && Middle.bDown)
-				|| (HandlesRightMouseButton() && Right.bDown);
+		return (IsLeftButtonActive() && Left.bDown)
+				|| (IsMiddleButtonActive() && Middle.bDown)
+				|| (IsRightButtonActive() && Right.bDown);
 	}
 
 	return false;
@@ -107,16 +223,59 @@ bool UViewportClickDragBehavior::AreAnyButtonsPressed(const FInputDeviceState& I
 {
 	if (InInput.IsFromDevice(EInputDevices::Mouse)) 
 	{
-		ActiveDevice = EInputDevices::Mouse;
-
 		const FDeviceButtonState& Left = InInput.Mouse.Left;
 		const FDeviceButtonState& Middle = InInput.Mouse.Middle;
 		const FDeviceButtonState& Right = InInput.Mouse.Right;
 
-		return (HandlesLeftMouseButton() && Left.bPressed)
-			|| (HandlesMiddleMouseButton() && Middle.bPressed)
-			|| (HandlesRightMouseButton() && Right.bPressed);
+		return (IsLeftButtonActive() && Left.bPressed)
+			|| (IsMiddleButtonActive() && Middle.bPressed)
+			|| (IsRightButtonActive() && Right.bPressed);
 	}
 
 	return false;
 }
+
+void UViewportClickDragBehavior::OnDragStartInternal(const FInputDeviceState& InputState)
+{
+	Target->OnDragStart(GetDeviceRay(InputState));
+}
+
+void UViewportClickDragBehavior::OnClickDragInternal(const FInputDeviceState& InputState, const FInputCaptureData& Data)
+{
+	Target->OnDrag(GetDeviceRay(InputState));
+}
+
+void UViewportClickDragBehavior::OnDragEndInternal(const FInputDeviceState& InputState, const FInputCaptureData& Data)
+{
+	Target->OnDragEnd(GetDeviceRay(InputState));
+}
+
+FInputDeviceRay UViewportClickDragBehavior::GetDeviceRay(const FInputDeviceState& input)
+{
+	if (input.IsFromDevice(EInputDevices::Mouse))
+	{
+		return FInputDeviceRay(input.Mouse.WorldRay, input.Mouse.Position2D);
+	}
+
+	return FInputDeviceRay(FRay(FVector::ZeroVector, FVector(0, 0, 1), true));
+}
+
+void UViewportClickDragBehavior::OnStateUpdatedInternal(const FInputDeviceState& InInputDeviceState)
+{
+	Target->OnStateUpdated(InInputDeviceState);
+}
+
+bool UViewportClickDragBehavior::IsLeftButtonActive() const
+{
+	return ActiveMouseButtons.Contains(EKeys::LeftMouseButton);
+}
+
+bool UViewportClickDragBehavior::IsMiddleButtonActive() const
+{
+	return ActiveMouseButtons.Contains(EKeys::MiddleMouseButton);
+}
+
+bool UViewportClickDragBehavior::IsRightButtonActive() const
+{
+	return ActiveMouseButtons.Contains(EKeys::RightMouseButton);
+}
diff --git a/Engine/Source/Editor/Experimental/EditorInteractiveToolsFramework/Private/ViewportInteractions/ViewportDragInteraction.cpp b/Engine/Source/Editor/Experimental/EditorInteractiveToolsFramework/Private/ViewportInteractions/ViewportDragInteraction.cpp
--- a/Engine/Source/Editor/Experimental/EditorInteractiveToolsFramework/Private/ViewportInteractions/ViewportDragInteraction.cpp
+++ b/Engine/Source/Editor/Experimental/EditorInteractiveToolsFramework/Private/ViewportInteractions/ViewportDragInteraction.cpp
@@ -1,7 +1,6 @@
 // Copyright Epic Games, Inc. All Rights Reserved.
 
 #include "ViewportInteractions/ViewportDragInteraction.h"
-#include "BaseBehaviors/ClickDragBehavior.h"
 #include "BaseBehaviors/SingleClickOrDragBehavior.h"
 #include "Behaviors/ViewportClickDragBehavior.h"
 #include "ViewportInteractions/ViewportInteractionsBehaviorSource.h"
@@ -12,7 +11,7 @@ UViewportDragInteraction::UViewportDragInteraction()
 	MouseX = 0;
 	MouseY = 0;
 
-	UViewportClickDragBehavior* ClickDragBehavior = NewObject<UViewportClickDragBehavior>();
+	ULocalViewportClickDragBehavior* ClickDragBehavior = NewObject<ULocalViewportClickDragBehavior>();
 	ClickDragBehavior->Initialize();
 	ClickDragBehavior->SetTriggerMode(UViewportClickDragBehavior::ETriggerMode::Down);
 	ClickDragBehavior->SetModifierMode(UViewportClickDragBehavior::EModifierMode::Continuous);
@@ -23,19 +22,19 @@ UViewportDragInteraction::UViewportDragInteraction()
 		return CanBeginClickDragSequence(InPressPos);
 	};
 
-	ClickDragBehavior->OnClickPressFunc = [this](const FInputDeviceRay& InPressPos)
+	ClickDragBehavior->OnDragStartFunc = [this](const FInputDeviceRay& InPressPos)
 	{
-		OnClickPress(InPressPos);
+		OnDragStart(InPressPos);
 	};
 
-	ClickDragBehavior->OnClickReleaseFunc = [this](const FInputDeviceRay& InReleasePos)
+	ClickDragBehavior->OnDragEndFunc = [this](const FInputDeviceRay& InReleasePos)
 	{
-		OnClickRelease(InReleasePos);
+		OnDragEnd(InReleasePos);
 	};
 
 	ClickDragBehavior->OnClickDragFunc = [this](const FInputDeviceRay& InDragPos)
 	{
-		OnClickDrag(InDragPos);
+		OnDrag(InDragPos);
 	};
 
 	ClickDragBehavior->SetDefaultPriority(UE::Editor::ViewportInteractions::DRAG_PRIORITY);
@@ -54,7 +53,7 @@ FInputRayHit UViewportDragInteraction::CanBeginClickDragSequence(const FInputDev
 	return FInputRayHit();
 }
 
-void UViewportDragInteraction::OnClickPress(const FInputDeviceRay& InPressPos)
+void UViewportDragInteraction::OnDragStart(const FInputDeviceRay& InDragStartPos)
 {
 	if (!IsEnabled())
 	{
@@ -62,15 +61,15 @@ void UViewportDragInteraction::OnClickPress(const FInputDeviceRay& InPressPos)
 		return;
 	}
 
-	MouseX = static_cast<int32>(InPressPos.ScreenPosition.X);
-	MouseY = static_cast<int32>(InPressPos.ScreenPosition.Y);
+	MouseX = static_cast<int32>(InDragStartPos.ScreenPosition.X);
+	MouseY = static_cast<int32>(InDragStartPos.ScreenPosition.Y);
 
 	bIsDragging = true;
 
 	GetViewportInteractionsBehaviorSource()->SetMouseCursorOverride(EMouseCursor::None);
 }
 
-void UViewportDragInteraction::OnClickDrag(const FInputDeviceRay& InDragPos)
+void UViewportDragInteraction::OnDrag(const FInputDeviceRay& InDragPos)
 {
 	if (!IsEnabled())
 	{
@@ -89,17 +88,17 @@ void UViewportDragInteraction::OnClickDrag(const FInputDeviceRay& InDragPos)
 	OnDragDelta(static_cast<float>(MouseDeltaX), static_cast<float>(MouseDeltaY));
 }
 
-void UViewportDragInteraction::OnClickRelease(const FInputDeviceRay& InReleasePos)
+void UViewportDragInteraction::OnDragEnd(const FInputDeviceRay& InDragEndPos)
 {
-	MouseX = static_cast<int32>(InReleasePos.ScreenPosition.X);
-	MouseY = static_cast<int32>(InReleasePos.ScreenPosition.Y);
+	MouseX = static_cast<int32>(InDragEndPos.ScreenPosition.X);
+	MouseY = static_cast<int32>(InDragEndPos.ScreenPosition.Y);
 
 	bIsDragging = false;
 }
 
 void UViewportDragInteraction::OnForceEndCapture()
 {
-	OnClickRelease(FInputDeviceRay());
+	OnDragEnd(FInputDeviceRay());
 }
 
 void UViewportDragInteraction::Initialize(UViewportInteractionsBehaviorSource* const InViewportInteractionsBehaviorSource)
@@ -108,11 +107,6 @@ void UViewportDragInteraction::Initialize(UViewportInteractionsBehaviorSource* c
 
 	if (ViewportClickDragBehavior.IsValid())
 	{
-		ViewportClickDragBehavior->SetUseCustomMouseButton([this](const FInputDeviceState& InInputDeviceState)
-		{
-			return GetActiveMouseButtonState(InInputDeviceState);
-		});
-
 		ViewportClickDragBehavior->ModifierCheckFunc = [this](const FInputDeviceState& InInputDeviceState)
 		{
 			return CanBeActivated(InInputDeviceState);
diff --git a/Engine/Source/Editor/Experimental/EditorInteractiveToolsFramework/Private/ViewportInteractions/ViewportMoveYawInteraction.cpp b/Engine/Source/Editor/Experimental/EditorInteractiveToolsFramework/Private/ViewportInteractions/ViewportMoveYawInteraction.cpp
--- a/Engine/Source/Editor/Experimental/EditorInteractiveToolsFramework/Private/ViewportInteractions/ViewportMoveYawInteraction.cpp
+++ b/Engine/Source/Editor/Experimental/EditorInteractiveToolsFramework/Private/ViewportInteractions/ViewportMoveYawInteraction.cpp
@@ -28,11 +28,11 @@ FInputRayHit UViewportMoveYawInteraction::CanBeginClickDragSequence(const FInput
 	return bSuccess ? FInputRayHit(TNumericLimits<float>::Max()) : FInputRayHit();
 }
 
-void UViewportMoveYawInteraction::OnClickPress(const FInputDeviceRay& PressPos)
+void UViewportMoveYawInteraction::OnDragStart(const FInputDeviceRay& InDragStartPos)
 {
 	UE::Editor::ViewportInteractions::LOG(TEXT("MOVE YAW 1"));
 
-	UViewportDragInteraction::OnClickPress(PressPos);
+	UViewportDragInteraction::OnDragStart(InDragStartPos);
 
 	if (UViewportInteractionsBehaviorSource* BehaviorSource = GetViewportInteractionsBehaviorSource())
 	{
@@ -40,11 +40,11 @@ void UViewportMoveYawInteraction::OnClickPress(const FInputDeviceRay& PressPos)
 	}
 }
 
-void UViewportMoveYawInteraction::OnClickRelease(const FInputDeviceRay& ReleasePos)
+void UViewportMoveYawInteraction::OnDragEnd(const FInputDeviceRay& InDragEndPos)
 {
 	UE::Editor::ViewportInteractions::LOG(TEXT("MOVE YAW 0"));
 
-	UViewportDragInteraction::OnClickRelease(ReleasePos);
+	UViewportDragInteraction::OnDragEnd(InDragEndPos);
 
 	if (UViewportInteractionsBehaviorSource* BehaviorSource = GetViewportInteractionsBehaviorSource())
 	{
diff --git a/Engine/Source/Editor/Experimental/EditorInteractiveToolsFramework/Private/ViewportInteractions/ViewportOrbitInteraction.cpp b/Engine/Source/Editor/Experimental/EditorInteractiveToolsFramework/Private/ViewportInteractions/ViewportOrbitInteraction.cpp
--- a/Engine/Source/Editor/Experimental/EditorInteractiveToolsFramework/Private/ViewportInteractions/ViewportOrbitInteraction.cpp
+++ b/Engine/Source/Editor/Experimental/EditorInteractiveToolsFramework/Private/ViewportInteractions/ViewportOrbitInteraction.cpp
@@ -36,7 +36,7 @@ FInputRayHit UViewportOrbitInteraction::CanBeginClickDragSequence(const FInputDe
 	return FInputRayHit();
 }
 
-void UViewportOrbitInteraction::OnClickPress(const FInputDeviceRay& InPressPos)
+void UViewportOrbitInteraction::OnDragStart(const FInputDeviceRay& InDragStartPos)
 {
 	UE::Editor::ViewportInteractions::LOG(TEXT("ORBIT 1"));
 
@@ -45,10 +45,10 @@ void UViewportOrbitInteraction::OnClickPress(const FInputDeviceRay& InPressPos)
 		BehaviorSource->SetIsMouseLooking(true);
 	}
 
-	UViewportDragInteraction::OnClickPress(InPressPos);
+	UViewportDragInteraction::OnDragStart(InDragStartPos);
 }
 
-void UViewportOrbitInteraction::OnClickRelease(const FInputDeviceRay& InReleasePos)
+void UViewportOrbitInteraction::OnDragEnd(const FInputDeviceRay& InDragEndPos)
 {
 	UE::Editor::ViewportInteractions::LOG(TEXT("ORBIT 0"));
 
@@ -57,7 +57,7 @@ void UViewportOrbitInteraction::OnClickRelease(const FInputDeviceRay& InReleaseP
 		BehaviorSource->SetIsMouseLooking(false);
 	}
 
-	UViewportDragInteraction::OnClickRelease(InReleasePos);
+	UViewportDragInteraction::OnDragEnd(InDragEndPos);
 }
 
 void UViewportOrbitInteraction::OnDragDelta(float InMouseDeltaX, float InMouseDeltaY)
diff --git a/Engine/Source/Editor/Experimental/EditorInteractiveToolsFramework/Private/ViewportInteractions/ViewportOrthoPanInteraction.cpp b/Engine/Source/Editor/Experimental/EditorInteractiveToolsFramework/Private/ViewportInteractions/ViewportOrthoPanInteraction.cpp
--- a/Engine/Source/Editor/Experimental/EditorInteractiveToolsFramework/Private/ViewportInteractions/ViewportOrthoPanInteraction.cpp
+++ b/Engine/Source/Editor/Experimental/EditorInteractiveToolsFramework/Private/ViewportInteractions/ViewportOrthoPanInteraction.cpp
@@ -26,24 +26,24 @@ FInputRayHit UViewportOrthoPanInteraction::CanBeginClickDragSequence(const FInpu
 	return bSuccess ? FInputRayHit(TNumericLimits<float>::Max()) : FInputRayHit();
 }
 
-void UViewportOrthoPanInteraction::OnClickPress(const FInputDeviceRay& PressPos)
+void UViewportOrthoPanInteraction::OnDragStart(const FInputDeviceRay& InDragStartPos)
 {
 	UE::Editor::ViewportInteractions::LOG(TEXT("ORTHO PAN 1"));
 
-	UViewportDragInteraction::OnClickPress(PressPos);
+	UViewportDragInteraction::OnDragStart(InDragStartPos);
 
 	if (UViewportInteractionsBehaviorSource* BehaviorSource = GetViewportInteractionsBehaviorSource())
 	{
-		GetViewportInteractionsBehaviorSource()->SetMouseCursorOverride(EMouseCursor::Type::GrabHandClosed);
+		BehaviorSource->SetMouseCursorOverride(EMouseCursor::Type::GrabHandClosed);
 		BehaviorSource->SetIsMouseLooking(true);
 	}
 }
 
-void UViewportOrthoPanInteraction::OnClickRelease(const FInputDeviceRay& ReleasePos)
+void UViewportOrthoPanInteraction::OnDragEnd(const FInputDeviceRay& InDragEndPos)
 {
 	UE::Editor::ViewportInteractions::LOG(TEXT("ORTHO PAN 0"));
 
-	UViewportDragInteraction::OnClickRelease(ReleasePos);
+	UViewportDragInteraction::OnDragEnd(InDragEndPos);
 
 	if (UViewportInteractionsBehaviorSource* BehaviorSource = GetViewportInteractionsBehaviorSource())
 	{
diff --git a/Engine/Source/Editor/Experimental/EditorInteractiveToolsFramework/Private/ViewportInteractions/ViewportPanInteraction.cpp b/Engine/Source/Editor/Experimental/EditorInteractiveToolsFramework/Private/ViewportInteractions/ViewportPanInteraction.cpp
--- a/Engine/Source/Editor/Experimental/EditorInteractiveToolsFramework/Private/ViewportInteractions/ViewportPanInteraction.cpp
+++ b/Engine/Source/Editor/Experimental/EditorInteractiveToolsFramework/Private/ViewportInteractions/ViewportPanInteraction.cpp
@@ -2,7 +2,6 @@
 
 #include "ViewportInteractions/ViewportPanInteraction.h"
 #include "BaseBehaviors/ClickDragBehavior.h"
-#include "Behaviors/MultiButtonClickDragBehavior.h"
 #include "Behaviors/ViewportClickDragBehavior.h"
 #include "EditorViewportClient.h"
 #include "Settings/LevelEditorViewportSettings.h"
@@ -18,7 +17,7 @@ UViewportPanInteraction::UViewportPanInteraction()
 	EnableButton(EKeys::MiddleMouseButton);
 
 	// LMB + RMB + Drag to Pan
-	UViewportClickDragBehavior* MultiButtonClickDragBehavior = NewObject<UViewportClickDragBehavior>();
+	ULocalViewportClickDragBehavior* MultiButtonClickDragBehavior = NewObject<ULocalViewportClickDragBehavior>();
 	MultiButtonClickDragBehavior->Initialize();
 	MultiButtonClickDragBehavior->EnableButton(EKeys::LeftMouseButton);
 	MultiButtonClickDragBehavior->EnableButton(EKeys::RightMouseButton);
@@ -37,20 +36,20 @@ UViewportPanInteraction::UViewportPanInteraction()
 		return bSuccess ? FInputRayHit(TNumericLimits<float>::Max()) : FInputRayHit();
 	};
 
-	MultiButtonClickDragBehavior->OnClickPressFunc = [this](const FInputDeviceRay& InPressPos)
+	MultiButtonClickDragBehavior->OnDragStartFunc = [this](const FInputDeviceRay& InDragStartPos)
 	{
-		OnClickPress(InPressPos);
+		OnDragStart(InDragStartPos);
 		bMultiButtonPanInteraction = true;
 	};
 
-	MultiButtonClickDragBehavior->OnClickReleaseFunc = [this](const FInputDeviceRay& InReleasePos)
+	MultiButtonClickDragBehavior->OnDragEndFunc = [this](const FInputDeviceRay& InDragEndPos)
 	{
-		OnClickRelease(InReleasePos);
+		OnDragEnd(InDragEndPos);
 	};
 
 	MultiButtonClickDragBehavior->OnClickDragFunc = [this](const FInputDeviceRay& InDragPos)
 	{
-		OnClickDrag(InDragPos);
+		OnDrag(InDragPos);
 	};
 
 	MultiButtonClickDragBehavior->SetDefaultPriority(UE::Editor::ViewportInteractions::MULTIBUTTON_DRAG_PRIORITY);
@@ -72,11 +71,11 @@ FInputRayHit UViewportPanInteraction::CanBeginClickDragSequence(const FInputDevi
 	return bSuccess ? FInputRayHit(TNumericLimits<float>::Max()) : FInputRayHit();
 }
 
-void UViewportPanInteraction::OnClickPress(const FInputDeviceRay& PressPos)
+void UViewportPanInteraction::OnDragStart(const FInputDeviceRay& InDragStartPos)
 {
 	UE::Editor::ViewportInteractions::LOG(TEXT("PAN 1"));
 
-	UViewportDragInteraction::OnClickPress(PressPos);
+	UViewportDragInteraction::OnDragStart(InDragStartPos);
 
 	if (UViewportInteractionsBehaviorSource* BehaviorSource = GetViewportInteractionsBehaviorSource())
 	{
@@ -86,11 +85,11 @@ void UViewportPanInteraction::OnClickPress(const FInputDeviceRay& PressPos)
 	bMultiButtonPanInteraction = false;
 }
 
-void UViewportPanInteraction::OnClickRelease(const FInputDeviceRay& ReleasePos)
+void UViewportPanInteraction::OnDragEnd(const FInputDeviceRay& InDragEndPos)
 {
 	UE::Editor::ViewportInteractions::LOG(TEXT("PAN 0"));
 
-	UViewportDragInteraction::OnClickRelease(ReleasePos);
+	UViewportDragInteraction::OnDragEnd(InDragEndPos);
 
 	if (UViewportInteractionsBehaviorSource* BehaviorSource = GetViewportInteractionsBehaviorSource())
 	{
diff --git a/Engine/Source/Editor/Experimental/EditorInteractiveToolsFramework/Private/ViewportInteractions/ViewportViewAngleInteraction.cpp b/Engine/Source/Editor/Experimental/EditorInteractiveToolsFramework/Private/ViewportInteractions/ViewportViewAngleInteraction.cpp
--- a/Engine/Source/Editor/Experimental/EditorInteractiveToolsFramework/Private/ViewportInteractions/ViewportViewAngleInteraction.cpp
+++ b/Engine/Source/Editor/Experimental/EditorInteractiveToolsFramework/Private/ViewportInteractions/ViewportViewAngleInteraction.cpp
@@ -27,11 +27,11 @@ FInputRayHit UViewportViewAngleInteraction::CanBeginClickDragSequence(const FInp
 	return bSuccess ? FInputRayHit(TNumericLimits<float>::Max()) : FInputRayHit();
 }
 
-void UViewportViewAngleInteraction::OnClickPress(const FInputDeviceRay& PressPos)
+void UViewportViewAngleInteraction::OnDragStart(const FInputDeviceRay& InDragStartPos)
 {
 	UE::Editor::ViewportInteractions::LOG(TEXT("VIEW ANGLE 1"));
 
-	UViewportDragInteraction::OnClickPress(PressPos);
+	UViewportDragInteraction::OnDragStart(InDragStartPos);
 
 	if (UViewportInteractionsBehaviorSource* BehaviorSource = GetViewportInteractionsBehaviorSource())
 	{
@@ -39,11 +39,11 @@ void UViewportViewAngleInteraction::OnClickPress(const FInputDeviceRay& PressPos
 	}
 }
 
-void UViewportViewAngleInteraction::OnClickRelease(const FInputDeviceRay& ReleasePos)
+void UViewportViewAngleInteraction::OnDragEnd(const FInputDeviceRay& InDragEndPos)
 {
 	UE::Editor::ViewportInteractions::LOG(TEXT("VIEW ANGLE 0"));
 
-	UViewportDragInteraction::OnClickRelease(ReleasePos);
+	UViewportDragInteraction::OnDragEnd(InDragEndPos);
 
 	if (UViewportInteractionsBehaviorSource* BehaviorSource = GetViewportInteractionsBehaviorSource())
 	{
@@ -79,7 +79,3 @@ FDeviceButtonState UViewportViewAngleInteraction::GetActiveMouseButtonState(cons
 	return Input.Mouse.Right;
 }
 
-void UViewportViewAngleInteraction::OnUpdateModifierState(int ModifierID, bool bIsOn)
-{
-	Super::OnUpdateModifierState(ModifierID, bIsOn);
-}
diff --git a/Engine/Source/Editor/Experimental/EditorInteractiveToolsFramework/Private/ViewportInteractions/ViewportZoomInteraction.cpp b/Engine/Source/Editor/Experimental/EditorInteractiveToolsFramework/Private/ViewportInteractions/ViewportZoomInteraction.cpp
--- a/Engine/Source/Editor/Experimental/EditorInteractiveToolsFramework/Private/ViewportInteractions/ViewportZoomInteraction.cpp
+++ b/Engine/Source/Editor/Experimental/EditorInteractiveToolsFramework/Private/ViewportInteractions/ViewportZoomInteraction.cpp
@@ -3,7 +3,6 @@
 #include "ViewportInteractions/ViewportZoomInteraction.h"
 #include "BaseBehaviors/ClickDragBehavior.h"
 #include "BaseBehaviors/KeyInputBehavior.h"
-#include "Behaviors/MultiButtonClickDragBehavior.h"
 #include "Behaviors/ViewportClickDragBehavior.h"
 #include "EditorViewportClient.h"
 #include "Settings/LevelEditorViewportSettings.h"
@@ -30,7 +29,7 @@ UViewportZoomInteraction::UViewportZoomInteraction()
 	KeyInputBehavior->SetDefaultPriority(UE::Editor::ViewportInteractions::VIEWPORT_INTERACTIONS_DEFAULT_PRIORITY);
 
 	// LMB + RMB + Drag zoom.
-	UViewportClickDragBehavior* MultiButtonClickDragBehavior = NewObject<UViewportClickDragBehavior>();
+	ULocalViewportClickDragBehavior* MultiButtonClickDragBehavior = NewObject<ULocalViewportClickDragBehavior>();
 	MultiButtonClickDragBehavior->Initialize();
 	MultiButtonClickDragBehavior->EnableButton(EKeys::LeftMouseButton);
 	MultiButtonClickDragBehavior->EnableButton(EKeys::RightMouseButton);
@@ -51,14 +50,14 @@ UViewportZoomInteraction::UViewportZoomInteraction()
 		return InInputDeviceState.Mouse.Left.bDown && InInputDeviceState.Mouse.Right.bDown;
 	};
 
-	MultiButtonClickDragBehavior->OnClickPressFunc = [this](const FInputDeviceRay& InPressPos)
+	MultiButtonClickDragBehavior->OnDragStartFunc = [this](const FInputDeviceRay& InDragStartPos)
 	{
-		OnClickPress(InPressPos);
+		OnDragStart(InDragStartPos);
 	};
 
-	MultiButtonClickDragBehavior->OnClickReleaseFunc = [this](const FInputDeviceRay& InReleasePos)
+	MultiButtonClickDragBehavior->OnDragEndFunc = [this](const FInputDeviceRay& InDragEndPos)
 	{
-		OnClickRelease(InReleasePos);
+		OnDragEnd(InDragEndPos);
 	};
 
 	MultiButtonClickDragBehavior->OnClickDragFunc = [this](const FInputDeviceRay& InDragPos)
@@ -73,12 +72,12 @@ UViewportZoomInteraction::UViewportZoomInteraction()
 
 FInputRayHit UViewportZoomInteraction::ShouldRespondToMouseWheel(const FInputDeviceRay& InCurrentPos)
 {
-	return IsEnabled() && !IsMouseLooking() ? FInputRayHit(TNumericLimits<float>::Max()) : FInputRayHit();
+	return IsEnabled() && !IsAnyMouseButtonDown() ? FInputRayHit(TNumericLimits<float>::Max()) : FInputRayHit();
 }
 
 void UViewportZoomInteraction::OnMouseWheelScrollUp(const FInputDeviceRay& InCurrentPos)
 {
-	if (!IsMouseLooking())
+	if (!IsAnyMouseButtonDown())
 	{
 		constexpr bool bZoomIn = false;
 		MouseWheelZoom(bZoomIn);
@@ -87,7 +86,7 @@ void UViewportZoomInteraction::OnMouseWheelScrollUp(const FInputDeviceRay& InCur
 
 void UViewportZoomInteraction::OnMouseWheelScrollDown(const FInputDeviceRay& InCurrentPos)
 {
-	if (!IsMouseLooking())
+	if (!IsAnyMouseButtonDown())
 	{
 		constexpr bool bZoomIn = true;
 		MouseWheelZoom(bZoomIn);
@@ -117,7 +116,7 @@ void UViewportZoomInteraction::OnKeyReleased(const FKey& InKeyID)
 void UViewportZoomInteraction::OnForceEndCapture()
 {
 	ZoomDirection = 0;
-	OnClickRelease(FInputDeviceRay());
+	OnDragEnd(FInputDeviceRay());
 }
 
 void UViewportZoomInteraction::Tick(float InDeltaTime) const
@@ -292,7 +291,7 @@ void UViewportZoomInteraction::MouseWheelZoom(bool bInZoomIn) const
 	}
 }
 
-void UViewportZoomInteraction::OnClickPress(const FInputDeviceRay& InPressPos)
+void UViewportZoomInteraction::OnDragStart(const FInputDeviceRay& InDragStartPos)
 {
 	UE::Editor::ViewportInteractions::LOG(TEXT("ORTHO ZOOM LMB + RMB 1"));
 
@@ -302,10 +301,10 @@ void UViewportZoomInteraction::OnClickPress(const FInputDeviceRay& InPressPos)
 		BehaviorSource->SetIsMouseLooking(true);
 	}
 
-	LastZoomDragY = static_cast<float>(InPressPos.ScreenPosition.Y);
+	LastZoomDragY = static_cast<float>(InDragStartPos.ScreenPosition.Y);
 }
 
-void UViewportZoomInteraction::OnClickRelease(const FInputDeviceRay& InReleasePos)
+void UViewportZoomInteraction::OnDragEnd(const FInputDeviceRay& InDragEndPos)
 {
 	UE::Editor::ViewportInteractions::LOG(TEXT("ORTHO ZOOM LMB + RMB 0"));
 
diff --git a/Engine/Source/Editor/Experimental/EditorInteractiveToolsFramework/Public/Behaviors/ViewportClickDragBehavior.h b/Engine/Source/Editor/Experimental/EditorInteractiveToolsFramework/Public/Behaviors/ViewportClickDragBehavior.h
--- a/Engine/Source/Editor/Experimental/EditorInteractiveToolsFramework/Public/Behaviors/ViewportClickDragBehavior.h
+++ b/Engine/Source/Editor/Experimental/EditorInteractiveToolsFramework/Public/Behaviors/ViewportClickDragBehavior.h
@@ -5,14 +5,67 @@
 #include "MultiButtonClickDragBehavior.h"
 #include "ViewportClickDragBehavior.generated.h"
 
+class IViewportClickDragBehaviorTarget
+{
+public:
+	virtual ~IViewportClickDragBehaviorTarget()
+	{
+	}
+
+	/**
+	 * Test if target can begin click-drag interaction at this point
+	 * @param PressPos device position/ray at click point
+	 * @return hit information at this point
+	 */
+	virtual FInputRayHit CanBeginClickDragSequence(const FInputDeviceRay& PressPos) = 0;
+
+	/**
+	 * Notify Target that drag has started
+	 * @param InDragStartPos device position/ray at drag start point
+	 */
+	virtual void OnDragStart(const FInputDeviceRay& InDragStartPos) = 0;
+
+	/**
+	 * Notify Target that input position has changed
+	 * @param DragPos device position/ray at click point
+	 */
+	virtual void OnDrag(const FInputDeviceRay& DragPos) = 0;
+
+	/**
+	 * Notify Target that drag has ended
+	 * @param InDragEndPos device position/ray at drag end point
+	 */
+	virtual void OnDragEnd(const FInputDeviceRay& InDragEndPos) = 0;
+
+	/**
+	 * Notify Target button state has changed
+	 * @param InInputDeviceState current device position/ray
+	 */
+	virtual void OnStateUpdated(const FInputDeviceState& InInputDeviceState) = 0;
+
+	/**
+	 * Notify Target that click-drag sequence has been explicitly terminated (eg by escape key)
+	 */
+	virtual void OnTerminateDragSequence() = 0;
+
+	/**
+	 * Notify Target that capture is ending
+	 */
+	virtual void OnForceEndCapture()
+	{
+	}
+};
+
 /**
- * 
+ * A behavior dedicated to Viewport Drag Interactions, directly supporting threshold-based capture.
+ * It can be used when capture stealing e.g. a single click behavior. The
+ * IViewportClickDragBehaviorTarget of this behavior will receive OnDragStart and OnDragEnd when drag
+ * starts/ends, and not on actual button press/release
  */
 UCLASS(MinimalAPI)
-class UViewportClickDragBehavior : public UMultiButtonClickDragBehavior
+class UViewportClickDragBehavior : public UInputBehavior
 {
 	GENERATED_BODY()
-
 public:
 
 	enum class ETriggerMode
@@ -31,6 +84,8 @@ public:
 		Continuous
 	};
 
+	EDITORINTERACTIVETOOLSFRAMEWORK_API void Initialize(IViewportClickDragBehaviorTarget* InTarget);
+
 	/**
 	 * Set the priority this Drag Behavior will return from GetPriority once the current interaction is confirmed to be a drag.
 	 * Should be higher than its DefaultPriority, and higher than the priority of any other behavior it wants to steal capture from
@@ -40,8 +95,22 @@ public:
 	EDITORINTERACTIVETOOLSFRAMEWORK_API virtual FInputCapturePriority GetPriority() override;
 
 	EDITORINTERACTIVETOOLSFRAMEWORK_API virtual FInputCaptureRequest WantsCapture(const FInputDeviceState& InputState) override;
+	EDITORINTERACTIVETOOLSFRAMEWORK_API virtual FInputCaptureUpdate BeginCapture(
+		const FInputDeviceState& InputState, EInputCaptureSide eSide
+	) override;
+	EDITORINTERACTIVETOOLSFRAMEWORK_API virtual FInputCaptureUpdate UpdateCapture(
+		const FInputDeviceState& InputState, const FInputCaptureData& CaptureData
+	) override;
 	EDITORINTERACTIVETOOLSFRAMEWORK_API virtual void ForceEndCapture(const FInputCaptureData& CaptureData) override;
 
+	EDITORINTERACTIVETOOLSFRAMEWORK_API void EnableButton(const FKey& InButton);
+	EDITORINTERACTIVETOOLSFRAMEWORK_API void DisableButton(const FKey& InButton);
+
+	/** @return current 3D world ray and optional 2D position for Target Device */
+	EDITORINTERACTIVETOOLSFRAMEWORK_API virtual FInputDeviceRay GetDeviceRay(const FInputDeviceState& input);
+
+	TFunction<bool(const FInputDeviceState&)> ModifierCheckFunc = nullptr;
+
 	void SetRequireAllButtons(bool bInRequireAllButtons) { bRequireAllButtons = bInRequireAllButtons; }
 	void SetTriggerMode(ETriggerMode InTriggerMode) { TriggerMode = InTriggerMode; }
 	void SetModifierMode(EModifierMode InModifierMode) { ModifierMode = InModifierMode; }
@@ -54,6 +123,28 @@ protected:
 	/** @return true if any target button has been pressed. */
     bool AreAnyButtonsPressed(const FInputDeviceState& InInput);
 
+	/** @return true if any target button has been released. */
+	bool AreAnyButtonsReleased(const FInputDeviceState& InInput);
+
+	/** @return true if any captured button changed its bPressed/bReleased state. */
+	bool DidAnyButtonChangeState(const FInputDeviceState& InInput) const;
+
+	bool IsModifiersStateValid(const FInputDeviceState& InInputState) const;
+
+	EDITORINTERACTIVETOOLSFRAMEWORK_API virtual void OnDragStartInternal(const FInputDeviceState& InputState);
+	EDITORINTERACTIVETOOLSFRAMEWORK_API virtual void OnClickDragInternal(
+		const FInputDeviceState& InputState, const FInputCaptureData& Data
+	);
+	EDITORINTERACTIVETOOLSFRAMEWORK_API virtual void OnDragEndInternal(
+		const FInputDeviceState& InputState, const FInputCaptureData& Data
+	);
+
+	virtual void OnStateUpdatedInternal(const FInputDeviceState& InInputDeviceState);
+
+	bool IsLeftButtonActive() const;
+	bool IsMiddleButtonActive() const;
+	bool IsRightButtonActive() const;
+
 	/** Distance after which a drag is considered such */
 	float ClickDistanceThreshold = 1.0f;
 	
@@ -77,4 +168,93 @@ protected:
 	 * Should be higher than its DefaultPriority, and higher than the priority of any other behavior it wants to steal capture from
 	 */
 	FInputCapturePriority DragConfirmedPriority;
+
+	IViewportClickDragBehaviorTarget* Target = nullptr;
+
+	TSet<FKey> ActiveMouseButtons;
+
+	/** True when capturing and dragging */
+	bool bDragging = false;
+};
+
+/**
+ * An implementation of UViewportClickDragBehavior that also implements IViewportClickDragBehaviorTarget directly, using
+ * a set of local lambda functions.
+ */
+UCLASS(MinimalAPI)
+class ULocalViewportClickDragBehavior
+	: public UViewportClickDragBehavior
+	, public IViewportClickDragBehaviorTarget
+{
+	GENERATED_BODY()
+protected:
+	using UViewportClickDragBehavior::Initialize;
+
+public:
+	/** Call this to initialize the class */
+	virtual void Initialize()
+	{
+		Initialize(this);
+	}
+
+	/** lambda implementation of CanBeginClickDragSequence */
+	TUniqueFunction<FInputRayHit(const FInputDeviceRay& PressPos)> CanBeginClickDragFunc = [](const FInputDeviceRay&)
+	{
+		return FInputRayHit();
+	};
+
+	/** lambda implementation of OnDragStart */
+	TUniqueFunction<void(const FInputDeviceRay& PressPos)> OnDragStartFunc = [](const FInputDeviceRay&)
+	{
+	};
+
+	/** lambda implementation of OnClickDrag */
+	TUniqueFunction<void(const FInputDeviceRay& PressPos)> OnClickDragFunc = [](const FInputDeviceRay&)
+	{
+	};
+
+	/** lambda implementation of OnDragEnd */
+	TUniqueFunction<void(const FInputDeviceRay& ReleasePos)> OnDragEndFunc = [](const FInputDeviceRay&)
+	{
+	};
+
+	/** lambda implementation of OnTerminateDragSequence */
+	TUniqueFunction<void()> OnTerminateFunc = []()
+	{
+	};
+
+	/** lambda implementation of OnStateUpdated */
+	TUniqueFunction<void(const FInputDeviceState& Input)> OnStateUpdatedFunc = [](const FInputDeviceState&)
+	{
+	};
+
+	virtual FInputRayHit CanBeginClickDragSequence(const FInputDeviceRay& PressPos) override
+	{
+		return CanBeginClickDragFunc(PressPos);
+	}
+
+	virtual void OnDragStart(const FInputDeviceRay& InDragStartPos) override
+	{
+		OnDragStartFunc(InDragStartPos);
+	}
+
+	virtual void OnDrag(const FInputDeviceRay& DragPos) override
+	{
+		OnClickDragFunc(DragPos);
+	}
+
+	virtual void OnDragEnd(const FInputDeviceRay& InDragEndPos) override
+	{
+		OnDragEndFunc(InDragEndPos);
+	}
+
+	virtual void OnStateUpdated(const FInputDeviceState& InInputDeviceState) override
+	{
+		OnStateUpdatedFunc(InInputDeviceState);
+	}
+
+	virtual void OnTerminateDragSequence() override
+	{
+		OnTerminateFunc();
+	}
 };
diff --git a/Engine/Source/Editor/Experimental/EditorInteractiveToolsFramework/Public/ViewportInteractions/ViewportDragInteraction.h b/Engine/Source/Editor/Experimental/EditorInteractiveToolsFramework/Public/ViewportInteractions/ViewportDragInteraction.h
--- a/Engine/Source/Editor/Experimental/EditorInteractiveToolsFramework/Public/ViewportInteractions/ViewportDragInteraction.h
+++ b/Engine/Source/Editor/Experimental/EditorInteractiveToolsFramework/Public/ViewportInteractions/ViewportDragInteraction.h
@@ -4,6 +4,7 @@
 
 #include "BaseBehaviors/BehaviorTargetInterfaces.h"
 #include "ViewportInteraction.h"
+#include "Behaviors/ViewportClickDragBehavior.h"
 
 #include "ViewportDragInteraction.generated.h"
 
@@ -20,23 +21,22 @@ class UViewportClickDragBehavior;
 UCLASS(MinimalAPI, Transient, Abstract)
 class UViewportDragInteraction
 	: public UViewportInteraction
-	, public IClickDragBehaviorTarget
+	, public IViewportClickDragBehaviorTarget
 {
+private:
 	GENERATED_BODY()
 public:
 	UE_API UViewportDragInteraction();
 
-	//~ Begin IClickDragBehaviorTarget
+	//~ Begin IViewportClickDragBehaviorTarget
 	UE_API virtual FInputRayHit CanBeginClickDragSequence(const FInputDeviceRay& InPressPos) override;
-	UE_API virtual void OnClickPress(const FInputDeviceRay& InPressPos) override;
-	UE_API virtual void OnClickDrag(const FInputDeviceRay& InDragPos) override;
-	UE_API virtual void OnClickRelease(const FInputDeviceRay& InReleasePos) override;
-	virtual void OnTerminateDragSequence() override {};
-	//~ End IClickDragBehaviorTarget
-
-	//~ Begin IModifierToggleBehaviorTarget
+	UE_API virtual void OnDragStart(const FInputDeviceRay& InDragStartPos) override;
+	UE_API virtual void OnDrag(const FInputDeviceRay& InDragPos) override;
+	UE_API virtual void OnDragEnd(const FInputDeviceRay& InDragEndPos) override;
+	UE_API virtual void OnTerminateDragSequence() override {}
+	UE_API virtual void OnStateUpdated(const FInputDeviceState& InInputDeviceState) override {}
 	UE_API virtual void OnForceEndCapture() override;
-	//~ End IModifierToggleBehaviorTarget
+	//~ End IViewportClickDragBehaviorTarget
 
 	//~ Begin FViewportInputToolBase
 	UE_API virtual void Initialize(UViewportInteractionsBehaviorSource* const InViewportInteractionsBehaviorSource) override;
diff --git a/Engine/Source/Editor/Experimental/EditorInteractiveToolsFramework/Public/ViewportInteractions/ViewportMoveYawInteraction.h b/Engine/Source/Editor/Experimental/EditorInteractiveToolsFramework/Public/ViewportInteractions/ViewportMoveYawInteraction.h
--- a/Engine/Source/Editor/Experimental/EditorInteractiveToolsFramework/Public/ViewportInteractions/ViewportMoveYawInteraction.h
+++ b/Engine/Source/Editor/Experimental/EditorInteractiveToolsFramework/Public/ViewportInteractions/ViewportMoveYawInteraction.h
@@ -22,11 +22,11 @@ class UViewportMoveYawInteraction
 public:
 	UE_API UViewportMoveYawInteraction();
 
-	//~ Begin IClickDragBehaviorTarget
+	//~ Begin IViewportClickDragBehaviorTarget
 	UE_API virtual FInputRayHit CanBeginClickDragSequence(const FInputDeviceRay& InPressPos) override;
-	UE_API virtual void OnClickPress(const FInputDeviceRay& PressPos) override;
-	UE_API virtual void OnClickRelease(const FInputDeviceRay& ReleasePos) override;
-	//~ End IClickDragBehaviorTarget
+	UE_API virtual void OnDragStart(const FInputDeviceRay& InDragStartPos) override;
+	UE_API virtual void OnDragEnd(const FInputDeviceRay& InDragEndPos) override;
+	//~ End IViewportClickDragBehaviorTarget
 
 	//~ Begin UViewportDragInteraction
 	UE_API virtual void OnDragDelta(float InMouseDeltaX, float InMouseDeltaY) override;
diff --git a/Engine/Source/Editor/Experimental/EditorInteractiveToolsFramework/Public/ViewportInteractions/ViewportOrbitInteraction.h b/Engine/Source/Editor/Experimental/EditorInteractiveToolsFramework/Public/ViewportInteractions/ViewportOrbitInteraction.h
--- a/Engine/Source/Editor/Experimental/EditorInteractiveToolsFramework/Public/ViewportInteractions/ViewportOrbitInteraction.h
+++ b/Engine/Source/Editor/Experimental/EditorInteractiveToolsFramework/Public/ViewportInteractions/ViewportOrbitInteraction.h
@@ -20,11 +20,11 @@ class UViewportOrbitInteraction
 public:
 	UE_API UViewportOrbitInteraction();
 
-	//~ Begin IClickDragBehaviorTarget
+	//~ Begin IViewportClickDragBehaviorTarget
 	UE_API virtual FInputRayHit CanBeginClickDragSequence(const FInputDeviceRay& InPressPos) override;
-	UE_API virtual void OnClickPress(const FInputDeviceRay& InPressPos) override;
-	UE_API virtual void OnClickRelease(const FInputDeviceRay& InReleasePos) override;
-	//~ End IClickDragBehaviorTarget
+	UE_API virtual void OnDragStart(const FInputDeviceRay& InDragStartPos) override;
+	UE_API virtual void OnDragEnd(const FInputDeviceRay& InDragEndPos) override;
+	//~ End IViewportClickDragBehaviorTarget
 
 	//~ Begin UViewportDragInteraction
 	UE_API virtual void OnDragDelta(float InMouseDeltaX, float InMouseDeltaY) override;
diff --git a/Engine/Source/Editor/Experimental/EditorInteractiveToolsFramework/Public/ViewportInteractions/ViewportOrthoPanInteraction.h b/Engine/Source/Editor/Experimental/EditorInteractiveToolsFramework/Public/ViewportInteractions/ViewportOrthoPanInteraction.h
--- a/Engine/Source/Editor/Experimental/EditorInteractiveToolsFramework/Public/ViewportInteractions/ViewportOrthoPanInteraction.h
+++ b/Engine/Source/Editor/Experimental/EditorInteractiveToolsFramework/Public/ViewportInteractions/ViewportOrthoPanInteraction.h
@@ -19,11 +19,11 @@ class UViewportOrthoPanInteraction : public UViewportDragInteraction
 public:
 	UE_API UViewportOrthoPanInteraction();
 
-	//~ Begin IClickDragBehaviorTarget
+	//~ Begin IViewportClickDragBehaviorTarget
 	UE_API virtual FInputRayHit CanBeginClickDragSequence(const FInputDeviceRay& InPressPos) override;
-	UE_API virtual void OnClickPress(const FInputDeviceRay& PressPos) override;
-	UE_API virtual void OnClickRelease(const FInputDeviceRay& ReleasePos) override;
-	//~ End IClickDragBehaviorTarget
+	UE_API virtual void OnDragStart(const FInputDeviceRay& InDragStartPos) override;
+	UE_API virtual void OnDragEnd(const FInputDeviceRay& InDragEndPos) override;
+	//~ End IViewportClickDragBehaviorTarget
 
 	//~ Begin UViewportDragInteraction
 	UE_API virtual void OnDragDelta(float InMouseDeltaX, float InMouseDeltaY) override;
diff --git a/Engine/Source/Editor/Experimental/EditorInteractiveToolsFramework/Public/ViewportInteractions/ViewportPanInteraction.h b/Engine/Source/Editor/Experimental/EditorInteractiveToolsFramework/Public/ViewportInteractions/ViewportPanInteraction.h
--- a/Engine/Source/Editor/Experimental/EditorInteractiveToolsFramework/Public/ViewportInteractions/ViewportPanInteraction.h
+++ b/Engine/Source/Editor/Experimental/EditorInteractiveToolsFramework/Public/ViewportInteractions/ViewportPanInteraction.h
@@ -21,11 +21,11 @@ class UViewportPanInteraction : public UViewportDragInteraction
 public:
 	UE_API UViewportPanInteraction();
 
-	//~ Begin IClickDragBehaviorTarget
+	//~ Begin IViewportClickDragBehaviorTarget
 	UE_API virtual FInputRayHit CanBeginClickDragSequence(const FInputDeviceRay& InPressPos) override;
-	UE_API virtual void OnClickPress(const FInputDeviceRay& PressPos) override;
-	UE_API virtual void OnClickRelease(const FInputDeviceRay& ReleasePos) override;
-	//~ End IClickDragBehaviorTarget
+	UE_API virtual void OnDragStart(const FInputDeviceRay& InDragStartPos) override;
+	UE_API virtual void OnDragEnd(const FInputDeviceRay& InDragEndPos) override;
+	//~ End IViewportClickDragBehaviorTarget
 
 	//~ Begin UViewportDragInteraction
 	UE_API virtual void OnDragDelta(float InMouseDeltaX, float InMouseDeltaY) override;
diff --git a/Engine/Source/Editor/Experimental/EditorInteractiveToolsFramework/Public/ViewportInteractions/ViewportViewAngleInteraction.h b/Engine/Source/Editor/Experimental/EditorInteractiveToolsFramework/Public/ViewportInteractions/ViewportViewAngleInteraction.h
--- a/Engine/Source/Editor/Experimental/EditorInteractiveToolsFramework/Public/ViewportInteractions/ViewportViewAngleInteraction.h
+++ b/Engine/Source/Editor/Experimental/EditorInteractiveToolsFramework/Public/ViewportInteractions/ViewportViewAngleInteraction.h
@@ -20,20 +20,17 @@ class UViewportViewAngleInteraction
 public:
 	UE_API UViewportViewAngleInteraction();
 
-	//~ Begin IClickDragBehaviorTarget
+	//~ Begin IViewportClickDragBehaviorTarget
 	UE_API virtual FInputRayHit CanBeginClickDragSequence(const FInputDeviceRay& InPressPos) override;
-	UE_API virtual void OnClickPress(const FInputDeviceRay& PressPos) override;
-	UE_API virtual void OnClickRelease(const FInputDeviceRay& ReleasePos) override;
-	//~ End IClickDragBehaviorTarget
+	UE_API virtual void OnDragStart(const FInputDeviceRay& InDragStartPos) override;
+	UE_API virtual void OnDragEnd(const FInputDeviceRay& InDragEndPos) override;
+	//~ End IViewportClickDragBehaviorTarget
 
 	//~ Begin UViewportDragInteraction
 	UE_API virtual void OnDragDelta(float InMouseDeltaX, float InMouseDeltaY) override;
 	UE_API virtual FDeviceButtonState GetActiveMouseButtonState(const FInputDeviceState& Input) override;
 	//~ End UViewportDragInteraction
 
-	//~ Begin IModifierToggleBehaviorTarget
-	UE_API virtual void OnUpdateModifierState(int ModifierID, bool bIsOn) override;
-	//~ End IModifierToggleBehaviorTarget
 };
 
 #undef UE_API
diff --git a/Engine/Source/Editor/Experimental/EditorInteractiveToolsFramework/Public/ViewportInteractions/ViewportZoomInteraction.h b/Engine/Source/Editor/Experimental/EditorInteractiveToolsFramework/Public/ViewportInteractions/ViewportZoomInteraction.h
--- a/Engine/Source/Editor/Experimental/EditorInteractiveToolsFramework/Public/ViewportInteractions/ViewportZoomInteraction.h
+++ b/Engine/Source/Editor/Experimental/EditorInteractiveToolsFramework/Public/ViewportInteractions/ViewportZoomInteraction.h
@@ -44,8 +44,8 @@ private:
 	UE_API void PerspectiveZoom(bool bInZoomIn) const;
 	UE_API void MouseWheelZoom(bool bInZoomIn) const;
 
-	UE_API void OnClickPress(const FInputDeviceRay& InPressPos);
-	UE_API void OnClickRelease(const FInputDeviceRay& InReleasePos);
+	UE_API void OnDragStart(const FInputDeviceRay& InDragStartPos);
+	UE_API void OnDragEnd(const FInputDeviceRay& InDragEndPos);
 	UE_API void OnClickDrag(const FInputDeviceRay& InDragPos);
 
 	int16 ZoomDirection = 0;
