[ITF Drag Tools] Fix for Shift+Drag interaction to drag camera

- Add computation of "CameraDelta" in Axis and Planar translation methods for new TRS Gizmo. This is used to compensate for potential camera movement which might happen during ApplyTranslateDelta

- Limit the MoveCameraWithObject Drag Tool to just handle shift modifier, and update the Viewport Client Navigation Helper accordingly. This might look better after the Drag Tools refactor which will have them follow the same structure as Viewport Interactions.

diff --git a/Engine/Source/Editor/Experimental/EditorInteractiveToolsFramework/Private/EditorDragTools/EditorMoveCameraWithObject.cpp b/Engine/Source/Editor/Experimental/EditorInteractiveToolsFramework/Private/EditorDragTools/EditorMoveCameraWithObject.cpp
--- a/Engine/Source/Editor/Experimental/EditorInteractiveToolsFramework/Private/EditorDragTools/EditorMoveCameraWithObject.cpp
+++ b/Engine/Source/Editor/Experimental/EditorInteractiveToolsFramework/Private/EditorDragTools/EditorMoveCameraWithObject.cpp
@@ -7,10 +7,13 @@
 #include "EditorGizmos/EditorTransformGizmo.h"
 #include "EditorGizmos/EditorTransformGizmoUtil.h"
 #include "EditorInteractiveGizmoManager.h"
+#include "ILevelEditor.h"
 #include "InteractiveGizmo.h"
 #include "InteractiveGizmoManager.h"
 #include "InteractiveToolManager.h"
 #include "LevelEditorViewport.h"
+#include "SLevelViewport.h"
+#include "ViewportClientNavigationHelper.h"
 #include "Tools/EdModeInteractiveToolsContext.h"
 
 FEditorMoveCameraWithObject::FEditorMoveCameraWithObject(UInteractiveToolManager* InToolManager)
@@ -59,9 +62,6 @@ void FEditorMoveCameraWithObject::Initialize()
 
 			OnEndPivotEditDelegate =
 				TransformProxy->OnEndTransformEdit.AddRaw(this, &FEditorMoveCameraWithObject::OnGizmoMovementEnd);
-
-			OnTransformChangedDelegate =
-				TransformProxy->OnTransformChanged.AddRaw(this, &FEditorMoveCameraWithObject::OnGizmoTransformChanged);
 		}
 	}
 }
@@ -75,7 +75,6 @@ void FEditorMoveCameraWithObject::Reset()
 		{
 			TransformProxy->OnBeginTransformEdit.Remove(OnBeginPivotEditDelegate);
 			TransformProxy->OnEndTransformEdit.Remove(OnEndPivotEditDelegate);
-			TransformProxy->OnTransformChanged.Remove(OnTransformChangedDelegate);
 		}
 
 		TransformGizmo = nullptr;
@@ -100,6 +99,22 @@ void FEditorMoveCameraWithObject::OnUsesNewTRSGizmosChanged(bool bInUseNewTRSGiz
 	}
 }
 
+FLevelEditorViewportClient* FEditorMoveCameraWithObject::GetLevelEditorViewportClient()
+{
+	if (GCurrentLevelEditingViewportClient)
+	{
+		if (TSharedPtr<ILevelEditor> LevelEditor = GCurrentLevelEditingViewportClient->ParentLevelEditor.Pin())
+		{
+			if (TSharedPtr<SLevelViewport> ViewportInterface = LevelEditor->GetActiveViewportInterface())
+			{
+				return &ViewportInterface->GetLevelViewportClient();
+			}
+		}
+	}
+
+	return nullptr;
+}
+
 FEditorMoveCameraWithObject::~FEditorMoveCameraWithObject()
 {
 	UEditorInteractiveGizmoManager::OnUsesNewTRSGizmosChangedDelegate().RemoveAll(this);
@@ -128,22 +143,38 @@ void FEditorMoveCameraWithObject::OnUpdateModifierState(int InModifierID, bool b
 	{
 		bModifierKeyIsPressed = bInIsOn;
 	}
+
+	if (FLevelEditorViewportClient* LevelEditorViewportClient = GetLevelEditorViewportClient())
+	{
+		if (FViewportClientNavigationHelper* NavigationHelper = LevelEditorViewportClient->GetViewportNavigationHelper())
+		{
+			NavigationHelper->bCameraFollowsGizmo = bGizmoIsDragged && bModifierKeyIsPressed;
+		}
+	}
 }
 
 void FEditorMoveCameraWithObject::OnGizmoMovementBegin(UTransformProxy* InTransformProxy)
 {
 	bGizmoIsDragged = true;
+
+	if (FLevelEditorViewportClient* LevelEditorViewportClient = GetLevelEditorViewportClient())
+	{
+		if (FViewportClientNavigationHelper* NavigationHelper = LevelEditorViewportClient->GetViewportNavigationHelper())
+		{
+			NavigationHelper->bCameraFollowsGizmo = bGizmoIsDragged && bModifierKeyIsPressed;
+		}
+	}
 }
 
 void FEditorMoveCameraWithObject::OnGizmoMovementEnd(UTransformProxy* InTransformProxy)
 {
 	bGizmoIsDragged = false;
-}
 
-void FEditorMoveCameraWithObject::OnGizmoTransformChanged(UTransformProxy* InTransformProxy, UE::Math::TTransform<double> InTransform)
-{
-	if (bGizmoIsDragged && bModifierKeyIsPressed && GCurrentLevelEditingViewportClient)
+	if (FLevelEditorViewportClient* LevelEditorViewportClient = GetLevelEditorViewportClient())
 	{
-		GCurrentLevelEditingViewportClient->MoveViewportCamera(InTransform.GetLocation(), FRotator::ZeroRotator);
+		if (FViewportClientNavigationHelper* NavigationHelper = LevelEditorViewportClient->GetViewportNavigationHelper())
+		{
+			NavigationHelper->bCameraFollowsGizmo = false;
+		}
 	}
 }
diff --git a/Engine/Source/Editor/Experimental/EditorInteractiveToolsFramework/Public/EditorDragTools/EditorMoveCameraWithObject.h b/Engine/Source/Editor/Experimental/EditorInteractiveToolsFramework/Public/EditorDragTools/EditorMoveCameraWithObject.h
--- a/Engine/Source/Editor/Experimental/EditorInteractiveToolsFramework/Public/EditorDragTools/EditorMoveCameraWithObject.h
+++ b/Engine/Source/Editor/Experimental/EditorInteractiveToolsFramework/Public/EditorDragTools/EditorMoveCameraWithObject.h
@@ -27,6 +27,8 @@ public:
 	UE_API virtual ~FEditorMoveCameraWithObject() override;
 
 protected:
+	static FLevelEditorViewportClient* GetLevelEditorViewportClient();
+
 	//~Begin IModifierToggleBehaviorTarget
 	UE_API virtual void OnUpdateModifierState(int InModifierID, bool bInIsOn) override;
 	//~End IModifierToggleBehaviorTarget
@@ -41,7 +43,6 @@ protected:
 
 	UE_API void OnGizmoMovementBegin(UTransformProxy* InTransformProxy);
 	UE_API void OnGizmoMovementEnd(UTransformProxy* InTransformProxy);
-	UE_API void OnGizmoTransformChanged(UTransformProxy* InTransformProxy, UE::Math::TTransform<double> InTransform);
 
 	UE_API void Initialize();
 	UE_API void Reset();
diff --git a/Engine/Source/Editor/UnrealEd/Private/LevelEditorViewport.cpp b/Engine/Source/Editor/UnrealEd/Private/LevelEditorViewport.cpp
--- a/Engine/Source/Editor/UnrealEd/Private/LevelEditorViewport.cpp
+++ b/Engine/Source/Editor/UnrealEd/Private/LevelEditorViewport.cpp
@@ -92,6 +92,7 @@
 #include "IHeadMountedDisplay.h"
 #include "IXRTrackingSystem.h"
 #include "ActorGroupingUtils.h"
+#include "ContextObjectStore.h"
 #include "EditorInteractiveGizmoManager.h"
 #include "EditorWorldExtension.h"
 #include "VREditorMode.h"
@@ -113,6 +114,7 @@
 #include "GenericPlatform/GenericPlatformInputDeviceMapper.h"
 #include "Rendering/StaticLightingSystemInterface.h"
 #include "TextureResource.h"
+#include "ViewportClientNavigationHelper.h"
 #include "EditorDragTools/EditorDragToolBehaviorTarget.h"
 #include "Subsystems/PlacementSubsystem.h"
 #include "Subsystems/EditorElementSubsystem.h"
@@ -3388,12 +3390,38 @@ bool FLevelEditorViewportClient::InputWidgetDelta(FViewport* InViewport, EAxisLi
 				ModeTools->PivotLocation += Drag;
 				ModeTools->SnappedLocation += Drag;
 
-				if (IsShiftPressed() && !UEditorInteractiveGizmoManager::UsesNewTRSGizmos())
+				const bool CameraFollowsGizmo = GetViewportNavigationHelper()->bCameraFollowsGizmo;
+				if ((IsShiftPressed() && !UEditorInteractiveGizmoManager::UsesNewTRSGizmos()) || CameraFollowsGizmo)
 				{
 					bApplyCameraSpeedScaleByDistance = false;
 					FVector CameraDelta( Drag );
 					MoveViewportCamera( CameraDelta, FRotator::ZeroRotator );
 					bApplyCameraSpeedScaleByDistance = true;
+
+					// New TRS Gizmos need un updated GizmoViewContext
+					if (CameraFollowsGizmo)
+					{
+						if (ModeTools)
+						{
+							if (const UModeManagerInteractiveToolsContext* const ModeManagerInteractiveContext =
+									ModeTools->GetInteractiveToolsContext())
+							{
+								if (const UContextObjectStore* const ContextObjectStore = ModeManagerInteractiveContext->ContextObjectStore)
+								{
+									if (UGizmoViewContext* const GizmoViewContext = ContextObjectStore->FindContext<UGizmoViewContext>())
+									{
+										FSceneViewFamilyContext ViewFamily(
+											FSceneViewFamily::ConstructionValues(Viewport, GetScene(), EngineShowFlags)
+												.SetRealtimeUpdate(IsRealtime())
+										);
+ 
+										FSceneView* View = CalcSceneView(&ViewFamily);
+										GizmoViewContext->ResetFromSceneView(*View);
+									}
+								}
+							}
+						}
+					}
 				}
 
 				ModeTools->UpdateInternalData();
diff --git a/Engine/Source/Editor/UnrealEd/Private/ViewportClientNavigationHelper.cpp b/Engine/Source/Editor/UnrealEd/Private/ViewportClientNavigationHelper.cpp
--- a/Engine/Source/Editor/UnrealEd/Private/ViewportClientNavigationHelper.cpp
+++ b/Engine/Source/Editor/UnrealEd/Private/ViewportClientNavigationHelper.cpp
@@ -4,6 +4,7 @@
 
 FViewportClientNavigationHelper::FViewportClientNavigationHelper()
 	: bIsMouseLooking(false)
+	, bCameraFollowsGizmo(false)
 {
 	ResetTransformDelta();
 	ResetImpulseData();
diff --git a/Engine/Source/Editor/UnrealEd/Public/ViewportClientNavigationHelper.h b/Engine/Source/Editor/UnrealEd/Public/ViewportClientNavigationHelper.h
--- a/Engine/Source/Editor/UnrealEd/Public/ViewportClientNavigationHelper.h
+++ b/Engine/Source/Editor/UnrealEd/Public/ViewportClientNavigationHelper.h
@@ -75,4 +75,9 @@ public:
 	 * This flag is true when the user is controlling a viewport camera using the mouse
 	 */
 	bool bIsMouseLooking;
+
+	/**
+	 * Can be used by Editor Viewport Client in case camera needs to move together with gizmo dragged selection (e.g. see FLevelEditorViewportClient::InputWidgetDelta)
+	 */
+	bool bCameraFollowsGizmo;
 };
