[Editor TRS] Add an option to control persistent gizmo axis selection

diff --git a/Engine/Source/Editor/Experimental/EditorInteractiveToolsFramework/Public/EditorGizmos/TransformGizmo.h b/Engine/Source/Editor/Experimental/EditorInteractiveToolsFramework/Public/EditorGizmos/TransformGizmo.h
--- a/Engine/Source/Editor/Experimental/EditorInteractiveToolsFramework/Public/EditorGizmos/TransformGizmo.h
+++ b/Engine/Source/Editor/Experimental/EditorInteractiveToolsFramework/Public/EditorGizmos/TransformGizmo.h
@@ -170,6 +170,10 @@ struct FGizmosParameters
 	UPROPERTY()
 	float AxisSizeMultiplier = 1.0f;
 
+	/** Whether to allow persistent handle selection. If disabled, implicit indirect modification will also be disabled. Explicit (ctrl+mouse) will still work.  */
+	UPROPERTY(EditAnywhere, Category = "Experimental")
+	bool bPersistHandleSelection = true;
+
 	UPROPERTY(EditAnywhere, Category = "Transform Gizmo|Development", meta=(ShowOnlyInnerProperties))
 	FGizmoDevParameters DevParameters;
 };
