[ITF Viewport Interactions] Fix an issue where middle mouse viewport pan flipped what "inverted" meant.

#rnx

diff --git a/Engine/Source/Editor/Experimental/EditorInteractiveToolsFramework/Private/ViewportInteractions/ViewportPanInteraction.cpp b/Engine/Source/Editor/Experimental/EditorInteractiveToolsFramework/Private/ViewportInteractions/ViewportPanInteraction.cpp
--- a/Engine/Source/Editor/Experimental/EditorInteractiveToolsFramework/Private/ViewportInteractions/ViewportPanInteraction.cpp
+++ b/Engine/Source/Editor/Experimental/EditorInteractiveToolsFramework/Private/ViewportInteractions/ViewportPanInteraction.cpp
@@ -83,7 +83,7 @@ void UViewportPanInteraction::PanCamera(FEditorViewportClient& Client, float Mou
 	}
 	else
 	{
-		LocationDelta = FQuat(ViewRotation) * FVector( 0.0f, MouseDeltaX, -MouseDeltaY) * (bInvert ? 1.0f : -1.0f);
+		LocationDelta = FQuat(ViewRotation) * FVector( 0.0f, MouseDeltaX, -MouseDeltaY) * (bInvert ? -1.0f : 1.0f);
 	}
 
 	FVector& ViewportLocationDelta = Client.GetViewportNavigationHelper()->LocationDelta;
