[ITF Alt + Drag Duplication] Source object should not move during duplication

Duplication is now happening on gizmo movement begin, before the object actually translates.

diff --git a/Engine/Source/Editor/Experimental/EditorInteractiveToolsFramework/Private/EditorDragTools/EditorDuplicateDragSelection.cpp b/Engine/Source/Editor/Experimental/EditorInteractiveToolsFramework/Private/EditorDragTools/EditorDuplicateDragSelection.cpp
--- a/Engine/Source/Editor/Experimental/EditorInteractiveToolsFramework/Private/EditorDragTools/EditorDuplicateDragSelection.cpp
+++ b/Engine/Source/Editor/Experimental/EditorInteractiveToolsFramework/Private/EditorDragTools/EditorDuplicateDragSelection.cpp
@@ -61,8 +61,6 @@ void FEditorDuplicateDragSelection::Initialize()
 				this, &FEditorDuplicateDragSelection::OnGizmoMovementBegin
 			);
 
-			TransformProxy->OnTransformChanged.AddRaw(this, &FEditorDuplicateDragSelection::OnGizmoTransformChanged);
-
 			TransformProxy->OnEndTransformEdit.AddRaw(
 				this, &FEditorDuplicateDragSelection::OnGizmoMovementEnd
 			);
@@ -153,29 +151,19 @@ void FEditorDuplicateDragSelection::OnGizmoMovementBegin(UTransformProxy* InTran
 	if (!bGizmoIsDragged)
 	{
 		bGizmoIsDragged = true;
+
+		if (bModifierKeyIsPressed)
+		{
+			// Duplicate selection and select the newly created objects
+			constexpr bool bSelectNewElements = true;
+			DuplicateSelection(bSelectNewElements);
+		}
 	}
 }
 
 void FEditorDuplicateDragSelection::OnGizmoMovementEnd(UTransformProxy* InTransformProxy)
 {
 	bGizmoIsDragged = false;
-	bGizmoTransformHasChanged = false;
-}
-
-void FEditorDuplicateDragSelection::OnGizmoTransformChanged(UTransformProxy* InTransformProxy, UE::Math::TTransform<double> InTransform)
-{
-	if (bGizmoTransformHasChanged || InTransform.GetTranslation().IsNearlyZero())
-	{
-		return;
-	}
-
-	bGizmoTransformHasChanged = true;
-	if (bModifierKeyIsPressed)
-	{
-		// Duplicate selection and select the newly created objects
-		constexpr bool bSelectNewElements = true;
-		DuplicateSelection(bSelectNewElements);
-	}
 }
 
 void FEditorDuplicateDragSelection::DuplicateSelection(bool bInSelectNewElements) const
diff --git a/Engine/Source/Editor/Experimental/EditorInteractiveToolsFramework/Public/EditorDragTools/EditorDuplicateDragSelection.h b/Engine/Source/Editor/Experimental/EditorInteractiveToolsFramework/Public/EditorDragTools/EditorDuplicateDragSelection.h
--- a/Engine/Source/Editor/Experimental/EditorInteractiveToolsFramework/Public/EditorDragTools/EditorDuplicateDragSelection.h
+++ b/Engine/Source/Editor/Experimental/EditorInteractiveToolsFramework/Public/EditorDragTools/EditorDuplicateDragSelection.h
@@ -40,7 +40,6 @@ protected:
 
 	void OnGizmoMovementBegin(UTransformProxy* InTransformProxy);
 	void OnGizmoMovementEnd(UTransformProxy* InTransformProxy);
-	void OnGizmoTransformChanged(UTransformProxy* InTransformProxy, UE::Math::TTransform<double> InTransform);
 
 	void Initialize();
 	void Reset();
@@ -60,5 +59,4 @@ private:
 
 	bool bGizmoIsDragged = false;
 	bool bModifierKeyIsPressed = false;
-	bool bGizmoTransformHasChanged = false;
 };
