[Editor TRS] Snapping with new TRS Gizmos does not work if the new Gizmos are not enabled at launch

#FailedValidations Preflight

diff --git a/Engine/Source/Editor/Experimental/EditorInteractiveToolsFramework/Private/EditorInteractiveGizmoManager.cpp b/Engine/Source/Editor/Experimental/EditorInteractiveToolsFramework/Private/EditorInteractiveGizmoManager.cpp
--- a/Engine/Source/Editor/Experimental/EditorInteractiveToolsFramework/Private/EditorInteractiveGizmoManager.cpp
+++ b/Engine/Source/Editor/Experimental/EditorInteractiveToolsFramework/Private/EditorInteractiveGizmoManager.cpp
@@ -106,15 +106,8 @@ void UEditorInteractiveGizmoManager::InitializeWithEditorModeManager(IToolsConte
 	Super::Initialize(QueriesAPIIn, TransactionsAPIIn, InputRouterIn);
 	EditorModeManager = InEditorModeManager;
 
-	UModeManagerInteractiveToolsContext* InteractiveToolsContext = EditorModeManager->GetInteractiveToolsContext();
-	check(InteractiveToolsContext);
-
+	// @note: most of the registration occurs in OnEditorModeChanged to ensure external callers to Deactivate/Activate mode don't invalidate the state 
 	EditorModeManager->OnEditorModeIDChanged().AddUObject(this, &UEditorInteractiveGizmoManager::OnEditorModeChanged);
-
-	if (InteractiveToolsContext)
-	{
-		UE::Editor::Gizmos::RegisterSceneSnappingManager(InteractiveToolsContext);
-	}
 }
 
 void UEditorInteractiveGizmoManager::Shutdown()
@@ -168,10 +161,16 @@ void UEditorInteractiveGizmoManager::OnEditorModeChanged(const FName& InModeID,
 	// When the system is shared between modes (ie. snapping), this can de-register the snapping AFTER being registered in init
 	// This callback allows us to de-register things before the other mode registers
 
-	// Only execute on exit (not enter)
-	if (!bInIsEnteringMode && EditorModeManager)
+	if (EditorModeManager)
 	{
-		if (const UModeManagerInteractiveToolsContext* InteractiveToolsContext = EditorModeManager->GetInteractiveToolsContext())
+		UModeManagerInteractiveToolsContext* InteractiveToolsContext = EditorModeManager->GetInteractiveToolsContext();
+		check(InteractiveToolsContext);
+
+		if (bInIsEnteringMode)
+		{
+			UE::Editor::Gizmos::RegisterSceneSnappingManager(InteractiveToolsContext);
+		}
+		else // Only execute on exit (not enter)
 		{
 			UE::Editor::Gizmos::DeregisterSceneSnappingManager(InteractiveToolsContext);
 		}
