[ITF Alt Drag duplicate] Fix wrong undo history + update transaction label

diff --git a/Engine/Source/Editor/Experimental/EditorInteractiveToolsFramework/Private/EditorDragTools/EditorDuplicateDragSelection.cpp b/Engine/Source/Editor/Experimental/EditorInteractiveToolsFramework/Private/EditorDragTools/EditorDuplicateDragSelection.cpp
--- a/Engine/Source/Editor/Experimental/EditorInteractiveToolsFramework/Private/EditorDragTools/EditorDuplicateDragSelection.cpp
+++ b/Engine/Source/Editor/Experimental/EditorInteractiveToolsFramework/Private/EditorDragTools/EditorDuplicateDragSelection.cpp
@@ -3,6 +3,7 @@
 #include "EditorDragTools/EditorDuplicateDragSelection.h"
 
 #include "ContextObjectStore.h"
+#include "EditorDragTools/EditorDragToolBehaviorTarget.h"
 #include "EditorGizmos/EditorTransformGizmo.h"
 #include "EditorGizmos/EditorTransformGizmoUtil.h"
 #include "EditorInteractiveGizmoManager.h"
@@ -13,9 +14,11 @@
 #include "InteractiveGizmoManager.h"
 #include "InteractiveToolManager.h"
 #include "LevelEditorViewport.h"
-#include "EditorDragTools/EditorDragToolBehaviorTarget.h"
+#include "ScopedTransaction.h"
 #include "Tools/AssetEditorContextInterface.h"
 
+#define LOCTEXT_NAMESPACE "EditorDuplicateDragSelection"
+
 FEditorDuplicateDragSelection::FEditorDuplicateDragSelection(UInteractiveToolManager* InToolManager)
 	: ToolManager(InToolManager)
 {
@@ -98,6 +101,7 @@ void FEditorDuplicateDragSelection::OnUsesNewTRSGizmosChanged(bool bInUseNewTRSG
 {
 	if (!bInUseNewTRSGizmos)
 	{
+		DuplicateTransaction.Reset();
 		Reset();
 	}
 }
@@ -105,6 +109,7 @@ void FEditorDuplicateDragSelection::OnUsesNewTRSGizmosChanged(bool bInUseNewTRSG
 FEditorDuplicateDragSelection::~FEditorDuplicateDragSelection()
 {
 	UEditorInteractiveGizmoManager::OnUsesNewTRSGizmosChangedDelegate().RemoveAll(this);
+	DuplicateTransaction.Reset();
 	Reset();
 }
 
@@ -144,6 +149,7 @@ void FEditorDuplicateDragSelection::OnForceEndCapture()
 {
 	bGizmoIsDragged = false;
 	bModifierKeyIsPressed = false;
+	DuplicateTransaction.Reset();
 }
 
 void FEditorDuplicateDragSelection::OnGizmoMovementBegin(UTransformProxy* InTransformProxy)
@@ -154,6 +160,8 @@ void FEditorDuplicateDragSelection::OnGizmoMovementBegin(UTransformProxy* InTran
 
 		if (bModifierKeyIsPressed)
 		{
+			DuplicateTransaction = MakeShared<FScopedTransaction>(LOCTEXT("DuplicateAndTransformTransactionName", "Duplicate and Transform"));
+
 			// Duplicate selection and select the newly created objects
 			constexpr bool bSelectNewElements = true;
 			DuplicateSelection(bSelectNewElements);
@@ -164,6 +172,7 @@ void FEditorDuplicateDragSelection::OnGizmoMovementBegin(UTransformProxy* InTran
 void FEditorDuplicateDragSelection::OnGizmoMovementEnd(UTransformProxy* InTransformProxy)
 {
 	bGizmoIsDragged = false;
+	DuplicateTransaction.Reset();
 }
 
 void FEditorDuplicateDragSelection::DuplicateSelection(bool bInSelectNewElements) const
@@ -212,3 +221,5 @@ void FEditorDuplicateDragSelection::DuplicateSelection(bool bInSelectNewElements
 		}
 	}
 }
+
+#undef LOCTEXT_NAMESPACE
diff --git a/Engine/Source/Editor/Experimental/EditorInteractiveToolsFramework/Public/EditorDragTools/EditorDuplicateDragSelection.h b/Engine/Source/Editor/Experimental/EditorInteractiveToolsFramework/Public/EditorDragTools/EditorDuplicateDragSelection.h
--- a/Engine/Source/Editor/Experimental/EditorInteractiveToolsFramework/Public/EditorDragTools/EditorDuplicateDragSelection.h
+++ b/Engine/Source/Editor/Experimental/EditorInteractiveToolsFramework/Public/EditorDragTools/EditorDuplicateDragSelection.h
@@ -7,6 +7,7 @@
 #include "UObject/GCObject.h"
 
 class FLevelEditorViewportClient;
+class FScopedTransaction;
 class UEditorTransformGizmo;
 class UInteractiveToolManager;
 class UTransformProxy;
@@ -59,4 +60,7 @@ private:
 
 	bool bGizmoIsDragged = false;
 	bool bModifierKeyIsPressed = false;
+
+	/** Keeps the duplicate + transform transaction alive between gizmo movement begin and end */
+	TSharedPtr<FScopedTransaction> DuplicateTransaction;
 };
