[ITF Drag Tools] Drag Tools now implement IViewportClickDragBehaviorTargets

Also updates gizmos so that drag tools work better together with them.

diff --git a/Engine/Source/Editor/Experimental/EditorInteractiveToolsFramework/Private/EditorDragTools/DragToolsBehaviourSource.cpp b/Engine/Source/Editor/Experimental/EditorInteractiveToolsFramework/Private/EditorDragTools/DragToolsBehaviourSource.cpp
--- a/Engine/Source/Editor/Experimental/EditorInteractiveToolsFramework/Private/EditorDragTools/DragToolsBehaviourSource.cpp
+++ b/Engine/Source/Editor/Experimental/EditorInteractiveToolsFramework/Private/EditorDragTools/DragToolsBehaviourSource.cpp
@@ -71,26 +71,14 @@ void UDragToolsBehaviorSource::Initialize(const UEditorInteractiveToolsContext*
 		DragTools.Add(EDragTool::BoxSelect, BoxSelectTool);
 
 		// Add drag input behavior
-		UClickDragInputBehavior* MouseBehavior = NewObject<UClickDragInputBehavior>();
+		UViewportClickDragBehavior* MouseBehavior = NewObject<UViewportClickDragBehavior>();
 		MouseBehavior->Initialize(BoxSelectTool.Get());
-		MouseBehavior->SetDefaultPriority(BOX_SELECT_PRIORITY);
+		MouseBehavior->SetDefaultPriority(UE::Editor::ViewportInteractions::DRAG_PRIORITY);
+		MouseBehavior->SetDragConfirmedPriority(HANDLED_AFTER_GIZMO_PRIORITY);
 		MouseBehavior->SetAllowsCaptureStealing(true);
+		MouseBehavior->EnableButton(EKeys::LeftMouseButton);
 		BehaviorSet->Add(MouseBehavior);
 
-		// Register modifier listeners
-		UKeyAsModifierInputBehavior* KeyModifierInputBehavior = NewObject<UKeyAsModifierInputBehavior>();
-		KeyModifierInputBehavior->Initialize(
-			BoxSelectTool.Get(), UE::Editor::DragTools::ShiftKeyMod, FInputDeviceState::IsShiftKeyDown
-		);
-		KeyModifierInputBehavior->Initialize(
-			BoxSelectTool.Get(), UE::Editor::DragTools::CtrlKeyMod, FInputDeviceState::IsCtrlKeyDown
-		);
-		KeyModifierInputBehavior->Initialize(
-			BoxSelectTool.Get(), UE::Editor::DragTools::AltKeyMod, FInputDeviceState::IsAltKeyDown
-		);
-		KeyModifierInputBehavior->SetDefaultPriority(BOX_SELECT_PRIORITY);
-		BehaviorSet->Add(KeyModifierInputBehavior);
-
 		BoxSelectTool->OnActivateTool().AddUObject(this, &UDragToolsBehaviorSource::OnActivateTool, EDragTool::BoxSelect);
 		BoxSelectTool->OnDeactivateTool().AddUObject(this, &UDragToolsBehaviorSource::OnDeactivateTool, EDragTool::BoxSelect);
 	}
@@ -103,26 +91,18 @@ void UDragToolsBehaviorSource::Initialize(const UEditorInteractiveToolsContext*
 		DragTools.Add(EDragTool::FrustumSelect, FrustumSelectTool);
 
 		// Add drag input behavior
-		UClickDragInputBehavior* MouseBehavior = NewObject<UClickDragInputBehavior>();
+		UViewportClickDragBehavior* MouseBehavior = NewObject<UViewportClickDragBehavior>();
 		MouseBehavior->Initialize(FrustumSelectTool.Get());
-		MouseBehavior->SetDefaultPriority(FRUSTUM_SELECT_PRIORITY);
-		MouseBehavior->SetUseLeftMouseButton();
+		MouseBehavior->SetDefaultPriority(UE::Editor::ViewportInteractions::DRAG_PRIORITY);
+		MouseBehavior->SetDragConfirmedPriority(UE::Editor::ViewportInteractions::VIEWPORT_INTERACTIONS_DEFAULT_PRIORITY);
+		MouseBehavior->EnableButton(EKeys::LeftMouseButton);
+		MouseBehavior->SetAllowsCaptureStealing(true);
+		MouseBehavior->ModifierCheckFunc = [](const FInputDeviceState& InputState)
+		{
+			return FInputDeviceState::IsCtrlKeyDown(InputState) && FInputDeviceState::IsAltKeyDown(InputState);
+		};
 		BehaviorSet->Add(MouseBehavior);
 
-		// Register modifier listeners
-		UKeyAsModifierInputBehavior* KeyModifierInputBehavior = NewObject<UKeyAsModifierInputBehavior>();
-		KeyModifierInputBehavior->Initialize(
-			FrustumSelectTool.Get(), UE::Editor::DragTools::ShiftKeyMod, FInputDeviceState::IsShiftKeyDown
-		);
-		KeyModifierInputBehavior->Initialize(
-			FrustumSelectTool.Get(), UE::Editor::DragTools::CtrlKeyMod, FInputDeviceState::IsCtrlKeyDown
-		);
-		KeyModifierInputBehavior->Initialize(
-			FrustumSelectTool.Get(), UE::Editor::DragTools::AltKeyMod, FInputDeviceState::IsAltKeyDown
-		);
-		KeyModifierInputBehavior->SetDefaultPriority(FRUSTUM_SELECT_PRIORITY);
-		BehaviorSet->Add(KeyModifierInputBehavior);
-
 		FrustumSelectTool->OnActivateTool().AddUObject(
 			this, &UDragToolsBehaviorSource::OnActivateTool, EDragTool::FrustumSelect
 		);
@@ -138,10 +118,11 @@ void UDragToolsBehaviorSource::Initialize(const UEditorInteractiveToolsContext*
 		DragTools.Add(EDragTool::Measure, MeasureTool);
 
 		// Add drag input behavior
-		UClickDragInputBehavior* MouseBehavior = NewObject<UClickDragInputBehavior>();
+		UViewportClickDragBehavior* MouseBehavior = NewObject<UViewportClickDragBehavior>();
 		MouseBehavior->Initialize(MeasureTool.Get());
-		MouseBehavior->SetDefaultPriority(MEASURE_PRIORITY);
-		MouseBehavior->SetUseMiddleMouseButton();
+		MouseBehavior->SetDefaultPriority(UE::Editor::ViewportInteractions::DRAG_PRIORITY);
+		MouseBehavior->SetDragConfirmedPriority(HANDLED_AFTER_GIZMO_PRIORITY);
+		MouseBehavior->EnableButton(EKeys::MiddleMouseButton);
 		BehaviorSet->Add(MouseBehavior);
 
 		MeasureTool->OnActivateTool().AddUObject(this, &UDragToolsBehaviorSource::OnActivateTool, EDragTool::Measure);
@@ -367,20 +348,17 @@ void UDragToolsBehaviorSource::ActivateViewportChangeTool()
 			DragTools.Emplace(EDragTool::ViewportChange, ViewportChangeTool);
 
 			// Add drag input behavior
-			UClickDragInputBehavior* MouseBehavior = NewObject<UClickDragInputBehavior>();
+			UViewportClickDragBehavior* MouseBehavior = NewObject<UViewportClickDragBehavior>();
 			MouseBehavior->Initialize(ViewportChangeTool.Get());
-			MouseBehavior->SetDefaultPriority(VIEWPORT_CHANGE_PRIORITY);
-			MouseBehavior->SetUseMiddleMouseButton();
+			MouseBehavior->SetDefaultPriority(UE::Editor::ViewportInteractions::DRAG_PRIORITY - 1);
+			MouseBehavior->SetDragConfirmedPriority(HANDLED_AFTER_GIZMO_PRIORITY);
+			MouseBehavior->EnableButton(EKeys::MiddleMouseButton);
+			MouseBehavior->ModifierCheckFunc = [](const FInputDeviceState& InputState)
+			{
+				return FInputDeviceState::IsCtrlKeyDown(InputState) && !FInputDeviceState::IsAltKeyDown(InputState) && !FInputDeviceState::IsShiftKeyDown(InputState);
+			};
 			BehaviorSet->Add(MouseBehavior, this, UE::Editor::DragTools::ViewportChangeBehaviorGroup);
 
-			// Register modifier listeners
-			UKeyAsModifierInputBehavior* KeyModifierInputBehavior = NewObject<UKeyAsModifierInputBehavior>();
-			KeyModifierInputBehavior->Initialize(
-				ViewportChangeTool.Get(), UE::Editor::DragTools::CtrlKeyMod, FInputDeviceState::IsCtrlKeyDown
-			);
-			KeyModifierInputBehavior->SetDefaultPriority(VIEWPORT_CHANGE_PRIORITY);
-			BehaviorSet->Add(KeyModifierInputBehavior, this, UE::Editor::DragTools::ViewportChangeBehaviorGroup);
-
 			ViewportChangeTool->OnActivateTool().AddUObject(
 				this, &UDragToolsBehaviorSource::OnActivateTool, EDragTool::ViewportChange
 			);
diff --git a/Engine/Source/Editor/Experimental/EditorInteractiveToolsFramework/Private/EditorDragTools/EditorBoxSelect.cpp b/Engine/Source/Editor/Experimental/EditorInteractiveToolsFramework/Private/EditorDragTools/EditorBoxSelect.cpp
--- a/Engine/Source/Editor/Experimental/EditorInteractiveToolsFramework/Private/EditorDragTools/EditorBoxSelect.cpp
+++ b/Engine/Source/Editor/Experimental/EditorInteractiveToolsFramework/Private/EditorDragTools/EditorBoxSelect.cpp
@@ -81,7 +81,7 @@ FInputRayHit FEditorBoxSelect::CanBeginClickDragSequence(const FInputDeviceRay&
 	return FInputRayHit();
 }
 
-void FEditorBoxSelect::OnClickPress(const FInputDeviceRay& InPressPos)
+void FEditorBoxSelect::OnDragStart(const FInputDeviceRay& InPressPos)
 {
 	if (!EditorViewportClientProxy)
 	{
@@ -134,7 +134,7 @@ void FEditorBoxSelect::OnClickPress(const FInputDeviceRay& InPressPos)
 	}
 }
 
-void FEditorBoxSelect::OnClickDrag(const FInputDeviceRay& InDragPos)
+void FEditorBoxSelect::OnDrag(const FInputDeviceRay& InDragPos)
 {
 	if (!EditorViewportClientProxy)
 	{
@@ -218,7 +218,7 @@ void FEditorBoxSelect::OnClickDrag(const FInputDeviceRay& InDragPos)
 	}
 }
 
-void FEditorBoxSelect::OnClickRelease(const FInputDeviceRay& InReleasePos)
+void FEditorBoxSelect::OnDragEnd(const FInputDeviceRay& InReleasePos)
 {
 	if (EditorViewportClientProxy)
 	{
@@ -268,6 +268,7 @@ void FEditorBoxSelect::OnClickRelease(const FInputDeviceRay& InReleasePos)
 			FBox SelBBox;
 			CalculateBox(SelBBox);
 
+			// TODO: with ITF Gizmos indirect interactions, this is probably not going to be triggered
 			const bool bControlDown = FInputDeviceState::IsCtrlKeyDown(InputState);
 			const bool bShiftDown = FInputDeviceState::IsShiftKeyDown(InputState);
 
@@ -380,7 +381,7 @@ void FEditorBoxSelect::OnClickRelease(const FInputDeviceRay& InReleasePos)
 		}
 	}
 
-	FEditorDragToolBehaviorTarget::OnClickRelease(InReleasePos);
+	FEditorDragToolBehaviorTarget::OnDragEnd(InReleasePos);
 }
 
 void FEditorBoxSelect::OnTerminateDragSequence()
diff --git a/Engine/Source/Editor/Experimental/EditorInteractiveToolsFramework/Private/EditorDragTools/EditorDragToolBehaviorTarget.cpp b/Engine/Source/Editor/Experimental/EditorInteractiveToolsFramework/Private/EditorDragTools/EditorDragToolBehaviorTarget.cpp
--- a/Engine/Source/Editor/Experimental/EditorInteractiveToolsFramework/Private/EditorDragTools/EditorDragToolBehaviorTarget.cpp
+++ b/Engine/Source/Editor/Experimental/EditorInteractiveToolsFramework/Private/EditorDragTools/EditorDragToolBehaviorTarget.cpp
@@ -69,24 +69,12 @@ FEditorDragToolBehaviorTarget::FEditorDragToolBehaviorTarget(FEditorViewportClie
 	}
 }
 
-void FEditorDragToolBehaviorTarget::OnUpdateModifierState(int InModifierID, bool bInIsOn)
+void FEditorDragToolBehaviorTarget::OnStateUpdated(const FInputDeviceState& InInputDeviceState)
 {
-	switch (InModifierID)
-	{
-	case UE::Editor::DragTools::ShiftKeyMod:
-		InputState.bShiftKeyDown = bInIsOn;
-		break;
-	case UE::Editor::DragTools::CtrlKeyMod:
-		InputState.bCtrlKeyDown = bInIsOn;
-		break;
-	case UE::Editor::DragTools::AltKeyMod:
-		InputState.bAltKeyDown = bInIsOn;
-		break;
-	default:;
-	}
+	InputState = InInputDeviceState;
 }
 
-void FEditorDragToolBehaviorTarget::OnClickRelease(const FInputDeviceRay& InReleasePos)
+void FEditorDragToolBehaviorTarget::OnDragEnd(const FInputDeviceRay& InDragEndPos)
 {
 	Start = End = FVector::ZeroVector;
 	bIsDragging = false;
@@ -115,28 +103,6 @@ void FEditorDragToolBehaviorTarget::OnTerminateDragSequence()
 	EditorViewportClientProxy->GetEditorViewportClient()->Invalidate(true, false);
 }
 
-bool FEditorDragToolBehaviorTarget::IsActivationChordPressed(const FInputChord& InChord) const
-{
-	bool bIsChordPressed = true;
-
-	if (InChord.NeedsControl())
-	{
-		bIsChordPressed = bIsChordPressed && FInputDeviceState::IsCtrlKeyDown(InputState);
-	}
-
-	if (InChord.NeedsAlt())
-	{
-		bIsChordPressed = bIsChordPressed && FInputDeviceState::IsAltKeyDown(InputState);
-	}
-
-	if (InChord.NeedsShift())
-	{
-		bIsChordPressed = bIsChordPressed && FInputDeviceState::IsShiftKeyDown(InputState);
-	}
-
-	return bIsChordPressed;
-}
-
 bool FEditorDragToolBehaviorTarget::IsCurrentModeSupported()
 {
 	if (ModeTools)
diff --git a/Engine/Source/Editor/Experimental/EditorInteractiveToolsFramework/Private/EditorDragTools/EditorDuplicateDragSelection.cpp b/Engine/Source/Editor/Experimental/EditorInteractiveToolsFramework/Private/EditorDragTools/EditorDuplicateDragSelection.cpp
--- a/Engine/Source/Editor/Experimental/EditorInteractiveToolsFramework/Private/EditorDragTools/EditorDuplicateDragSelection.cpp
+++ b/Engine/Source/Editor/Experimental/EditorInteractiveToolsFramework/Private/EditorDragTools/EditorDuplicateDragSelection.cpp
@@ -158,7 +158,7 @@ void FEditorDuplicateDragSelection::OnGizmoMovementBegin(UTransformProxy* InTran
 	{
 		bGizmoIsDragged = true;
 
-		if (bModifierKeyIsPressed)
+		if (bModifierKeyIsPressed && !TransformGizmo->IsIndirectManipulation())
 		{
 			DuplicateTransaction = MakeShared<FScopedTransaction>(LOCTEXT("DuplicateAndTransformTransactionName", "Duplicate and Transform"));
 
diff --git a/Engine/Source/Editor/Experimental/EditorInteractiveToolsFramework/Private/EditorDragTools/EditorFrustumSelect.cpp b/Engine/Source/Editor/Experimental/EditorInteractiveToolsFramework/Private/EditorDragTools/EditorFrustumSelect.cpp
--- a/Engine/Source/Editor/Experimental/EditorInteractiveToolsFramework/Private/EditorDragTools/EditorFrustumSelect.cpp
+++ b/Engine/Source/Editor/Experimental/EditorInteractiveToolsFramework/Private/EditorDragTools/EditorFrustumSelect.cpp
@@ -138,16 +138,12 @@ void FViewportFrustum::Calculate(FConvexVolume& OutFrustum) const
 
 FInputRayHit FEditorFrustumSelect::CanBeginClickDragSequence(const FInputDeviceRay& InPressPos)
 {
-	// Todo: this could be retrieved from a command for customization
-	FInputChord ActivationChord(EModifierKey::Control | EModifierKey::Alt, EKeys::LeftMouseButton);
-
-	return IsActivationChordPressed(ActivationChord)
-				&& IsCurrentModeSupported()
+	return IsCurrentModeSupported()
 			 ? FInputRayHit(TNumericLimits<float>::Max()) // bHit is true. Depth is max to lose the standard tiebreaker.
 			 : FInputRayHit();
 }
 
-void FEditorFrustumSelect::OnClickPress(const FInputDeviceRay& InPressPos)
+void FEditorFrustumSelect::OnDragStart(const FInputDeviceRay& InPressPos)
 {
 	// Signal that this tool is now active
 	OnActivateTool().Broadcast();
@@ -171,12 +167,12 @@ void FEditorFrustumSelect::OnClickPress(const FInputDeviceRay& InPressPos)
 	End = Start;
 }
 
-void FEditorFrustumSelect::OnClickDrag(const FInputDeviceRay& InDragPos)
+void FEditorFrustumSelect::OnDrag(const FInputDeviceRay& InDragPos)
 {
 	End = FVector(InDragPos.ScreenPosition.X, InDragPos.ScreenPosition.Y, 0);
 }
 
-void FEditorFrustumSelect::OnClickRelease(const FInputDeviceRay& InReleasePos)
+void FEditorFrustumSelect::OnDragEnd(const FInputDeviceRay& InReleasePos)
 {
 	End = FVector(InReleasePos.ScreenPosition.X, InReleasePos.ScreenPosition.Y, 0);
 
@@ -382,7 +378,7 @@ void FEditorFrustumSelect::OnClickRelease(const FInputDeviceRay& InReleasePos)
 	// Clear any hovered objects that might have been created while dragging
 	FLevelEditorViewportClient::ClearHoverFromObjects();
 
-	FEditorDragToolBehaviorTarget::OnClickRelease(InReleasePos);
+	FEditorDragToolBehaviorTarget::OnDragEnd(InReleasePos);
 }
 
 void FEditorFrustumSelect::OnTerminateDragSequence()
diff --git a/Engine/Source/Editor/Experimental/EditorInteractiveToolsFramework/Private/EditorDragTools/EditorMeasureTool.cpp b/Engine/Source/Editor/Experimental/EditorInteractiveToolsFramework/Private/EditorDragTools/EditorMeasureTool.cpp
--- a/Engine/Source/Editor/Experimental/EditorInteractiveToolsFramework/Private/EditorDragTools/EditorMeasureTool.cpp
+++ b/Engine/Source/Editor/Experimental/EditorInteractiveToolsFramework/Private/EditorDragTools/EditorMeasureTool.cpp
@@ -104,7 +104,7 @@ FInputRayHit FEditorMeasureTool::CanBeginClickDragSequence(const FInputDeviceRay
 	return FInputRayHit();
 }
 
-void FEditorMeasureTool::OnClickPress(const FInputDeviceRay& InPressPos)
+void FEditorMeasureTool::OnDragStart(const FInputDeviceRay& InPressPos)
 {
 	FEditorViewportClient* const EditorViewportClient = EditorViewportClientProxy->GetEditorViewportClient();
 	if (EditorViewportClient && EditorViewportClient->Viewport)
@@ -134,7 +134,7 @@ void FEditorMeasureTool::OnClickPress(const FInputDeviceRay& InPressPos)
 	}
 }
 
-void FEditorMeasureTool::OnClickDrag(const FInputDeviceRay& InDragPos)
+void FEditorMeasureTool::OnDrag(const FInputDeviceRay& InDragPos)
 {
 	if (FEditorViewportClient* const EditorViewportClient = EditorViewportClientProxy->GetEditorViewportClient())
 	{
@@ -147,16 +147,6 @@ void FEditorMeasureTool::OnClickDrag(const FInputDeviceRay& InDragPos)
 	}
 }
 
-void FEditorMeasureTool::OnClickRelease(const FInputDeviceRay& InReleasePos)
-{
-	FEditorDragToolBehaviorTarget::OnClickRelease(InReleasePos);
-}
-
-void FEditorMeasureTool::OnTerminateDragSequence()
-{
-	FEditorDragToolBehaviorTarget::OnTerminateDragSequence();
-}
-
 FVector2D FEditorMeasureTool::GetSnappedPixelPos(FVector2D InPixelPos) const
 {
 	FEditorViewportClient* const EditorViewportClient = EditorViewportClientProxy->GetEditorViewportClient();
diff --git a/Engine/Source/Editor/Experimental/EditorInteractiveToolsFramework/Private/EditorDragTools/EditorViewportChange.cpp b/Engine/Source/Editor/Experimental/EditorInteractiveToolsFramework/Private/EditorDragTools/EditorViewportChange.cpp
--- a/Engine/Source/Editor/Experimental/EditorInteractiveToolsFramework/Private/EditorDragTools/EditorViewportChange.cpp
+++ b/Engine/Source/Editor/Experimental/EditorInteractiveToolsFramework/Private/EditorDragTools/EditorViewportChange.cpp
@@ -40,15 +40,10 @@ void FEditorViewportChange::Render(const FSceneView* View, FCanvas* Canvas, EVie
 
 FInputRayHit FEditorViewportChange::CanBeginClickDragSequence(const FInputDeviceRay& InPressPos)
 {
-	// Todo: this could be retrieved from a command for customization
-	FInputChord ActivationChord(EModifierKey::Control, EKeys::MiddleMouseButton);
-
-	return IsActivationChordPressed(ActivationChord)
-			 ? FInputRayHit(TNumericLimits<float>::Max()) // bHit is true. Depth is max to lose the standard tiebreaker.
-			 : FInputRayHit();
+	return FInputRayHit(TNumericLimits<float>::Max()); // bHit is true. Depth is max to lose the standard tiebreaker.
 }
 
-void FEditorViewportChange::OnClickPress(const FInputDeviceRay& InPressPos)
+void FEditorViewportChange::OnDragStart(const FInputDeviceRay& InPressPos)
 {
 	FEditorViewportClient* const EditorViewportClient = EditorViewportClientProxy->GetEditorViewportClient();
 	if (EditorViewportClient && EditorViewportClient->Viewport)
@@ -75,7 +70,7 @@ void FEditorViewportChange::OnClickPress(const FInputDeviceRay& InPressPos)
 	}
 }
 
-void FEditorViewportChange::OnClickDrag(const FInputDeviceRay& InDragPos)
+void FEditorViewportChange::OnDrag(const FInputDeviceRay& InDragPos)
 {
 	if (FEditorViewportClient* const EditorViewportClient = EditorViewportClientProxy->GetEditorViewportClient())
 	{
@@ -90,19 +85,14 @@ void FEditorViewportChange::OnClickDrag(const FInputDeviceRay& InDragPos)
 	}
 }
 
-void FEditorViewportChange::OnClickRelease(const FInputDeviceRay& InReleasePos)
+void FEditorViewportChange::OnDragEnd(const FInputDeviceRay& InReleasePos)
 {
 	if (FEditorViewportClient* const EditorViewportClient = EditorViewportClientProxy->GetEditorViewportClient())
 	{
 		EditorViewportClient->SetViewportType(GetDesiredViewportType());
 	}
 
-	FEditorDragToolBehaviorTarget::OnClickRelease(InReleasePos);
-}
-
-void FEditorViewportChange::OnTerminateDragSequence()
-{
-	FEditorDragToolBehaviorTarget::OnTerminateDragSequence();
+	FEditorDragToolBehaviorTarget::OnDragEnd(InReleasePos);
 }
 
 ELevelViewportType FEditorViewportChange::GetDesiredViewportType() const
diff --git a/Engine/Source/Editor/Experimental/EditorInteractiveToolsFramework/Public/EditorDragTools/DragToolsBehaviourSource.h b/Engine/Source/Editor/Experimental/EditorInteractiveToolsFramework/Public/EditorDragTools/DragToolsBehaviourSource.h
--- a/Engine/Source/Editor/Experimental/EditorInteractiveToolsFramework/Public/EditorDragTools/DragToolsBehaviourSource.h
+++ b/Engine/Source/Editor/Experimental/EditorInteractiveToolsFramework/Public/EditorDragTools/DragToolsBehaviourSource.h
@@ -21,10 +21,6 @@ static constexpr int DRAG_TOOLS_DEFAULT_PRIORITY = HANDLED_BEFORE_VIEWPORT_INTER
 static constexpr int DRAG_TOOLS_LOW_PRIORITY = HANDLED_BEFORE_VIEWPORT_INTERACTIONS - 1;
 
 static constexpr int DUPLICATE_SELECTION_PRIORITY = DRAG_TOOLS_DEFAULT_PRIORITY;
-static constexpr int FRUSTUM_SELECT_PRIORITY = DRAG_TOOLS_DEFAULT_PRIORITY;
-static constexpr int VIEWPORT_CHANGE_PRIORITY = DRAG_TOOLS_DEFAULT_PRIORITY;
-static constexpr int BOX_SELECT_PRIORITY = HANDLED_AFTER_GIZMO_PRIORITY;
-static constexpr int MEASURE_PRIORITY = HANDLED_AFTER_GIZMO_PRIORITY;
 
 }
 
diff --git a/Engine/Source/Editor/Experimental/EditorInteractiveToolsFramework/Public/EditorDragTools/EditorBoxSelect.h b/Engine/Source/Editor/Experimental/EditorInteractiveToolsFramework/Public/EditorDragTools/EditorBoxSelect.h
--- a/Engine/Source/Editor/Experimental/EditorInteractiveToolsFramework/Public/EditorDragTools/EditorBoxSelect.h
+++ b/Engine/Source/Editor/Experimental/EditorInteractiveToolsFramework/Public/EditorDragTools/EditorBoxSelect.h
@@ -16,13 +16,13 @@ public:
 	{
 	}
 
-	//~ Begin IClickDragBehaviorTarget
+	//~ Begin IViewportClickDragBehaviorTarget
 	virtual FInputRayHit CanBeginClickDragSequence(const FInputDeviceRay& InPressPos) override;
-	virtual void OnClickPress(const FInputDeviceRay& InPressPos) override;
-	virtual void OnClickDrag(const FInputDeviceRay& InDragPos) override;
-	virtual void OnClickRelease(const FInputDeviceRay& InReleasePos) override;
+	virtual void OnDragStart(const FInputDeviceRay& InPressPos) override;
+	virtual void OnDrag(const FInputDeviceRay& InDragPos) override;
+	virtual void OnDragEnd(const FInputDeviceRay& InReleasePos) override;
 	virtual void OnTerminateDragSequence() override;
-	//~ End IClickDragBehaviorTarget
+	//~ End IViewportClickDragBehaviorTarget
 
 protected:
 	virtual TArray<FEditorModeID> GetUnsupportedModes() override;
diff --git a/Engine/Source/Editor/Experimental/EditorInteractiveToolsFramework/Public/EditorDragTools/EditorDragToolBehaviorTarget.h b/Engine/Source/Editor/Experimental/EditorInteractiveToolsFramework/Public/EditorDragTools/EditorDragToolBehaviorTarget.h
--- a/Engine/Source/Editor/Experimental/EditorInteractiveToolsFramework/Public/EditorDragTools/EditorDragToolBehaviorTarget.h
+++ b/Engine/Source/Editor/Experimental/EditorInteractiveToolsFramework/Public/EditorDragTools/EditorDragToolBehaviorTarget.h
@@ -5,6 +5,7 @@
 #include "BaseBehaviors/BehaviorTargetInterfaces.h"
 #include "Editor.h"
 #include "ToolContextInterfaces.h"
+#include "Behaviors/ViewportClickDragBehavior.h"
 
 class IEditorViewportClientProxy;
 struct FWorldSelectionElementArgs;
@@ -41,10 +42,8 @@ constexpr int CtrlKeyMod = 3;
 /**
  * The base class that all drag tools inherit from.
  * The drag tools implement special behaviors for the user clicking and dragging in a viewport.
- *
- * Inheriting from IClickDragBehaviorTarget so that Drag Tools logic can be shared between legacy and ITF usage of the same tools
  */
-class FEditorDragToolBehaviorTarget : public IClickDragBehaviorTarget
+class FEditorDragToolBehaviorTarget : public IViewportClickDragBehaviorTarget
 {
 public:
 	FEditorDragToolBehaviorTarget(FEditorViewportClient* InEditorViewportClient);
@@ -69,32 +68,28 @@ public:
 	/** Does this drag tool need to have the mouse movement converted to the viewport orientation? */
 	bool bConvertDelta;
 
-	//~ Begin IModifierToggleBehaviorTarget
-	virtual void OnUpdateModifierState(int InModifierID, bool bInIsOn) override;
-	//~ End IModifierToggleBehaviorTarget
-
-	//~ Begin IClickDragBehaviorTarget
+	//~ Begin IViewportClickDragBehaviorTarget
 	virtual FInputRayHit CanBeginClickDragSequence(const FInputDeviceRay& InPressPos) override
 	{
 		return FInputRayHit();
 	}
 
-	virtual void OnClickPress(const FInputDeviceRay& InPressPos) override
+	virtual void OnDragStart(const FInputDeviceRay& InDragStartPos) override
 	{
 	}
 
-	virtual void OnClickDrag(const FInputDeviceRay& InDragPos) override
+	virtual void OnDrag(const FInputDeviceRay& InDragPos) override
 	{
 	}
 
-	virtual void OnClickRelease(const FInputDeviceRay& InReleasePos) override;
+	virtual void OnDragEnd(const FInputDeviceRay& InDragEndPos) override;
 
 	virtual void OnForceEndCapture() override;
 
 	virtual void OnTerminateDragSequence() override;
-	//~ End IClickDragBehaviorTarget
 
-	bool IsActivationChordPressed(const FInputChord& InChord) const;
+	virtual void OnStateUpdated(const FInputDeviceState& InInputDeviceState) override;
+	//~ End IViewportClickDragBehaviorTarget
 
 	DECLARE_MULTICAST_DELEGATE(FOnToolStateChange);
 	FOnToolStateChange& OnActivateTool()
diff --git a/Engine/Source/Editor/Experimental/EditorInteractiveToolsFramework/Public/EditorDragTools/EditorFrustumSelect.h b/Engine/Source/Editor/Experimental/EditorInteractiveToolsFramework/Public/EditorDragTools/EditorFrustumSelect.h
--- a/Engine/Source/Editor/Experimental/EditorInteractiveToolsFramework/Public/EditorDragTools/EditorFrustumSelect.h
+++ b/Engine/Source/Editor/Experimental/EditorInteractiveToolsFramework/Public/EditorDragTools/EditorFrustumSelect.h
@@ -26,13 +26,13 @@ public:
 	{
 	}
 
-	//~ Begin IClickDragBehaviorTarget
+	//~ Begin IViewportClickDragBehaviorTarget
 	virtual FInputRayHit CanBeginClickDragSequence(const FInputDeviceRay& InPressPos) override;
-	virtual void OnClickPress(const FInputDeviceRay& InPressPos) override;
-	virtual void OnClickDrag(const FInputDeviceRay& InDragPos) override;
-	virtual void OnClickRelease(const FInputDeviceRay& InReleasePos) override;
+	virtual void OnDragStart(const FInputDeviceRay& InPressPos) override;
+	virtual void OnDrag(const FInputDeviceRay& InDragPos) override;
+	virtual void OnDragEnd(const FInputDeviceRay& InReleasePos) override;
 	virtual void OnTerminateDragSequence() override;
-	//~ End IClickDragBehaviorTarget
+	//~ End IViewportClickDragBehaviorTarget
 
 protected:
 	virtual TArray<FEditorModeID> GetUnsupportedModes() override;
diff --git a/Engine/Source/Editor/Experimental/EditorInteractiveToolsFramework/Public/EditorDragTools/EditorMeasureTool.h b/Engine/Source/Editor/Experimental/EditorInteractiveToolsFramework/Public/EditorDragTools/EditorMeasureTool.h
--- a/Engine/Source/Editor/Experimental/EditorInteractiveToolsFramework/Public/EditorDragTools/EditorMeasureTool.h
+++ b/Engine/Source/Editor/Experimental/EditorInteractiveToolsFramework/Public/EditorDragTools/EditorMeasureTool.h
@@ -13,13 +13,11 @@ public:
 
 	virtual void Render(const FSceneView* InView,FCanvas* InCanvas, EViewInteractionState InInteractionState) override;
 
-	//~ Begin IClickDragBehaviorTarget
+	//~ Begin IViewportClickDragBehaviorTarget
 	virtual FInputRayHit CanBeginClickDragSequence(const FInputDeviceRay& InPressPos) override;
-	virtual void OnClickPress(const FInputDeviceRay& InPressPos) override;
-	virtual void OnClickDrag(const FInputDeviceRay& InDragPos) override;
-	virtual void OnClickRelease(const FInputDeviceRay& InReleasePos) override;
-	virtual void OnTerminateDragSequence() override;
-	//~ End IClickDragBehaviorTarget
+	virtual void OnDragStart(const FInputDeviceRay& InPressPos) override;
+	virtual void OnDrag(const FInputDeviceRay& InDragPos) override;
+	//~ End IViewportClickDragBehaviorTarget
 
 private:
 	/**
diff --git a/Engine/Source/Editor/Experimental/EditorInteractiveToolsFramework/Public/EditorDragTools/EditorViewportChange.h b/Engine/Source/Editor/Experimental/EditorInteractiveToolsFramework/Public/EditorDragTools/EditorViewportChange.h
--- a/Engine/Source/Editor/Experimental/EditorInteractiveToolsFramework/Public/EditorDragTools/EditorViewportChange.h
+++ b/Engine/Source/Editor/Experimental/EditorInteractiveToolsFramework/Public/EditorDragTools/EditorViewportChange.h
@@ -16,13 +16,12 @@ public:
 	explicit FEditorViewportChange(FEditorViewportClient* InEditorViewportClient);
 	virtual void Render(const FSceneView* View, FCanvas* Canvas, EViewInteractionState InInteractionState) override;
 
-	//~ Begin IClickDragBehaviorTarget
+	//~ Begin IViewportClickDragBehaviorTarget
 	virtual FInputRayHit CanBeginClickDragSequence(const FInputDeviceRay& InPressPos) override;
-	virtual void OnClickPress(const FInputDeviceRay& InPressPos) override;
-	virtual void OnClickDrag(const FInputDeviceRay& InDragPos) override;
-	virtual void OnClickRelease(const FInputDeviceRay& InReleasePos) override;
-	virtual void OnTerminateDragSequence() override;
-	//~ End IClickDragBehaviorTarget
+	virtual void OnDragStart(const FInputDeviceRay& InPressPos) override;
+	virtual void OnDrag(const FInputDeviceRay& InDragPos) override;
+	virtual void OnDragEnd(const FInputDeviceRay& InReleasePos) override;
+	//~ End IViewportClickDragBehaviorTarget
 
 private:
 	ELevelViewportType GetDesiredViewportType() const;
diff --git a/Engine/Source/Editor/Experimental/EditorInteractiveToolsFramework/Public/EditorGizmos/TransformGizmo.h b/Engine/Source/Editor/Experimental/EditorInteractiveToolsFramework/Public/EditorGizmos/TransformGizmo.h
--- a/Engine/Source/Editor/Experimental/EditorInteractiveToolsFramework/Public/EditorGizmos/TransformGizmo.h
+++ b/Engine/Source/Editor/Experimental/EditorInteractiveToolsFramework/Public/EditorGizmos/TransformGizmo.h
@@ -320,7 +320,15 @@ public:
 	 * Handle user parameters changes
 	 */
 	UE_API virtual void OnParametersChanged(const FGizmosParameters& InParameters);
-	
+
+	/**
+	 * True if gizmo is being used for indirect manipulation
+	 */
+	bool IsIndirectManipulation() const
+	{
+		return bIndirectManipulation;
+	}
+
 public:
 
 	/** The active target object for the Gizmo */
