[Transform Gizmo] Changed the default rotation mode to "ScreenArc"

diff --git a/Engine/Source/Editor/Experimental/EditorInteractiveToolsFramework/Public/EditorGizmos/TransformGizmo.h b/Engine/Source/Editor/Experimental/EditorInteractiveToolsFramework/Public/EditorGizmos/TransformGizmo.h
--- a/Engine/Source/Editor/Experimental/EditorInteractiveToolsFramework/Public/EditorGizmos/TransformGizmo.h
+++ b/Engine/Source/Editor/Experimental/EditorInteractiveToolsFramework/Public/EditorGizmos/TransformGizmo.h
@@ -150,7 +150,7 @@ struct FGizmosParameters
 
 	/** Determines how dragging the rotate gizmo affects the selected objects */
 	UPROPERTY(EditAnywhere, Category = "Experimental")
-	TEnumAsByte<EAxisRotateMode::Type> RotateMode = EAxisRotateMode::Arc;
+	TEnumAsByte<EAxisRotateMode::Type> RotateMode = EAxisRotateMode::ScreenArc;
 
 	/**
 	 * When enabled, Ctrl+MMB activates the Y axis and Ctrl+RMB activates the Z axis.
@@ -1085,10 +1085,10 @@ protected:
 
 	/** Default rotate mode used when using axis rotation handles. */
 	UPROPERTY()
-	TEnumAsByte<EAxisRotateMode::Type> DefaultRotateMode = EAxisRotateMode::Arc;
+	TEnumAsByte<EAxisRotateMode::Type> DefaultRotateMode = EAxisRotateMode::ScreenArc;
 
 	/** Actual rotate mode used (based on view dependant information). */
-	TEnumAsByte<EAxisRotateMode::Type> RotateMode = EAxisRotateMode::Arc;
+	TEnumAsByte<EAxisRotateMode::Type> RotateMode = EAxisRotateMode::ScreenArc;
 
 	/** Switch from tangential to normal projection based on the first mouse drag. */
 	bool bTrySwitchingToNormalPull = false;
