Unshelved from [ITF Viewport Interactions] Cursor rendering with incorrect icon when holding click in Level viewport

diff --git a/Engine/Source/Editor/Experimental/EditorInteractiveToolsFramework/Private/Behaviors/ViewportClickDragBehavior.cpp b/Engine/Source/Editor/Experimental/EditorInteractiveToolsFramework/Private/Behaviors/ViewportClickDragBehavior.cpp
--- a/Engine/Source/Editor/Experimental/EditorInteractiveToolsFramework/Private/Behaviors/ViewportClickDragBehavior.cpp
+++ b/Engine/Source/Editor/Experimental/EditorInteractiveToolsFramework/Private/Behaviors/ViewportClickDragBehavior.cpp
@@ -84,6 +84,11 @@ FInputCaptureRequest UViewportClickDragBehavior::WantsCapture(const FInputDevice
 
 FInputCaptureUpdate UViewportClickDragBehavior::BeginCapture(const FInputDeviceState& InputState, EInputCaptureSide eSide)
 {
+	if (Target)
+	{
+		Target->OnBeginCapture(GetDeviceRay(InputState));
+	}
+
 	OnStateUpdatedInternal(InputState);
 	return FInputCaptureUpdate::Begin(this, EInputCaptureSide::Any);
 }
diff --git a/Engine/Source/Editor/Experimental/EditorInteractiveToolsFramework/Private/EditorDragTools/EditorBoxSelect.cpp b/Engine/Source/Editor/Experimental/EditorInteractiveToolsFramework/Private/EditorDragTools/EditorBoxSelect.cpp
--- a/Engine/Source/Editor/Experimental/EditorInteractiveToolsFramework/Private/EditorDragTools/EditorBoxSelect.cpp
+++ b/Engine/Source/Editor/Experimental/EditorInteractiveToolsFramework/Private/EditorDragTools/EditorBoxSelect.cpp
@@ -83,6 +83,11 @@ FInputRayHit FEditorBoxSelect::CanBeginClickDragSequence(const FInputDeviceRay&
 
 void FEditorBoxSelect::OnDragStart(const FInputDeviceRay& InPressPos)
 {
+	if (UToolsContextCursorAPI* const ToolsContextCursorAPI = GetCursorAPI())
+	{
+		ToolsContextCursorAPI->SetCursorOverride(EMouseCursor::Default);
+	}
+
 	if (!EditorViewportClientProxy)
 	{
 		return;
@@ -220,6 +225,11 @@ void FEditorBoxSelect::OnDrag(const FInputDeviceRay& InDragPos)
 
 void FEditorBoxSelect::OnDragEnd(const FInputDeviceRay& InReleasePos)
 {
+	if (UToolsContextCursorAPI* const ToolsContextCursorAPI = GetCursorAPI())
+	{
+		ToolsContextCursorAPI->ClearCursorOverride();
+	}
+
 	if (EditorViewportClientProxy)
 	{
 		if (FEditorViewportClient* const EditorViewportClient = EditorViewportClientProxy->GetEditorViewportClient())
diff --git a/Engine/Source/Editor/Experimental/EditorInteractiveToolsFramework/Private/EditorDragTools/EditorDragToolBehaviorTarget.cpp b/Engine/Source/Editor/Experimental/EditorInteractiveToolsFramework/Private/EditorDragTools/EditorDragToolBehaviorTarget.cpp
--- a/Engine/Source/Editor/Experimental/EditorInteractiveToolsFramework/Private/EditorDragTools/EditorDragToolBehaviorTarget.cpp
+++ b/Engine/Source/Editor/Experimental/EditorInteractiveToolsFramework/Private/EditorDragTools/EditorDragToolBehaviorTarget.cpp
@@ -1,11 +1,13 @@
 ﻿// Copyright Epic Games, Inc. All Rights Reserved.
 
 #include "EditorDragTools/EditorDragToolBehaviorTarget.h"
+#include "ContextObjectStore.h"
 #include "Editor.h"
 #include "EditorDragTools/EditorViewportClientProxy.h"
 #include "EditorModeManager.h"
 #include "EditorViewportClient.h"
 #include "InputCoreTypes.h"
+#include "Tools/EdModeInteractiveToolsContext.h"
 
 namespace UE::Editor::DragTools
 {
@@ -76,6 +78,11 @@ void FEditorDragToolBehaviorTarget::OnStateUpdated(const FInputDeviceState& InIn
 
 void FEditorDragToolBehaviorTarget::OnDragEnd(const FInputDeviceRay& InDragEndPos)
 {
+	if (UToolsContextCursorAPI* const ToolsContextCursorAPI = GetCursorAPI())
+	{
+		ToolsContextCursorAPI->ClearCursorOverride();
+	}
+
 	Start = End = FVector::ZeroVector;
 	bIsDragging = false;
 
@@ -90,6 +97,11 @@ void FEditorDragToolBehaviorTarget::OnForceEndCapture()
 
 void FEditorDragToolBehaviorTarget::OnTerminateDragSequence()
 {
+	if (UToolsContextCursorAPI* const ToolsContextCursorAPI = GetCursorAPI())
+	{
+		ToolsContextCursorAPI->ClearCursorOverride();
+	}
+
 	Start = End = FVector::ZeroVector;
 	bIsDragging = false;
 
@@ -118,3 +130,19 @@ bool FEditorDragToolBehaviorTarget::IsCurrentModeSupported()
 
 	return true;
 }
+
+UToolsContextCursorAPI* FEditorDragToolBehaviorTarget::GetCursorAPI() const
+{
+	if (ModeTools)
+	{
+		if (const UModeManagerInteractiveToolsContext* const InteractiveToolsContext = ModeTools->GetInteractiveToolsContext())
+		{
+			if (const UContextObjectStore* const ContextObjectStore = InteractiveToolsContext->ContextObjectStore)
+			{
+				return ContextObjectStore->FindContext<UToolsContextCursorAPI>();
+			}
+		}
+	}
+
+	return nullptr;
+}
diff --git a/Engine/Source/Editor/Experimental/EditorInteractiveToolsFramework/Private/EditorDragTools/EditorFrustumSelect.cpp b/Engine/Source/Editor/Experimental/EditorInteractiveToolsFramework/Private/EditorDragTools/EditorFrustumSelect.cpp
--- a/Engine/Source/Editor/Experimental/EditorInteractiveToolsFramework/Private/EditorDragTools/EditorFrustumSelect.cpp
+++ b/Engine/Source/Editor/Experimental/EditorInteractiveToolsFramework/Private/EditorDragTools/EditorFrustumSelect.cpp
@@ -4,6 +4,7 @@
 
 #include "CanvasItem.h"
 #include "CanvasTypes.h"
+#include "ContextObjectStore.h"
 #include "Components/PrimitiveComponent.h"
 #include "Editor/UnrealEdEngine.h"
 #include "EditorDragTools/EditorViewportClientProxy.h"
@@ -23,6 +24,7 @@
 #include "SnappingUtils.h"
 #include "Subsystems/BrushEditingSubsystem.h"
 #include "UnrealEdGlobals.h"
+#include "Tools/EdModeInteractiveToolsContext.h"
 
 namespace UE::EditorDragTools
 {
@@ -145,6 +147,11 @@ FInputRayHit FEditorFrustumSelect::CanBeginClickDragSequence(const FInputDeviceR
 
 void FEditorFrustumSelect::OnDragStart(const FInputDeviceRay& InPressPos)
 {
+	if (UToolsContextCursorAPI* const ToolsContextCursorAPI = GetCursorAPI())
+	{
+		ToolsContextCursorAPI->SetCursorOverride(EMouseCursor::Default);
+	}
+
 	// Signal that this tool is now active
 	OnActivateTool().Broadcast();
 
@@ -174,6 +181,11 @@ void FEditorFrustumSelect::OnDrag(const FInputDeviceRay& InDragPos)
 
 void FEditorFrustumSelect::OnDragEnd(const FInputDeviceRay& InReleasePos)
 {
+	if (UToolsContextCursorAPI* const ToolsContextCursorAPI = GetCursorAPI())
+	{
+		ToolsContextCursorAPI->ClearCursorOverride();
+	}
+
 	End = FVector(InReleasePos.ScreenPosition.X, InReleasePos.ScreenPosition.Y, 0);
 
 	if (!EditorViewportClientProxy)
diff --git a/Engine/Source/Editor/Experimental/EditorInteractiveToolsFramework/Private/EditorDragTools/EditorMeasureTool.cpp b/Engine/Source/Editor/Experimental/EditorInteractiveToolsFramework/Private/EditorDragTools/EditorMeasureTool.cpp
--- a/Engine/Source/Editor/Experimental/EditorInteractiveToolsFramework/Private/EditorDragTools/EditorMeasureTool.cpp
+++ b/Engine/Source/Editor/Experimental/EditorInteractiveToolsFramework/Private/EditorDragTools/EditorMeasureTool.cpp
@@ -106,6 +106,11 @@ FInputRayHit FEditorMeasureTool::CanBeginClickDragSequence(const FInputDeviceRay
 
 void FEditorMeasureTool::OnDragStart(const FInputDeviceRay& InPressPos)
 {
+	if (UToolsContextCursorAPI* const ToolsContextCursorAPI = GetCursorAPI())
+	{
+		ToolsContextCursorAPI->SetCursorOverride(EMouseCursor::Default);
+	}
+
 	FEditorViewportClient* const EditorViewportClient = EditorViewportClientProxy->GetEditorViewportClient();
 	if (EditorViewportClient && EditorViewportClient->Viewport)
 	{
diff --git a/Engine/Source/Editor/Experimental/EditorInteractiveToolsFramework/Private/ViewportInteractions/ViewportClickInteraction.cpp b/Engine/Source/Editor/Experimental/EditorInteractiveToolsFramework/Private/ViewportInteractions/ViewportClickInteraction.cpp
--- a/Engine/Source/Editor/Experimental/EditorInteractiveToolsFramework/Private/ViewportInteractions/ViewportClickInteraction.cpp
+++ b/Engine/Source/Editor/Experimental/EditorInteractiveToolsFramework/Private/ViewportInteractions/ViewportClickInteraction.cpp
@@ -22,11 +22,21 @@ UViewportClickInteraction::UViewportClickInteraction()
 
 FInputRayHit UViewportClickInteraction::IsHitByClick(const FInputDeviceRay& InClickPos)
 {
+	if (UViewportInteractionsBehaviorSource* const ViewportInteractionBehaviorSource = GetViewportInteractionsBehaviorSource())
+	{
+		ViewportInteractionBehaviorSource->SetMouseCursorOverride(EMouseCursor::Crosshairs);
+	}
+
 	return FInputRayHit(0.0f);
 }
 
 void UViewportClickInteraction::OnClicked(const FInputDeviceRay& InClickPos)
 {
+	if (UViewportInteractionsBehaviorSource* const ViewportInteractionBehaviorSource = GetViewportInteractionsBehaviorSource())
+	{
+		ViewportInteractionBehaviorSource->ClearMouseCursorOverride();
+	}
+
 	// Make sure this release is not happening while user is moving the viewport camera - in that case, this is not considered a click.
 	if (IsMouseLooking())
 	{
diff --git a/Engine/Source/Editor/Experimental/EditorInteractiveToolsFramework/Private/ViewportInteractions/ViewportDragInteraction.cpp b/Engine/Source/Editor/Experimental/EditorInteractiveToolsFramework/Private/ViewportInteractions/ViewportDragInteraction.cpp
--- a/Engine/Source/Editor/Experimental/EditorInteractiveToolsFramework/Private/ViewportInteractions/ViewportDragInteraction.cpp
+++ b/Engine/Source/Editor/Experimental/EditorInteractiveToolsFramework/Private/ViewportInteractions/ViewportDragInteraction.cpp
@@ -33,6 +33,14 @@ FInputRayHit UViewportDragInteraction::CanBeginClickDragSequence(const FInputDev
 	return FInputRayHit();
 }
 
+void UViewportDragInteraction::OnBeginCapture(const FInputDeviceRay& InClickPressPos)
+{
+	if (UViewportInteractionsBehaviorSource* BehaviorSource = GetViewportInteractionsBehaviorSource())
+	{
+		BehaviorSource->SetMouseCursorOverride(EMouseCursor::Crosshairs);
+	}
+}
+
 void UViewportDragInteraction::OnDragStart(const FInputDeviceRay& InDragStartPos)
 {
 	if (!IsEnabled())
diff --git a/Engine/Source/Editor/Experimental/EditorInteractiveToolsFramework/Private/ViewportInteractions/ViewportInteractionsBehaviorSource.cpp b/Engine/Source/Editor/Experimental/EditorInteractiveToolsFramework/Private/ViewportInteractions/ViewportInteractionsBehaviorSource.cpp
--- a/Engine/Source/Editor/Experimental/EditorInteractiveToolsFramework/Private/ViewportInteractions/ViewportInteractionsBehaviorSource.cpp
+++ b/Engine/Source/Editor/Experimental/EditorInteractiveToolsFramework/Private/ViewportInteractions/ViewportInteractionsBehaviorSource.cpp
@@ -1,6 +1,8 @@
 // Copyright Epic Games, Inc. All Rights Reserved.
 
 #include "ViewportInteractions/ViewportInteractionsBehaviorSource.h"
+
+#include "ContextObjectStore.h"
 #include "BaseBehaviors/KeyAsModifierInputBehavior.h"
 #include "EditorModeManager.h"
 #include "EditorViewportClient.h"
@@ -102,6 +104,25 @@ static const FString ViewportInteractionsGroupName = TEXT("ViewportInteractions"
 
 } // namespace UE::Editor::ViewportInteractions
 
+UToolsContextCursorAPI* GetCursorAPI(const FEditorViewportClient* InViewportClient)
+{
+	if (InViewportClient)
+	{
+		if (FEditorModeTools* ModeTools = InViewportClient->GetModeTools())
+		{
+			if (const UModeManagerInteractiveToolsContext* const InteractiveToolsContext = ModeTools->GetInteractiveToolsContext())
+			{
+				if (const UContextObjectStore* const ContextObjectStore = InteractiveToolsContext->ContextObjectStore)
+				{
+					return ContextObjectStore->FindContext<UToolsContextCursorAPI>();
+				}
+			}
+		}
+	}
+
+	return nullptr;
+}
+
 void UViewportInteractionsBehaviorSource::OnUpdateModifierState(int InModifierID, bool bInIsOn)
 {
 	if (InModifierID == UE::Editor::ViewportInteractions::ShiftKeyMod)
@@ -286,6 +307,31 @@ void UViewportInteractionsBehaviorSource::Reset()
 void UViewportInteractionsBehaviorSource::SetMouseCursorOverride(EMouseCursor::Type InMouseCursor)
 {
 	CursorOverride = InMouseCursor;
+
+	if (FEditorViewportClient* const EditorViewportClient = GetEditorViewportClient())
+	{
+		if (UToolsContextCursorAPI* ToolsContextCursorAPI =  GetCursorAPI(EditorViewportClient))
+		{
+			ToolsContextCursorAPI->SetCursorOverride(CursorOverride);
+		}
+
+		EditorViewportClient->SetRequiredCursorOverride(true, CursorOverride);
+	}
+}
+
+void UViewportInteractionsBehaviorSource::ClearMouseCursorOverride()
+{
+	CursorOverride = EMouseCursor::Default;
+
+	if (FEditorViewportClient* const EditorViewportClient = GetEditorViewportClient())
+	{
+		if (UToolsContextCursorAPI* ToolsContextCursorAPI =  GetCursorAPI(EditorViewportClient))
+		{
+			ToolsContextCursorAPI->ClearCursorOverride();
+		}
+
+		EditorViewportClient->SetRequiredCursorOverride(false);
+	}
 }
 
 bool UViewportInteractionsBehaviorSource::IsMouseLooking() const
@@ -304,11 +350,11 @@ void UViewportInteractionsBehaviorSource::SetIsMouseLooking(bool bInIsLooking)
 	{
 		if (bInIsLooking)
 		{
-			EditorViewportClient->SetRequiredCursorOverride(true, CursorOverride);
+			SetMouseCursorOverride(CursorOverride);
 		}
 		else
 		{
-			EditorViewportClient->SetRequiredCursorOverride(false);
+			ClearMouseCursorOverride();
 		}
 
 		EditorViewportClient->GetViewportNavigationHelper()->bIsMouseLooking = bInIsLooking;
diff --git a/Engine/Source/Editor/Experimental/EditorInteractiveToolsFramework/Public/Behaviors/ViewportClickDragBehavior.h b/Engine/Source/Editor/Experimental/EditorInteractiveToolsFramework/Public/Behaviors/ViewportClickDragBehavior.h
--- a/Engine/Source/Editor/Experimental/EditorInteractiveToolsFramework/Public/Behaviors/ViewportClickDragBehavior.h
+++ b/Engine/Source/Editor/Experimental/EditorInteractiveToolsFramework/Public/Behaviors/ViewportClickDragBehavior.h
@@ -19,6 +19,12 @@ public:
 	 */
 	virtual FInputRayHit CanBeginClickDragSequence(const FInputDeviceRay& PressPos) = 0;
 
+	/**
+	 * Notify Target that capture has started
+	 * @param InClickPressPos device position/ray at click press point
+	 */
+	virtual void OnBeginCapture(const FInputDeviceRay& InClickPressPos) {}
+
 	/**
 	 * Notify Target that drag has started
 	 * @param InDragStartPos device position/ray at drag start point
diff --git a/Engine/Source/Editor/Experimental/EditorInteractiveToolsFramework/Public/EditorDragTools/EditorDragToolBehaviorTarget.h b/Engine/Source/Editor/Experimental/EditorInteractiveToolsFramework/Public/EditorDragTools/EditorDragToolBehaviorTarget.h
--- a/Engine/Source/Editor/Experimental/EditorInteractiveToolsFramework/Public/EditorDragTools/EditorDragToolBehaviorTarget.h
+++ b/Engine/Source/Editor/Experimental/EditorInteractiveToolsFramework/Public/EditorDragTools/EditorDragToolBehaviorTarget.h
@@ -106,6 +106,8 @@ protected:
 
 	bool IsCurrentModeSupported();
 
+	UToolsContextCursorAPI* GetCursorAPI() const;
+
 	FEditorModeTools* ModeTools;
 
 	/** The start/end location of the current drag. */
diff --git a/Engine/Source/Editor/Experimental/EditorInteractiveToolsFramework/Public/ViewportInteractions/ViewportDragInteraction.h b/Engine/Source/Editor/Experimental/EditorInteractiveToolsFramework/Public/ViewportInteractions/ViewportDragInteraction.h
--- a/Engine/Source/Editor/Experimental/EditorInteractiveToolsFramework/Public/ViewportInteractions/ViewportDragInteraction.h
+++ b/Engine/Source/Editor/Experimental/EditorInteractiveToolsFramework/Public/ViewportInteractions/ViewportDragInteraction.h
@@ -30,6 +30,7 @@ public:
 
 	//~ Begin IViewportClickDragBehaviorTarget
 	UE_API virtual FInputRayHit CanBeginClickDragSequence(const FInputDeviceRay& InPressPos) override;
+	UE_API virtual void OnBeginCapture(const FInputDeviceRay& InClickPressPos) override;
 	UE_API virtual void OnDragStart(const FInputDeviceRay& InDragStartPos) override;
 	UE_API virtual void OnDrag(const FInputDeviceRay& InDragPos) override;
 	UE_API virtual void OnDragEnd(const FInputDeviceRay& InDragEndPos) override;
diff --git a/Engine/Source/Editor/Experimental/EditorInteractiveToolsFramework/Public/ViewportInteractions/ViewportInteractionsBehaviorSource.h b/Engine/Source/Editor/Experimental/EditorInteractiveToolsFramework/Public/ViewportInteractions/ViewportInteractionsBehaviorSource.h
--- a/Engine/Source/Editor/Experimental/EditorInteractiveToolsFramework/Public/ViewportInteractions/ViewportInteractionsBehaviorSource.h
+++ b/Engine/Source/Editor/Experimental/EditorInteractiveToolsFramework/Public/ViewportInteractions/ViewportInteractionsBehaviorSource.h
@@ -171,6 +171,8 @@ public:
 	 */
 	void SetMouseCursorOverride(EMouseCursor::Type InMouseCursor);
 
+	void ClearMouseCursorOverride();
+
 	/**
 	 * Use the templated version in case there is some class specific initialization required by the newly created interaction,
 	 * since this method will return the newly instantiated interaction
@@ -254,7 +256,7 @@ private:
 	bool bIsMiddleMouseButtonDown;
 	bool bIsRightMouseButtonDown;
 
-	EMouseCursor::Type CursorOverride = EMouseCursor::None;
+	EMouseCursor::Type CursorOverride = EMouseCursor::Default;
 
 	UPROPERTY(Transient)
 	TWeakObjectPtr<const UEditorInteractiveToolsContext> EditorInteractiveToolsContextWeak;
