[ITF Gizmos] Rename the editor pref section from 'Transform Gizmo' to just 'Gizmos'

diff --git a/Engine/Source/Editor/Experimental/EditorInteractiveToolsFramework/Public/TransformGizmoEditorSettings.h b/Engine/Source/Editor/Experimental/EditorInteractiveToolsFramework/Public/TransformGizmoEditorSettings.h
--- a/Engine/Source/Editor/Experimental/EditorInteractiveToolsFramework/Public/TransformGizmoEditorSettings.h
+++ b/Engine/Source/Editor/Experimental/EditorInteractiveToolsFramework/Public/TransformGizmoEditorSettings.h
@@ -9,7 +9,7 @@
 
 #define UE_API EDITORINTERACTIVETOOLSFRAMEWORK_API
 
-UCLASS(MinimalAPI, BlueprintType, Config = EditorPerProjectUserSettings, meta = (DisplayName = "Transform Gizmo"))
+UCLASS(MinimalAPI, BlueprintType, Config = EditorPerProjectUserSettings, meta = (DisplayName = "Gizmos"))
 class UTransformGizmoEditorSettings : public UDeveloperSettings
 {
 	GENERATED_BODY()
