[ITF Viewport Interactions] Cannot create Mesh with Modeling mode with ITF Viewport Interaction enabled

Updating Viewport Interactions priorities to better work with existing ITF Tools

diff --git a/Engine/Source/Editor/Experimental/EditorInteractiveToolsFramework/Private/EditorDragTools/DragToolsBehaviourSource.cpp b/Engine/Source/Editor/Experimental/EditorInteractiveToolsFramework/Private/EditorDragTools/DragToolsBehaviourSource.cpp
--- a/Engine/Source/Editor/Experimental/EditorInteractiveToolsFramework/Private/EditorDragTools/DragToolsBehaviourSource.cpp
+++ b/Engine/Source/Editor/Experimental/EditorInteractiveToolsFramework/Private/EditorDragTools/DragToolsBehaviourSource.cpp
@@ -94,7 +94,7 @@ void UDragToolsBehaviorSource::Initialize(const UEditorInteractiveToolsContext*
 		UViewportClickDragBehavior* MouseBehavior = NewObject<UViewportClickDragBehavior>();
 		MouseBehavior->Initialize(FrustumSelectTool.Get());
 		MouseBehavior->SetDefaultPriority(UE::Editor::ViewportInteractions::DRAG_PRIORITY);
-		MouseBehavior->SetDragConfirmedPriority(UE::Editor::ViewportInteractions::VIEWPORT_INTERACTIONS_DEFAULT_PRIORITY);
+		MouseBehavior->SetDragConfirmedPriority(UE::Editor::ViewportInteractions::DRAG_TOOLS_DEFAULT_PRIORITY);
 		MouseBehavior->EnableButton(EKeys::LeftMouseButton);
 		MouseBehavior->SetAllowsCaptureStealing(true);
 		MouseBehavior->ModifierCheckFunc = [](const FInputDeviceState& InputState)
diff --git a/Engine/Source/Editor/Experimental/EditorInteractiveToolsFramework/Public/ViewportInteractions/ViewportInteraction.h b/Engine/Source/Editor/Experimental/EditorInteractiveToolsFramework/Public/ViewportInteractions/ViewportInteraction.h
--- a/Engine/Source/Editor/Experimental/EditorInteractiveToolsFramework/Public/ViewportInteractions/ViewportInteraction.h
+++ b/Engine/Source/Editor/Experimental/EditorInteractiveToolsFramework/Public/ViewportInteractions/ViewportInteraction.h
@@ -27,14 +27,15 @@ constexpr int LeftMouseButtonMod = 4;
 constexpr int RightMouseButtonMod = 5;
 constexpr int MiddleMouseButtonMod = 6;
 
-static constexpr int VIEWPORT_INTERACTIONS_DEFAULT_PRIORITY = FInputCapturePriority::DEFAULT_GIZMO_PRIORITY - 1;
-static constexpr int DRAG_CONFIRMED_PRIORITY = FInputCapturePriority::DEFAULT_GIZMO_PRIORITY + 5;
-static constexpr int CLICK_PRIORITY = FInputCapturePriority::DEFAULT_GIZMO_PRIORITY + 6;
-static constexpr int DRAG_PRIORITY = FInputCapturePriority::DEFAULT_GIZMO_PRIORITY + 7;
+static constexpr int VIEWPORT_INTERACTIONS_DEFAULT_PRIORITY = FInputCapturePriority::DEFAULT_TOOL_PRIORITY - 1;
+static constexpr int DRAG_TOOLS_DEFAULT_PRIORITY = FInputCapturePriority::DEFAULT_GIZMO_PRIORITY - 1;
+static constexpr int DRAG_CONFIRMED_PRIORITY = FInputCapturePriority::DEFAULT_TOOL_PRIORITY + 5;
+static constexpr int CLICK_PRIORITY = FInputCapturePriority::DEFAULT_TOOL_PRIORITY + 6;
+static constexpr int DRAG_PRIORITY = FInputCapturePriority::DEFAULT_TOOL_PRIORITY + 7;
 
 // Higher priority for behaviors needing multiple buttons, so that capture stealing can handle those
-static constexpr int MULTIBUTTON_DRAG_CONFIRMED_PRIORITY = DRAG_CONFIRMED_PRIORITY - 1;
-static constexpr int MULTIBUTTON_DRAG_PRIORITY = DRAG_PRIORITY - 1;
+static constexpr int MULTIBUTTON_DRAG_CONFIRMED_PRIORITY = FInputCapturePriority::DEFAULT_GIZMO_PRIORITY + 1;
+static constexpr int MULTIBUTTON_DRAG_PRIORITY = FInputCapturePriority::DEFAULT_GIZMO_PRIORITY + 2;
 
 /**
  * Possible types of Viewport Interactions.
