[ITF Interactions] Convert the new Drag Tools implementations to be UViewportInteractions. Deprecate drag tool related functions. Add group name support for all viewport interactions.

#rnx

diff --git a/Engine/Source/Editor/Experimental/EditorInteractiveToolsFramework/Private/EditorDragTools/BoxSelectInteraction.cpp b/Engine/Source/Editor/Experimental/EditorInteractiveToolsFramework/Private/EditorDragTools/BoxSelectInteraction.cpp
new file mode 100644
--- /dev/null
+++ b/Engine/Source/Editor/Experimental/EditorInteractiveToolsFramework/Private/EditorDragTools/BoxSelectInteraction.cpp
@@ -0,0 +1,506 @@
+// Copyright Epic Games, Inc. All Rights Reserved.
+
+#include "EditorDragTools/BoxSelectInteraction.h"
+#include "BaseBehaviors/ClickDragBehavior.h"
+#include "Components/PrimitiveComponent.h"
+#include "Editor/UnrealEdEngine.h"
+#include "EditorDragTools/EditorViewportClientProxy.h"
+#include "EditorModeManager.h"
+#include "EditorModes.h"
+#include "Elements/Framework/EngineElementsLibrary.h"
+#include "Elements/Framework/TypedElementRegistry.h"
+#include "Elements/Framework/TypedElementSelectionSet.h"
+#include "Elements/Interfaces/TypedElementWorldInterface.h"
+#include "Engine/LevelStreaming.h"
+#include "EngineUtils.h"
+#include "GameFramework/Volume.h"
+#include "LevelEditorViewport.h"
+#include "Model.h"
+#include "SceneView.h"
+#include "ScopedTransaction.h"
+#include "Settings/LevelEditorViewportSettings.h"
+#include "Subsystems/BrushEditingSubsystem.h"
+#include "UnrealEdGlobals.h"
+#include "ViewportInteractions/ViewportInteractionsBehaviorSource.h"
+
+namespace UE::EditorDragTools::Private
+{
+TArray<FTypedElementHandle> GetElementsIntersectingBox(
+	const AActor* Actor,
+	const FBox& InBox,
+	const IEditorViewportClientProxy* InEditorViewportClientProxy,
+	const FWorldSelectionElementArgs& SelectionArgs
+)
+{
+	if (InEditorViewportClientProxy && InEditorViewportClientProxy->IsActorVisible(Actor))
+	{
+		if (FTypedElementHandle ActorHandle = UEngineElementsLibrary::AcquireEditorActorElementHandle(Actor))
+		{
+			if (TTypedElement<ITypedElementWorldInterface> WorldElement =
+					UTypedElementRegistry::GetInstance()->GetElement<ITypedElementWorldInterface>(ActorHandle))
+			{
+				return WorldElement.GetSelectionElementsInBox(InBox, SelectionArgs);
+			}
+		}
+	}
+
+	return {};
+}
+
+void AddHoverEffect(const FViewportHoverTarget& InHoverTarget)
+{
+	FLevelEditorViewportClient::AddHoverEffect(InHoverTarget);
+	FLevelEditorViewportClient::HoveredObjects.Add(InHoverTarget);
+}
+
+void RemoveHoverEffect(const FViewportHoverTarget& InHoverTarget)
+{
+	FSetElementId Id = FLevelEditorViewportClient::HoveredObjects.FindId(InHoverTarget);
+	if (Id.IsValidId())
+	{
+		FLevelEditorViewportClient::RemoveHoverEffect(InHoverTarget);
+		FLevelEditorViewportClient::HoveredObjects.Remove(Id);
+	}
+}
+
+} // namespace UE::EditorDragTools::Private
+
+UBoxSelectInteraction::UBoxSelectInteraction()
+{
+	ToolType = UE::Editor::ViewportInteractions::BoxSelect;
+	InteractionName = TEXT("Box Select");
+
+	if (UViewportClickDragBehavior* ClickDragBehavior = ClickDragInputBehavior.Get())
+	{
+		ClickDragBehavior->SetDefaultPriority(UE::Editor::DragTools::BOX_SELECT_PRIORITY);
+		ClickDragBehavior->EnableButton(EKeys::LeftMouseButton);
+	}
+}
+
+FInputRayHit UBoxSelectInteraction::CanBeginClickDragSequence(const FInputDeviceRay& InPressPos)
+{
+	if (FEditorViewportClient* const EditorViewportClient = GetEditorViewportClient())
+	{
+		if (IsCurrentModeSupported())
+		{
+			if (EditorViewportClient->IsOrtho())
+			{
+				return FInputRayHit(TNumericLimits<float>::Max());
+			}
+		}
+	}
+
+	return FInputRayHit();
+}
+
+void UBoxSelectInteraction::OnDragStart(const FInputDeviceRay& InPressPos)
+{
+	if (UViewportInteractionsBehaviorSource* BehaviorSource = GetViewportInteractionsBehaviorSource())
+	{
+		BehaviorSource->SetMouseCursorOverride(EMouseCursor::Default);
+	}
+
+	FEditorViewportClient* const EditorViewportClient = GetEditorViewportClient();
+	if (!EditorViewportClient)
+	{
+		return;
+	}
+
+	FViewport* Viewport = EditorViewportClient->Viewport;
+	if (!Viewport)
+	{
+		return;
+	}
+
+	OnActivateTool().Broadcast();
+	bIsDragging = true;
+
+	FIntPoint MousePos;
+	Viewport->GetMousePos(MousePos);
+
+	Start = FVector(MousePos);
+	End = Start;
+
+	FLevelEditorViewportClient::ClearHoverFromObjects();
+
+	// Create a list of bsp models to check for intersection with the box
+	ModelsToCheck.Reset();
+	// Do not select BSP if its not visible
+	if (EditorViewportClient->EngineShowFlags.BSP)
+	{
+		UWorld* World = EditorViewportClient->GetWorld();
+		check(World);
+		// Add the persistent level always
+		ModelsToCheck.Add(World->PersistentLevel->Model);
+		// Add all streaming level models
+		for (ULevelStreaming* StreamingLevel : World->GetStreamingLevels())
+		{
+			// Only add streaming level models if the level is visible
+			if (StreamingLevel && StreamingLevel->GetShouldBeVisibleInEditor())
+			{
+				if (ULevel* Level = StreamingLevel->GetLoadedLevel())
+				{
+					ModelsToCheck.Add(Level->Model);
+				}
+			}
+		}
+	}
+}
+
+void UBoxSelectInteraction::OnDrag(const FInputDeviceRay& InDragPos)
+{
+	if (!EditorViewportClientProxy)
+	{
+		return;
+	}
+
+	FEditorViewportClient* const EditorViewportClient = GetEditorViewportClient();
+	if (!EditorViewportClient)
+	{
+		return;
+	}
+
+	FIntPoint MousePos;
+	EditorViewportClient->Viewport->GetMousePos(MousePos);
+
+	End = FVector(MousePos);
+
+	const bool bUseHoverFeedback = GEditor && GetDefault<ULevelEditorViewportSettings>()->bEnableViewportHoverFeedback;
+
+	if (bUseHoverFeedback)
+	{
+		const bool bStrictDragSelection = GetDefault<ULevelEditorViewportSettings>()->bStrictBoxSelection;
+
+		UTypedElementSelectionSet* SelectionSet = ModeTools->GetEditorSelectionSet();
+
+		UBrushEditingSubsystem* BrushSubsystem = GEditor->GetEditorSubsystem<UBrushEditingSubsystem>();
+		const bool bGeometryMode = BrushSubsystem ? BrushSubsystem->IsGeometryEditorModeActive() : false;
+
+		FWorldSelectionElementArgs SelectionArgs{ SelectionSet,
+												  ETypedElementSelectionMethod::Primary,
+												  FTypedElementSelectionOptions(),
+												  &(EditorViewportClient->EngineShowFlags),
+												  bStrictDragSelection,
+												  bGeometryMode };
+
+		// If we are using over feedback calculate a new box from the one being dragged
+		FBox SelBBox;
+		CalculateBox(SelBBox);
+
+		// Check every actor to see if it intersects the frustum created by the box
+		// If it does, the actor will be selected and should be given a hover cue
+		UWorld* IteratorWorld = GWorld;
+		for (FActorIterator It(IteratorWorld); It; ++It)
+		{
+			AActor& Actor = **It;
+			const bool bActorHitByBox = !UE::EditorDragTools::Private::GetElementsIntersectingBox(
+											 &Actor, SelBBox, EditorViewportClientProxy, SelectionArgs
+			)
+											 .IsEmpty();
+
+			if (bActorHitByBox)
+			{
+				// Apply a hover effect to any actor that will be selected
+				AddHoverEffect(Actor);
+			}
+			else
+			{
+				// Remove any hover effect on this actor as it no longer will be selected by the current box
+				RemoveHoverEffect(Actor);
+			}
+		}
+
+		// Check each model to see if it will be selected
+		for (int32 ModelIndex = 0; ModelIndex < ModelsToCheck.Num(); ++ModelIndex)
+		{
+			UModel& Model = *ModelsToCheck[ModelIndex];
+			for (int32 NodeIndex = 0; NodeIndex < Model.Nodes.Num(); NodeIndex++)
+			{
+				if (IntersectsBox(Model, NodeIndex, SelBBox, bStrictDragSelection))
+				{
+					// Apply a hover effect to any bsp surface that will be selected
+					AddHoverEffect(Model, Model.Nodes[NodeIndex].iSurf);
+				}
+				else
+				{
+					// Remove any hover effect on this bsp surface as it no longer will be selected by the current box
+					RemoveHoverEffect(Model, Model.Nodes[NodeIndex].iSurf);
+				}
+			}
+		}
+	}
+}
+
+void UBoxSelectInteraction::OnDragEnd(const FInputDeviceRay& InReleasePos)
+{
+	if (EditorViewportClientProxy)
+	{
+		if (FEditorViewportClient* const EditorViewportClient = GetEditorViewportClient())
+		{
+			// Workaround to use while ITF and legacy input coexist.
+			// Editor Viewport Client processes ITF input before anything else.
+			// This means that, while ITF and legacy input coexist, simple clicks will be "stolen" by ITF Box Select, which
+			// is enabled by a simple Left Mouse Button click, with no modifiers. Because IE_Released events are used as
+			// "Clicks" by the Viewport Client only while tracking, but tracking will be false after this OnClickRelease,
+			// the Viewport Client will never use this IE_Release to trigger a click.
+			// This workaround checks the distance traveled during this Drag, and if it's close to zero, forwards a IE_Released
+			// event for left mouse button to the Viewport Client, bypassing the Internal_InputKey method.
+			{
+				UE::Math::TVector<double> Distance = End - Start;
+				const bool bForwardAsSimpleClick = FMath::IsNearlyZero(Distance.Length());
+
+				if (bForwardAsSimpleClick)
+				{
+					FSceneViewFamilyContext ViewFamily(
+						FSceneViewFamily::ConstructionValues(
+							EditorViewportClient->Viewport, EditorViewportClient->GetScene(), EditorViewportClient->EngineShowFlags
+						)
+							.SetRealtimeUpdate(EditorViewportClient->IsRealtime())
+					);
+					FSceneView* View = EditorViewportClient->CalcSceneView(&ViewFamily);
+
+					FIntPoint MousePos;
+					EditorViewportClient->Viewport->GetMousePos(MousePos);
+
+					HHitProxy* const Proxy = EditorViewportClient->Viewport->GetHitProxy(MousePos.X, MousePos.Y);
+
+					EditorViewportClient->ProcessClick(
+						*View, Proxy, EKeys::LeftMouseButton, IE_Released, MousePos.X, MousePos.Y
+					);
+
+					return;
+				}
+			}
+
+			UBrushEditingSubsystem* BrushSubsystem = GEditor->GetEditorSubsystem<UBrushEditingSubsystem>();
+			const bool bGeometryMode = BrushSubsystem ? BrushSubsystem->IsGeometryEditorModeActive() : false;
+
+			FScopedTransaction Transaction(NSLOCTEXT("ActorFrustumSelect", "MarqueeSelectTransaction", "Marquee Select"));
+
+			bool bShouldSelect = true;
+			FBox SelBBox;
+			CalculateBox(SelBBox);
+
+			// TODO: with ITF Gizmos indirect interactions, this is probably not going to be triggered
+			const bool bControlDown = FInputDeviceState::IsCtrlKeyDown(InputState);
+			const bool bShiftDown = FInputDeviceState::IsShiftKeyDown(InputState);
+
+			if (bControlDown)
+			{
+				// If control is down remove from selection
+				bShouldSelect = false;
+			}
+			else if (!bShiftDown)
+			{
+				// If the user is selecting, but isn't holding down SHIFT, give modes a chance to clear selection
+				ModeTools->SelectNone();
+			}
+
+			constexpr bool bSelect = true;
+			// Let the editor mode try to handle the box selection.
+			const bool bEditorModeHandledBoxSelection = ModeTools->BoxSelect(SelBBox, bSelect);
+
+			// Let the component visualizers try to handle the selection.
+			const bool bComponentVisHandledSelection = !bEditorModeHandledBoxSelection
+													&& GUnrealEd->ComponentVisManager.HandleBoxSelect(
+														SelBBox, EditorViewportClient, EditorViewportClient->Viewport
+													);
+
+			// If the edit mode didn't handle the selection, try normal actor box selection.
+			if (!bEditorModeHandledBoxSelection && !bComponentVisHandledSelection)
+			{
+				const bool bMustEncompassEntireElement = IsWindowSelection();
+
+				UTypedElementSelectionSet* SelectionSet = ModeTools->GetEditorSelectionSet();
+				FTypedElementSelectionOptions ElementSelectionOption;
+				if (!bControlDown && !bShiftDown)
+				{
+					// If the user is selecting, but isn't holding down SHIFT, remove all current selections from the selection set
+					SelectionSet->ClearSelection(ElementSelectionOption);
+				}
+
+				FWorldSelectionElementArgs SelectionArgs{
+					SelectionSet,           ETypedElementSelectionMethod::Primary,
+					ElementSelectionOption, &(EditorViewportClient->EngineShowFlags),
+					bMustEncompassEntireElement,   bGeometryMode
+				};
+
+				// Select all element that are within the selection box area.  Be aware that certain modes do special processing below.
+				bool bSelectionChanged = false;
+				UWorld* IteratorWorld = GWorld;
+				const TArray<FName>& HiddenLayers = EditorViewportClientProxy->GetHiddenLayers();
+				TArray<FTypedElementHandle> Handles;
+
+				for (FActorIterator It(IteratorWorld); It; ++It)
+				{
+					AActor* Actor = *It;
+
+					bool bActorIsVisible = true;
+					for (auto Layer : Actor->Layers)
+					{
+						// Check the actor isn't in one of the layers hidden from this viewport.
+						if (HiddenLayers.Contains(Layer))
+						{
+							bActorIsVisible = false;
+							break;
+						}
+					}
+
+					// Select the actor or its child elements
+					if (bActorIsVisible)
+					{
+						Handles.Append(
+							UE::EditorDragTools::Private::GetElementsIntersectingBox(
+								Actor, SelBBox, EditorViewportClientProxy, SelectionArgs
+							)
+						);
+					}
+				}
+
+				if (bShouldSelect)
+				{
+					SelectionSet->SelectElements(Handles, ElementSelectionOption);
+				}
+				else
+				{
+					SelectionSet->DeselectElements(Handles, ElementSelectionOption);
+				}
+
+				// Check every model to see if its BSP surfaces should be selected
+				for (int32 ModelIndex = 0; ModelIndex < ModelsToCheck.Num(); ++ModelIndex)
+				{
+					UModel& Model = *ModelsToCheck[ModelIndex];
+					// Check every node in the model
+					for (int32 NodeIndex = 0; NodeIndex < Model.Nodes.Num(); NodeIndex++)
+					{
+						if (IntersectsBox(Model, NodeIndex, SelBBox, bMustEncompassEntireElement))
+						{
+							// If the node intersected the frustum select the corresponding surface
+							GEditor->SelectBSPSurf(&Model, Model.Nodes[NodeIndex].iSurf, bShouldSelect, false);
+							bSelectionChanged = true;
+						}
+					}
+				}
+
+				if (bSelectionChanged)
+				{
+					// If any selections were made.  Notify that now.
+					GEditor->NoteSelectionChange();
+				}
+			}
+
+			// Clear any hovered objects that might have been created while dragging
+			FLevelEditorViewportClient::ClearHoverFromObjects();
+		}
+	}
+
+	UDragToolInteraction::OnDragEnd(InReleasePos);
+}
+
+void UBoxSelectInteraction::OnTerminateDragSequence()
+{
+	UDragToolInteraction::OnTerminateDragSequence();
+}
+
+TArray<FEditorModeID> UBoxSelectInteraction::GetUnsupportedModes()
+{
+	const TArray<FEditorModeID> UnsupportedModes = { FBuiltinEditorModes::EM_Landscape, FBuiltinEditorModes::EM_Foliage };
+	return UnsupportedModes;
+}
+
+void UBoxSelectInteraction::CalculateBox(FBox& OutBox) const
+{
+	FEditorViewportClient* const EditorViewportClient = GetEditorViewportClient();
+	if (!EditorViewportClient)
+	{
+		return;
+	}
+
+	FSceneViewFamilyContext ViewFamily(
+		FSceneViewFamily::ConstructionValues(
+			EditorViewportClient->Viewport, EditorViewportClient->GetScene(), EditorViewportClient->EngineShowFlags
+		)
+			.SetRealtimeUpdate(EditorViewportClient->IsRealtime())
+	);
+
+	FSceneView* View = EditorViewportClient->CalcSceneView(&ViewFamily);
+
+	FVector3f StartFloat{ Start };
+	FVector3f EndFloat{ End };
+
+	FVector4 StartScreenPos = View->PixelToScreen(StartFloat.X, StartFloat.Y, 0);
+	FVector4 EndScreenPos = View->PixelToScreen(EndFloat.X, EndFloat.Y, 0);
+
+	FVector TransformedStart = View->ScreenToWorld(View->PixelToScreen(StartFloat.X, StartFloat.Y, 0.5f));
+	FVector TransformedEnd = View->ScreenToWorld(View->PixelToScreen(EndFloat.X, EndFloat.Y, 0.5f));
+
+	// Create a bounding box based on the start/end points (normalizes the points).
+	OutBox.Init();
+	OutBox += TransformedStart;
+	OutBox += TransformedEnd;
+
+	switch (EditorViewportClient->ViewportType)
+	{
+	case LVT_OrthoXY:
+	case LVT_OrthoNegativeXY:
+		OutBox.Min.Z = -WORLD_MAX;
+		OutBox.Max.Z = WORLD_MAX;
+		break;
+	case LVT_OrthoXZ:
+	case LVT_OrthoNegativeXZ:
+		OutBox.Min.Y = -WORLD_MAX;
+		OutBox.Max.Y = WORLD_MAX;
+		break;
+	case LVT_OrthoYZ:
+	case LVT_OrthoNegativeYZ:
+		OutBox.Min.X = -WORLD_MAX;
+		OutBox.Max.X = WORLD_MAX;
+		break;
+	case LVT_OrthoFreelook:
+	case LVT_Perspective:
+		break;
+	}
+}
+
+bool UBoxSelectInteraction::IntersectsBox(const UModel& InModel, int32 InNodeIndex, const FBox& InBox, bool bInUseStrictSelection)
+{
+	FBox NodeBB;
+	InModel.GetNodeBoundingBox(InModel.Nodes[InNodeIndex], NodeBB);
+
+	bool bFullyContained = false;
+	bool bIntersects = false;
+	if (!bInUseStrictSelection)
+	{
+		bIntersects = InBox.Intersect(NodeBB);
+	}
+	else
+	{
+		bIntersects = InBox.IsInside(NodeBB.Max) && InBox.IsInside(NodeBB.Min);
+	}
+
+	return bIntersects;
+}
+
+void UBoxSelectInteraction::AddHoverEffect(AActor& InActor)
+{
+	const FViewportHoverTarget HoverTarget(&InActor);
+	UE::EditorDragTools::Private::AddHoverEffect(HoverTarget);
+}
+
+void UBoxSelectInteraction::AddHoverEffect(UModel& InModel, int32 InSurfIndex)
+{
+	FViewportHoverTarget HoverTarget(&InModel, InSurfIndex);
+	UE::EditorDragTools::Private::AddHoverEffect(HoverTarget);
+}
+
+void UBoxSelectInteraction::RemoveHoverEffect(AActor& InActor)
+{
+	FViewportHoverTarget HoverTarget(&InActor);
+	UE::EditorDragTools::Private::RemoveHoverEffect(HoverTarget);
+}
+
+void UBoxSelectInteraction::RemoveHoverEffect(UModel& InModel, int32 InSurfIndex)
+{
+	FViewportHoverTarget HoverTarget(&InModel, InSurfIndex);
+	UE::EditorDragTools::Private::RemoveHoverEffect(HoverTarget);
+}
diff --git a/Engine/Source/Editor/Experimental/EditorInteractiveToolsFramework/Private/EditorDragTools/DragToolInteraction.cpp b/Engine/Source/Editor/Experimental/EditorInteractiveToolsFramework/Private/EditorDragTools/DragToolInteraction.cpp
new file mode 100644
--- /dev/null
+++ b/Engine/Source/Editor/Experimental/EditorInteractiveToolsFramework/Private/EditorDragTools/DragToolInteraction.cpp
@@ -0,0 +1,187 @@
+﻿// Copyright Epic Games, Inc. All Rights Reserved.
+
+#include "EditorDragTools/DragToolInteraction.h"
+#include "ContextObjectStore.h"
+#include "Editor.h"
+#include "EditorDragTools/EditorViewportClientProxy.h"
+#include "EditorModeManager.h"
+#include "EditorViewportClient.h"
+#include "InputCoreTypes.h"
+#include "Tools/EdModeInteractiveToolsContext.h"
+#include "ViewportInteractions/ViewportInteractionsBehaviorSource.h"
+
+namespace UE::Editor::DragTools
+{
+
+// CVar initializer
+static int32 UseITFTools = 0;
+static FAutoConsoleVariableRef CVarEnableITFTools(
+	TEXT("DragTools.EnableITFTools"),
+	UseITFTools,
+	TEXT("Is the ITF version of Drag Tools enabled?"),
+	FConsoleVariableDelegate::CreateLambda(
+		[](const IConsoleVariable* InVariable)
+		{
+			if (UseITFTools)
+			{
+				OnEditorDragToolsActivated().Broadcast();
+			}
+			else
+			{
+				OnEditorDragToolsDeactivated().Broadcast();
+			}
+		}
+	)
+);
+
+// Used to group Viewport Change behaviors inside the BehaviorSet and remove them if needed
+static const FString ViewportChangeBehaviorGroup = TEXT("ViewportChange");
+
+// CVar initializer
+static int32 UseViewportChangeTool = 0;
+static FAutoConsoleVariableRef CVarEnableViewportChangeTool(
+	TEXT("DragTools.EnableViewportChangeTool"),
+	UseViewportChangeTool,
+	TEXT("Is the ITF version of the viewport change tool enabled?"),
+	FConsoleVariableDelegate::CreateLambda(
+		[](const IConsoleVariable* InVariable)
+		{
+			if (UseViewportChangeTool)
+			{
+				OnViewportChangeToolActivated().Broadcast();
+			}
+			else
+			{
+				OnViewportChangeToolDeactivated().Broadcast();
+			}
+		}
+	)
+);
+
+bool UseEditorDragTools()
+{
+	return UseITFTools == 1;
+}
+
+FOnEditorDragToolsToggleDelegate& OnEditorDragToolsActivated()
+{
+	static FOnEditorDragToolsToggleDelegate OnDragToolsActivated;
+	return OnDragToolsActivated;
+}
+
+FOnEditorDragToolsToggleDelegate& OnEditorDragToolsDeactivated()
+{
+	static FOnEditorDragToolsToggleDelegate OnDragToolsDeactivated;
+	return OnDragToolsDeactivated;
+}
+
+
+bool IsViewportChangeToolEnabled()
+{
+	return UE::Editor::DragTools::UseViewportChangeTool == 1;
+}
+
+FOnOnViewportChangeToolToggleDelegate& OnViewportChangeToolActivated()
+{
+	static FOnOnViewportChangeToolToggleDelegate OnViewportChangedToolActivated;
+	return OnViewportChangedToolActivated;
+}
+
+FOnOnViewportChangeToolToggleDelegate& OnViewportChangeToolDeactivated()
+{
+	static FOnOnViewportChangeToolToggleDelegate OnViewportChangedToolDeactivated;
+	return OnViewportChangedToolDeactivated;
+}
+
+} // namespace UE::Editor::DragTools
+
+///////////////////////////////////////////////////////////////////////////////////////////////////////////////////////
+//
+// FEditorDragTools
+//
+///////////////////////////////////////////////////////////////////////////////////////////////////////////////////////
+
+UDragToolInteraction::UDragToolInteraction()
+	: bConvertDelta(false)
+	, ModeTools(nullptr)
+	, Start()
+	, End()
+	, bUseSnapping(false)
+	, bIsDragging(false)
+	, EditorViewportClientProxy(nullptr)
+{
+	GroupName = UE::Editor::ViewportInteractions::DragToolsGroupName;
+}
+
+void UDragToolInteraction::Initialize(UViewportInteractionsBehaviorSource* const InViewportInteractionsBehaviorSource)
+{
+	UViewportInteraction::Initialize(InViewportInteractionsBehaviorSource);
+
+	if (FEditorViewportClient* EditorViewportClient = GetEditorViewportClient())
+	{
+		ModeTools = EditorViewportClient->GetModeTools();
+		EditorViewportClientProxy = IEditorViewportClientProxy::CreateViewportClientProxy(ModeTools->GetInteractiveToolsContext());
+	}
+}
+
+void UDragToolInteraction::OnStateUpdated(const FInputDeviceState& InInputDeviceState)
+{
+	InputState = InInputDeviceState;
+}
+
+void UDragToolInteraction::OnDragEnd(const FInputDeviceRay& InReleasePos)
+{
+	if (UViewportInteractionsBehaviorSource* BehaviorSource = GetViewportInteractionsBehaviorSource())
+	{
+		BehaviorSource->ClearMouseCursorOverride();
+	}
+
+	Start = End = FVector::ZeroVector;
+	bIsDragging = false;
+
+	OnTerminateDragSequence();
+}
+
+void UDragToolInteraction::OnForceEndCapture()
+{
+	OnTerminateDragSequence();
+}
+
+void UDragToolInteraction::OnTerminateDragSequence()
+{
+	if (UViewportInteractionsBehaviorSource* BehaviorSource = GetViewportInteractionsBehaviorSource())
+	{
+		BehaviorSource->ClearMouseCursorOverride();
+	}
+
+	Start = End = FVector::ZeroVector;
+	bIsDragging = false;
+
+	InputState.bShiftKeyDown = false;
+	InputState.bCtrlKeyDown = false;
+	InputState.bAltKeyDown = false;
+
+	// Signal that this tool is no longer active
+	OnDeactivateTool().Broadcast();
+
+	if (FEditorViewportClient* EditorViewportClient = GetEditorViewportClient())
+	{
+		EditorViewportClient->Invalidate(true, false);
+	}
+}
+
+bool UDragToolInteraction::IsCurrentModeSupported()
+{
+	if (ModeTools)
+	{
+		for (const FEditorModeID UnsupportedMode : GetUnsupportedModes())
+		{
+			if (!ModeTools->EnsureNotInMode(UnsupportedMode))
+			{
+				return false;
+			}
+		}
+	}
+
+	return true;
+}
diff --git a/Engine/Source/Editor/Experimental/EditorInteractiveToolsFramework/Private/EditorDragTools/DragToolsBehaviourSource.cpp b/Engine/Source/Editor/Experimental/EditorInteractiveToolsFramework/Private/EditorDragTools/DragToolsBehaviourSource.cpp
deleted file mode 100644
diff --git a/Engine/Source/Editor/Experimental/EditorInteractiveToolsFramework/Private/EditorDragTools/DuplicateDragSelectionInteraction.cpp b/Engine/Source/Editor/Experimental/EditorInteractiveToolsFramework/Private/EditorDragTools/DuplicateDragSelectionInteraction.cpp
new file mode 100644
--- /dev/null
+++ b/Engine/Source/Editor/Experimental/EditorInteractiveToolsFramework/Private/EditorDragTools/DuplicateDragSelectionInteraction.cpp
@@ -0,0 +1,232 @@
+// Copyright Epic Games, Inc. All Rights Reserved.
+
+#include "EditorDragTools/DuplicateDragSelectionInteraction.h"
+#include "ContextObjectStore.h"
+#include "EditorDragTools/DragToolInteraction.h"
+#include "EditorGizmos/EditorTransformGizmo.h"
+#include "EditorGizmos/EditorTransformGizmoUtil.h"
+#include "EditorInteractiveGizmoManager.h"
+#include "EditorModeManager.h"
+#include "Elements/Framework/TypedElementCommonActions.h"
+#include "ILevelEditor.h"
+#include "InteractiveGizmo.h"
+#include "InteractiveGizmoManager.h"
+#include "InteractiveToolManager.h"
+#include "LevelEditorViewport.h"
+#include "ScopedTransaction.h"
+#include "Tools/AssetEditorContextInterface.h"
+#include "ViewportInteractions/ViewportInteractionsBehaviorSource.h"
+
+#define LOCTEXT_NAMESPACE "EditorDuplicateDragSelection"
+
+UDuplicateDragSelectionInteraction::UDuplicateDragSelectionInteraction()
+{
+	ToolType = UE::Editor::ViewportInteractions::DuplicateDrag;
+	InteractionName = TEXT("Duplicate Drag");
+	GroupName = UE::Editor::ViewportInteractions::DragToolsGroupName;
+}
+
+void UDuplicateDragSelectionInteraction::RegisterProxyDelegates()
+{
+	if (TransformGizmo)
+	{
+		// Register Proxy Delegates
+		if (TObjectPtr<UTransformProxy> TransformProxy = TransformGizmo->ActiveTarget)
+		{
+			TransformProxy->OnBeginTransformEdit.AddUObject(
+				this, &UDuplicateDragSelectionInteraction::OnGizmoMovementBegin
+			);
+
+			TransformProxy->OnEndTransformEdit.AddUObject(
+				this, &UDuplicateDragSelectionInteraction::OnGizmoMovementEnd
+			);
+		}
+	}
+}
+
+void UDuplicateDragSelectionInteraction::Reset()
+{
+	if (TransformGizmo)
+	{
+		// Unregister Proxy delegates
+		if (TObjectPtr<UTransformProxy> TransformProxy = TransformGizmo->ActiveTarget)
+		{
+			TransformProxy->OnBeginTransformEdit.RemoveAll(this);
+			TransformProxy->OnEndTransformEdit.RemoveAll(this);
+			TransformProxy->OnTransformChanged.RemoveAll(this);
+		}
+
+		TransformGizmo = nullptr;
+	}
+}
+
+void UDuplicateDragSelectionInteraction::OnGizmoCreatedDelegate(UTransformGizmo* InTransformGizmo)
+{
+	if (!TransformGizmo)
+	{
+		TransformGizmo = Cast<UEditorTransformGizmo>(InTransformGizmo);
+
+		RegisterProxyDelegates();
+	}
+}
+
+void UDuplicateDragSelectionInteraction::OnUsesNewTRSGizmosChanged(bool bInUseNewTRSGizmos)
+{
+	if (!bInUseNewTRSGizmos)
+	{
+		DuplicateTransaction.Reset();
+		Reset();
+	}
+}
+
+void UDuplicateDragSelectionInteraction::Initialize(UViewportInteractionsBehaviorSource* const InViewportInteractionsBehaviorSource)
+{
+	Super::Initialize(InViewportInteractionsBehaviorSource);
+
+	if (const FEditorViewportClient* const EditorViewportClient = GetEditorViewportClient())
+	{
+		if (const FEditorModeTools* const ModeTools = EditorViewportClient->GetModeTools())
+		{
+			if (const UModeManagerInteractiveToolsContext* const InteractiveToolsContext = ModeTools->GetInteractiveToolsContext())
+			{
+				ToolManager = InteractiveToolsContext->ToolManager;
+			}
+		}
+	}
+
+	// Map to modifier key state change
+	GetViewportInteractionsBehaviorSource()->OnAltKeyStateChanged().AddUObject(this, &UDuplicateDragSelectionInteraction::OnModifierKeyChanged);
+
+	// New TRS Gizmos are already enabled - let's retrieve the Transform Gizmo, from which we register to drag Begin and End delegates
+	if (UEditorInteractiveGizmoManager::UsesNewTRSGizmos())
+	{
+		const TArray<UInteractiveGizmo*> Gizmos =
+		ToolManager->GetPairedGizmoManager()->FindAllGizmosOfType("EditorTransformGizmoBuilder");
+
+		if (!Gizmos.IsEmpty())
+		{
+			TransformGizmo = Cast<UEditorTransformGizmo>(Gizmos[0]);
+		}
+
+		RegisterProxyDelegates();
+	}
+
+	// If New TRS Gizmos are not enabled, we need to be able to know as soon as they are enabled
+	// This allows to retrieve the Transform Gizmo, from which we register to drag Begin and End delegates
+	if (UContextObjectStore* ContextStore = ToolManager->GetContextObjectStore())
+	{
+		if (UEditorTransformGizmoContextObject* ContextObject =
+				ContextStore->FindContext<UEditorTransformGizmoContextObject>())
+		{
+			ContextObject->OnGizmoCreatedDelegate().AddUObject(this, &UDuplicateDragSelectionInteraction::OnGizmoCreatedDelegate);
+		}
+	}
+
+	// In case new TRS Gizmos are disabled, we want to know that, so we can stop listening to Drag delegates
+	UEditorInteractiveGizmoManager::OnUsesNewTRSGizmosChangedDelegate().AddUObject(
+		this, &UDuplicateDragSelectionInteraction::OnUsesNewTRSGizmosChanged
+	);
+}
+
+UDuplicateDragSelectionInteraction::~UDuplicateDragSelectionInteraction()
+{
+	UEditorInteractiveGizmoManager::OnUsesNewTRSGizmosChangedDelegate().RemoveAll(this);
+	DuplicateTransaction.Reset();
+	Reset();
+}
+
+void UDuplicateDragSelectionInteraction::OnGizmoMovementBegin(UTransformProxy* InTransformProxy)
+{
+	if (!bGizmoIsDragged)
+	{
+		bGizmoIsDragged = true;
+
+		if (bModifierKeyIsPressed && !TransformGizmo->IsIndirectManipulation())
+		{
+			DuplicateTransaction = MakeShared<FScopedTransaction>(LOCTEXT("DuplicateAndTransformTransactionName", "Duplicate and Transform"));
+
+			// Duplicate selection and select the newly created objects
+			constexpr bool bSelectNewElements = true;
+			DuplicateSelection(bSelectNewElements);
+		}
+	}
+}
+
+void UDuplicateDragSelectionInteraction::OnGizmoMovementEnd(UTransformProxy* InTransformProxy)
+{
+	bGizmoIsDragged = false;
+	DuplicateTransaction.Reset();
+}
+
+void UDuplicateDragSelectionInteraction::DuplicateSelection(bool bInSelectNewElements) const
+{
+	if (GCurrentLevelEditingViewportClient)
+	{
+		if (TSharedPtr<ILevelEditor> LevelEditor = GCurrentLevelEditingViewportClient->ParentLevelEditor.Pin())
+		{
+			if (UTypedElementCommonActions* CommonActions = LevelEditor->GetCommonActions())
+			{
+				FToolBuilderState State;
+				if (IAssetEditorContextInterface* AssetEditorContext =
+						ToolManager->GetContextObjectStore()->FindContext<IAssetEditorContextInterface>())
+				{
+					if (UTypedElementSelectionSet* MutableSelectionSet = AssetEditorContext->GetMutableSelectionSet())
+					{
+						const TArray<FTypedElementHandle> DuplicatedElements = CommonActions->DuplicateSelectedElements(
+							MutableSelectionSet, LevelEditor->GetWorld(), FVector::ZeroVector
+						);
+						if (DuplicatedElements.Num() > 0)
+						{
+							// (assuming chosen modifier key is Alt)
+							// Alt Duplicate has 2 modes:
+							// 1 - with Alt pressed, dragging will create and select a copy, then drag it
+							// 2 - while dragging, hitting Alt will crate a copy, NOT selecting it
+
+							// Select newly created elements: Alt + Drag only
+							if (bInSelectNewElements)
+							{
+								// Select the newly created elements
+								MutableSelectionSet->SetSelection(DuplicatedElements, FTypedElementSelectionOptions());
+								MutableSelectionSet->NotifyPendingChanges();
+							}
+
+							// Notify the global mode tools, the selection set should be identical to the new actors at this point
+							TArray<AActor*> SelectedActors = MutableSelectionSet->GetSelectedObjects<AActor>();
+							constexpr bool bDidOffsetDuplicate = false;
+							GLevelEditorModeTools().ActorsDuplicatedNotify(SelectedActors, SelectedActors, bDidOffsetDuplicate);
+						}
+
+						// Invalidate all viewports, so the new gizmo is rendered in each one
+						GEditor->RedrawLevelEditingViewports();
+					}
+				}
+			}
+		}
+	}
+}
+
+void UDuplicateDragSelectionInteraction::OnModifierKeyChanged(bool bInIsDown)
+{
+	if (!IsEnabled())
+	{
+		return;
+	}
+
+	// Nothing to do
+	if (bModifierKeyIsPressed == bInIsDown)
+	{
+		return;
+	}
+
+	bModifierKeyIsPressed = bInIsDown;
+
+	// If already dragging, let's duplicate without selecting new objects
+	// This will "drop" a copy of selection matching current drag transform
+	if (bGizmoIsDragged && bModifierKeyIsPressed)
+	{
+		constexpr bool bSelectNewElements = false;
+		DuplicateSelection(bSelectNewElements);
+	}
+}
+
+#undef LOCTEXT_NAMESPACE
diff --git a/Engine/Source/Editor/Experimental/EditorInteractiveToolsFramework/Private/EditorDragTools/EditorBoxSelect.cpp b/Engine/Source/Editor/Experimental/EditorInteractiveToolsFramework/Private/EditorDragTools/EditorBoxSelect.cpp
deleted file mode 100644
diff --git a/Engine/Source/Editor/Experimental/EditorInteractiveToolsFramework/Private/EditorDragTools/EditorDragToolBehaviorTarget.cpp b/Engine/Source/Editor/Experimental/EditorInteractiveToolsFramework/Private/EditorDragTools/EditorDragToolBehaviorTarget.cpp
deleted file mode 100644
diff --git a/Engine/Source/Editor/Experimental/EditorInteractiveToolsFramework/Private/EditorDragTools/EditorDuplicateDragSelection.cpp b/Engine/Source/Editor/Experimental/EditorInteractiveToolsFramework/Private/EditorDragTools/EditorDuplicateDragSelection.cpp
deleted file mode 100644
diff --git a/Engine/Source/Editor/Experimental/EditorInteractiveToolsFramework/Private/EditorDragTools/EditorFrustumSelect.cpp b/Engine/Source/Editor/Experimental/EditorInteractiveToolsFramework/Private/EditorDragTools/EditorFrustumSelect.cpp
deleted file mode 100644
diff --git a/Engine/Source/Editor/Experimental/EditorInteractiveToolsFramework/Private/EditorDragTools/EditorMarqueeSelect.cpp b/Engine/Source/Editor/Experimental/EditorInteractiveToolsFramework/Private/EditorDragTools/EditorMarqueeSelect.cpp
deleted file mode 100644
diff --git a/Engine/Source/Editor/Experimental/EditorInteractiveToolsFramework/Private/EditorDragTools/EditorMeasureTool.cpp b/Engine/Source/Editor/Experimental/EditorInteractiveToolsFramework/Private/EditorDragTools/EditorMeasureTool.cpp
deleted file mode 100644
diff --git a/Engine/Source/Editor/Experimental/EditorInteractiveToolsFramework/Private/EditorDragTools/EditorMoveCameraWithObject.cpp b/Engine/Source/Editor/Experimental/EditorInteractiveToolsFramework/Private/EditorDragTools/EditorMoveCameraWithObject.cpp
deleted file mode 100644
diff --git a/Engine/Source/Editor/Experimental/EditorInteractiveToolsFramework/Private/EditorDragTools/EditorViewportChange.cpp b/Engine/Source/Editor/Experimental/EditorInteractiveToolsFramework/Private/EditorDragTools/EditorViewportChange.cpp
deleted file mode 100644
diff --git a/Engine/Source/Editor/Experimental/EditorInteractiveToolsFramework/Private/EditorDragTools/FrustumSelectInteraction.cpp b/Engine/Source/Editor/Experimental/EditorInteractiveToolsFramework/Private/EditorDragTools/FrustumSelectInteraction.cpp
new file mode 100644
--- /dev/null
+++ b/Engine/Source/Editor/Experimental/EditorInteractiveToolsFramework/Private/EditorDragTools/FrustumSelectInteraction.cpp
@@ -0,0 +1,450 @@
+// Copyright Epic Games, Inc. All Rights Reserved.
+
+#include "EditorDragTools/FrustumSelectInteraction.h"
+#include "BaseBehaviors/ClickDragBehavior.h"
+#include "BaseBehaviors/KeyAsModifierInputBehavior.h"
+#include "Components/PrimitiveComponent.h"
+#include "Editor/UnrealEdEngine.h"
+#include "EditorDragTools/EditorViewportClientProxy.h"
+#include "EditorModeManager.h"
+#include "EditorModes.h"
+#include "Elements/Framework/EngineElementsLibrary.h"
+#include "Elements/Framework/TypedElementRegistry.h"
+#include "Elements/Framework/TypedElementSelectionSet.h"
+#include "Elements/Interfaces/TypedElementWorldInterface.h"
+#include "EngineUtils.h"
+#include "GameFramework/Volume.h"
+#include "LevelEditorViewport.h"
+#include "Model.h"
+#include "SceneView.h"
+#include "ScopedTransaction.h"
+#include "Settings/LevelEditorViewportSettings.h"
+#include "SnappingUtils.h"
+#include "Subsystems/BrushEditingSubsystem.h"
+#include "UnrealEdGlobals.h"
+#include "Tools/EdModeInteractiveToolsContext.h"
+#include "ViewportInteractions/ViewportInteractionsBehaviorSource.h"
+
+namespace UE::EditorDragTools
+{
+
+namespace Private
+{
+	TArray<FTypedElementHandle> GetElementsIntersectingFrustum(
+		const AActor* Actor,
+		const FConvexVolume& InFrustum,
+		const IEditorViewportClientProxy* InEditorViewportClientProxy,
+		const FWorldSelectionElementArgs& SelectionArgs
+	)
+	{
+		if (InEditorViewportClientProxy && InEditorViewportClientProxy->IsActorVisible(Actor))
+		{
+			if (FTypedElementHandle ActorHandle = UEngineElementsLibrary::AcquireEditorActorElementHandle(Actor))
+			{
+				if (TTypedElement<ITypedElementWorldInterface> WorldElement =
+						UTypedElementRegistry::GetInstance()->GetElement<ITypedElementWorldInterface>(ActorHandle))
+				{
+					return WorldElement.GetSelectionElementsInConvexVolume(InFrustum, SelectionArgs);
+				}
+			}
+		}
+
+		return {};
+	}
+}
+
+FViewportFrustum::FViewportFrustum(const FEditorViewportClient& InViewportClient, const FSceneView& InView,
+	const FVector& InStart, const FVector& InEnd, const bool InUseBoxFrustum)
+	: ViewportClient(InViewportClient)
+	, View(InView)
+	, Start(InStart)
+	, End(InEnd) 
+	, bUseBoxFrustum(InUseBoxFrustum)
+{}
+
+void FViewportFrustum::Calculate(FConvexVolume& OutFrustum) const
+{
+	if (bUseBoxFrustum)
+	{
+		const FVector& CamPoint = ViewportClient.GetViewLocation();
+
+		// extend the 2D box of 1 pixel if needed to avoid degenerated volume 
+		const double DX = Start.X == End.X ? 0.5 : 0.0;
+		const double Left = FMath::Min(Start.X, End.X) - DX;
+		const double Right = FMath::Max(Start.X, End.X) + DX;
+		const double DY = Start.Y == End.Y ? 0.5 : 0.0;
+		const double Bottom = FMath::Min(Start.Y, End.Y) - DY;
+		const double Top = FMath::Max(Start.Y, End.Y) + DY;
+
+		// Deproject the four corners of the selection box
+		const FVector2D Point1(Left, Bottom); // Upper Left Corner
+		const FVector2D Point2(Right, Bottom); // Upper Right Corner
+		const FVector2D Point3(Right, Top); // Lower Right Corner
+		const FVector2D Point4(Left, Top); // Lower Left Corner
+
+		FVector BoxPoint1, BoxPoint2, BoxPoint3, BoxPoint4;
+		FVector WorldDir1, WorldDir2, WorldDir3, WorldDir4;
+		View.DeprojectFVector2D(Point1, BoxPoint1, WorldDir1);
+		View.DeprojectFVector2D(Point2, BoxPoint2, WorldDir2);
+		View.DeprojectFVector2D(Point3, BoxPoint3, WorldDir3);
+		View.DeprojectFVector2D(Point4, BoxPoint4, WorldDir4);
+
+		// Use the camera position and the selection box to create the bounding planes
+		auto MakePlane = [](const FVector& A, const FVector& B, const FVector& C)
+		{
+			// the default tolerance used in GetSafeNormal (UE_SMALL_NUMBER) is not enough to handle low near plane values
+			// see TPlane<T>::TPlane(TVector<T> A, TVector<T> B, TVector<T> C) 
+			static constexpr double SmallNumberCubed = UE_SMALL_NUMBER * UE_SMALL_NUMBER * UE_SMALL_NUMBER;
+			const FVector Normal = (B - A ^ C - A).GetSafeNormal(SmallNumberCubed);
+			const double W = A | Normal;
+			FPlane Plane(Normal, W);
+			return MoveTemp(Plane);
+		};
+
+		FPlane TopPlane = MakePlane(BoxPoint1, BoxPoint2, CamPoint); // Top Plane
+		FPlane RightPlane = MakePlane(BoxPoint2, BoxPoint3, CamPoint); // Right Plane
+		FPlane BottomPlane = MakePlane(BoxPoint3, BoxPoint4, CamPoint); // Bottom Plane
+		FPlane LeftPlane = MakePlane(BoxPoint4, BoxPoint1, CamPoint); // Left Plane
+
+		// Try to get all six planes to create a frustum.
+		// The frustum is built with the first four planes corresponding to the sides of the frustum.
+		OutFrustum.Planes.Empty();
+		OutFrustum.Planes.Emplace(MoveTemp(TopPlane));
+		OutFrustum.Planes.Emplace(MoveTemp(RightPlane));
+		OutFrustum.Planes.Emplace(MoveTemp(BottomPlane));
+		OutFrustum.Planes.Emplace(MoveTemp(LeftPlane));
+
+		FPlane NearPlane;
+		if ( View.ViewMatrices.GetViewProjectionMatrix().GetFrustumNearPlane(NearPlane) )
+		{
+			OutFrustum.Planes.Emplace(MoveTemp(NearPlane));
+		}
+
+		FPlane FarPlane;
+		if ( View.ViewMatrices.GetViewProjectionMatrix().GetFrustumFarPlane(FarPlane) )
+		{
+			OutFrustum.Planes.Emplace(MoveTemp(FarPlane));
+		}
+
+		OutFrustum.Init();
+	}
+	else
+	{
+		OutFrustum = View.ViewFrustum;
+		OutFrustum.Init();
+	}
+}
+
+} // namespace UE::EditorDragTools
+
+UFrustumSelectInteraction::UFrustumSelectInteraction()
+{
+	ToolType = UE::Editor::ViewportInteractions::FrustumSelect;
+	InteractionName = TEXT("Frustum Select");
+
+	if (UViewportClickDragBehavior* ClickDragBehavior = ClickDragInputBehavior.Get())
+	{
+		ClickDragBehavior->SetDefaultPriority(UE::Editor::DragTools::FRUSTUM_SELECT_PRIORITY);
+		ClickDragBehavior->EnableButton(EKeys::LeftMouseButton);
+		ClickDragBehavior->ModifierCheckFunc = [](const FInputDeviceState& DeviceState)
+		{
+			return DeviceState.bCtrlKeyDown && DeviceState.bAltKeyDown;
+		};
+	}
+}
+
+FInputRayHit UFrustumSelectInteraction::CanBeginClickDragSequence(const FInputDeviceRay& InPressPos)
+{
+	return IsCurrentModeSupported()
+			 ? FInputRayHit(TNumericLimits<float>::Max()) // bHit is true. Depth is max to lose the standard tiebreaker.
+			 : FInputRayHit();
+}
+
+void UFrustumSelectInteraction::OnDragStart(const FInputDeviceRay& InPressPos)
+{
+	if (UViewportInteractionsBehaviorSource* BehaviorSource = GetViewportInteractionsBehaviorSource())
+	{
+		BehaviorSource->SetMouseCursorOverride(EMouseCursor::Default);
+	}
+
+	// Signal that this tool is now active
+	OnActivateTool().Broadcast();
+
+	Start = InPressPos.WorldRay.Origin;
+	bIsDragging = true;
+
+	// Snap to constraints.
+	if (bUseSnapping)
+	{
+		const float GridSize = GEditor->GetGridSize();
+		const FVector GridBase(GridSize, GridSize, GridSize);
+		FSnappingUtils::SnapPointToGrid(Start, GridBase);
+	}
+	End = Start;
+
+	// Remove any active hover objects
+	FLevelEditorViewportClient::ClearHoverFromObjects();
+
+	Start = FVector(InPressPos.ScreenPosition.X, InPressPos.ScreenPosition.Y, 0);
+	End = Start;
+}
+
+void UFrustumSelectInteraction::OnDrag(const FInputDeviceRay& InDragPos)
+{
+	End = FVector(InDragPos.ScreenPosition.X, InDragPos.ScreenPosition.Y, 0);
+}
+
+void UFrustumSelectInteraction::OnDragEnd(const FInputDeviceRay& InReleasePos)
+{
+	End = FVector(InReleasePos.ScreenPosition.X, InReleasePos.ScreenPosition.Y, 0);
+
+	if (!EditorViewportClientProxy)
+	{
+		return;
+	}
+
+	FEditorViewportClient* const EditorViewportClient = GetEditorViewportClient();
+	if (!EditorViewportClient)
+	{
+		return;
+	}
+
+	FViewport* Viewport = EditorViewportClient->Viewport;
+	if (!Viewport)
+	{
+		return;
+	}
+
+	UBrushEditingSubsystem* BrushSubsystem = GEditor->GetEditorSubsystem<UBrushEditingSubsystem>();
+	const bool bGeometryMode = BrushSubsystem ? BrushSubsystem->IsGeometryEditorModeActive() : false;
+
+	FSceneViewFamilyContext ViewFamily(
+		FSceneViewFamily::ConstructionValues(Viewport, EditorViewportClient->GetScene(), EditorViewportClient->EngineShowFlags)
+	);
+	FSceneView* SceneView = EditorViewportClient->CalcSceneView(&ViewFamily);
+
+	// Generate a frustum out of the dragged box
+	FConvexVolume Frustum;
+	CalculateFrustum(SceneView, Frustum, true);
+
+	FScopedTransaction Transaction(NSLOCTEXT("ActorFrustumSelect", "MarqueeSelectTransaction", "Marquee Select"));
+
+	if (!FInputDeviceState::IsShiftKeyDown(InputState))
+	{
+		// If the user is selecting, but isn't hold down SHIFT, remove all current selections.
+		ModeTools->SelectNone();
+	}
+
+	// Does an actor have to be fully contained in the box to be selected
+	const bool bMustEncompassEntireElement = IsWindowSelection();
+
+	constexpr bool bShouldSelect = true; // bLeftMouseButtonDown
+	// Let the editor mode try to handle the selection.
+	const bool bEditorModeHandledSelection = ModeTools->FrustumSelect(Frustum, EditorViewportClient, bShouldSelect);
+
+	// Let the component visualizers try to handle the selection.
+	const bool bComponentVisHandledSelection =
+		!bEditorModeHandledSelection
+		&& GUnrealEd->ComponentVisManager.HandleFrustumSelect(Frustum, EditorViewportClient, Viewport);
+
+	if (!bEditorModeHandledSelection && !bComponentVisHandledSelection)
+	{
+		UTypedElementSelectionSet* SelectionSet = ModeTools->GetEditorSelectionSet();
+		SelectionSet->Modify();
+
+		FWorldSelectionElementArgs SeletionArgs{ SelectionSet,
+												 ETypedElementSelectionMethod::Primary,
+												 FTypedElementSelectionOptions(),
+												 &(EditorViewportClient->EngineShowFlags),
+												 bMustEncompassEntireElement,
+												 bGeometryMode };
+
+		const int32 ViewportSizeX = Viewport->GetSizeXY().X;
+		const int32 ViewportSizeY = Viewport->GetSizeXY().Y;
+
+		if (Start.X > End.X)
+		{
+			Swap(Start.X, End.X);
+		}
+
+		if (Start.Y > End.Y)
+		{
+			Swap(Start.Y, End.Y);
+		}
+
+		TArray<FTypedElementHandle> ElementsToSelect;
+		const bool bTransparentBoxSelection = GetDefault<ULevelEditorViewportSettings>()->bTransparentBoxSelection;
+		if (bTransparentBoxSelection)
+		{
+			// Get a list of frustum-culled actors
+			for (FActorIterator It(EditorViewportClient->GetWorld()); It; ++It)
+			{
+				AActor* Actor = *It;
+				ElementsToSelect.Append(
+					UE::EditorDragTools::Private::GetElementsIntersectingFrustum(
+						Actor, Frustum, EditorViewportClientProxy, SeletionArgs
+					)
+				);
+			}
+		}
+		else
+		{
+			// Extend the endpoint of the rect to get the actual line
+
+			const int32 MinX = UE::LWC::FloatToIntCastChecked<int32>(FMath::Max<double>(0.0, Start.X));
+			const int32 MinY = UE::LWC::FloatToIntCastChecked<int32>(FMath::Max<double>(0.0, Start.Y));
+			const int32 MaxX = FMath::Min(ViewportSizeX, FMath::TruncToInt32(End.X + 1.0));
+			const int32 MaxY = FMath::Min(ViewportSizeY, FMath::TruncToInt32(End.Y + 1.0));
+
+			const FIntPoint Min{ MinX, MinY };
+			const FIntPoint Max{ MaxX, MaxY };
+			const FIntRect BoxRect{ Min, Max };
+
+			// Typed Element selection
+			{
+				FTypedElementListRef ElementList = UTypedElementRegistry::GetInstance()->CreateElementList();
+				Viewport->GetElementHandlesInRect(BoxRect, ElementList);
+
+				if (bMustEncompassEntireElement)
+				{
+					ElementList->ForEachElement<ITypedElementWorldInterface>(
+						[bMustEncompassEntireElement, &Frustum, &SelectionSet, &ElementsToSelect](
+							const TTypedElement<ITypedElementWorldInterface>& InElement
+						)
+						{
+							if (InElement.IsElementInConvexVolume(Frustum, bMustEncompassEntireElement))
+							{
+								ElementsToSelect.Add(
+									SelectionSet->GetSelectionElement(InElement, ETypedElementSelectionMethod::Primary)
+								);
+							}
+
+							return true;
+						}
+					);
+				}
+				else
+				{
+					// Grab only the selectable handles (this remove the components from the selection and select the actor instead)
+					ElementList->ForEachElementHandle(
+						[&SelectionSet, &ElementsToSelect](const FTypedElementHandle& InHandle)
+						{
+							ElementsToSelect.Add(
+								SelectionSet->GetSelectionElement(InHandle, ETypedElementSelectionMethod::Primary)
+							);
+							return true;
+						}
+					);
+				}
+			}
+
+			// We need this old code to support the BSP
+			TSet<AActor*> BSPActors;
+			TSet<UModel*> HitModels;
+			Viewport->GetActorsAndModelsInHitProxy(BoxRect, BSPActors, HitModels);
+			BSPActors.Empty(HitModels.Num());
+
+			if (HitModels.Num() > 0)
+			{
+				// Check every model to see if its BSP surfaces should be selected
+				for (auto It = HitModels.CreateConstIterator(); It; ++It)
+				{
+					UModel& Model = **It;
+					// Check every node in the model
+					for (int32 NodeIndex = 0; NodeIndex < Model.Nodes.Num(); NodeIndex++)
+					{
+						if (IntersectsFrustum(Model, NodeIndex, Frustum, bMustEncompassEntireElement))
+						{
+							uint32 SurfaceIndex = Model.Nodes[NodeIndex].iSurf;
+							FBspSurf& Surf = Model.Surfs[SurfaceIndex];
+							BSPActors.Add(Surf.Actor);
+						}
+					}
+				}
+			}
+
+			if (BSPActors.Num() > 0)
+			{
+				for (auto It = BSPActors.CreateConstIterator(); It; ++It)
+				{
+					AActor* Actor = *It;
+					if (bMustEncompassEntireElement)
+					{
+						ElementsToSelect.Append(
+							UE::EditorDragTools::Private::GetElementsIntersectingFrustum(
+								Actor, Frustum, EditorViewportClientProxy, SeletionArgs
+							)
+						);
+					}
+					else
+					{
+						ElementsToSelect.Add(UEngineElementsLibrary::AcquireEditorActorElementHandle(Actor));
+					}
+				}
+			}
+		}
+
+		FTypedElementSelectionOptions ElementSelectionOption;
+
+		if (!FInputDeviceState::IsShiftKeyDown(InputState))
+		{
+			// If the user is selecting, but isn't hold down SHIFT, remove the previous selections.
+			SelectionSet->SetSelection(MoveTemp(ElementsToSelect), ElementSelectionOption);
+		}
+		else
+		{
+			SelectionSet->SelectElements(MoveTemp(ElementsToSelect), ElementSelectionOption);
+		}
+	}
+
+	// Clear any hovered objects that might have been created while dragging
+	FLevelEditorViewportClient::ClearHoverFromObjects();
+
+	UDragToolInteraction::OnDragEnd(InReleasePos);
+}
+
+void UFrustumSelectInteraction::OnTerminateDragSequence()
+{
+	UDragToolInteraction::OnTerminateDragSequence();
+
+	Start = End = FVector::ZeroVector;
+}
+
+TArray<FEditorModeID> UFrustumSelectInteraction::GetUnsupportedModes()
+{
+	const TArray<FEditorModeID> UnsupportedModes = { FBuiltinEditorModes::EM_Landscape, FBuiltinEditorModes::EM_Foliage };
+	return UnsupportedModes;
+}
+
+bool UFrustumSelectInteraction::IntersectsFrustum(
+	const UModel& InModel, int32 NodeIndex, const FConvexVolume& InFrustum, bool bUseStrictSelection
+)
+{
+	FBox NodeBB;
+	// Get a bounding box of the node being checked
+	InModel.GetNodeBoundingBox(InModel.Nodes[NodeIndex], NodeBB);
+
+	bool bFullyContained = false;
+
+	// Does the box intersect the frustum
+	bool bIntersects = InFrustum.IntersectBox(NodeBB.GetCenter(), NodeBB.GetExtent(), bFullyContained);
+
+	return bIntersects && (!bUseStrictSelection || (bUseStrictSelection && bFullyContained));
+}
+
+void UFrustumSelectInteraction::CalculateFrustum(const FSceneView* InView, FConvexVolume& OutFrustum, bool bUseBoxFrustum) const
+{
+	if (!InView || !EditorViewportClientProxy)
+	{
+		return;
+	}
+
+	FEditorViewportClient* const EditorViewportClient = GetEditorViewportClient();
+	if (!EditorViewportClient)
+	{
+		return;
+	}
+
+	UE::EditorDragTools::FViewportFrustum ViewportFrustum(*EditorViewportClient, *InView, Start, End, bUseBoxFrustum);
+	ViewportFrustum.Calculate(OutFrustum);
+}
diff --git a/Engine/Source/Editor/Experimental/EditorInteractiveToolsFramework/Private/EditorDragTools/MarqueeSelectInteraction.cpp b/Engine/Source/Editor/Experimental/EditorInteractiveToolsFramework/Private/EditorDragTools/MarqueeSelectInteraction.cpp
new file mode 100644
--- /dev/null
+++ b/Engine/Source/Editor/Experimental/EditorInteractiveToolsFramework/Private/EditorDragTools/MarqueeSelectInteraction.cpp
@@ -0,0 +1,86 @@
+﻿// Copyright Epic Games, Inc. All Rights Reserved.
+
+#include "EditorDragTools/MarqueeSelectInteraction.h"
+
+#include "BaseBehaviors/ClickDragBehavior.h"
+#include "CanvasItem.h"
+#include "CanvasTypes.h"
+#include "Settings/LevelEditorViewportSettings.h"
+
+UMarqueeSelectInteraction::UMarqueeSelectInteraction()
+{
+	UViewportClickDragBehavior* ClickDragBehavior = NewObject<UViewportClickDragBehavior>();
+	ClickDragBehavior->Initialize(this);
+
+	ClickDragInputBehavior = ClickDragBehavior;
+
+	RegisterInputBehavior(ClickDragBehavior);
+}
+
+void UMarqueeSelectInteraction::Render(const FSceneView* View, FCanvas* Canvas, EViewInteractionState InteractionState)
+{
+	if (!IsDragging() || InteractionState != EViewInteractionState::Focused)
+	{
+		return;
+	}
+
+	if (!Canvas || !View)
+	{
+		return;
+	}
+
+	if (IsWindowSelection())
+	{
+		FCanvasBoxItem BoxItem(
+			FVector2D(Start.X, Start.Y) / Canvas->GetDPIScale(),
+			FVector2D(End.X - Start.X, End.Y - Start.Y) / Canvas->GetDPIScale()
+		);
+
+		BoxItem.SetColor(FLinearColor::White);
+		Canvas->DrawItem(BoxItem);
+	}
+	else
+	{
+		FVector LocalStart = Start;
+		FVector LocalEnd = End;
+
+		// Place the starting corner in the upper-left
+		if (LocalStart.X > LocalEnd.X)
+		{
+			Swap(LocalStart.X, LocalEnd.X);
+		}
+
+		if (LocalStart.Y > LocalEnd.Y)
+		{
+			Swap(LocalStart.Y, LocalEnd.Y);
+		}
+
+		FCanvasDashedBoxItem BoxItem(
+			FVector2D(LocalStart.X, LocalStart.Y) / Canvas->GetDPIScale(),
+			FVector2D(LocalEnd.X - LocalStart.X, LocalEnd.Y - LocalStart.Y) / Canvas->GetDPIScale()
+		);
+
+		BoxItem.DashLength = 12.0f;
+		BoxItem.DashGap = 6.0f;
+		BoxItem.SetColor(FLinearColor::White);
+
+		Canvas->DrawItem(BoxItem);
+	}
+}
+
+bool UMarqueeSelectInteraction::IsWindowSelection() const
+{
+	switch (GetDefault<ULevelEditorViewportSettings>()->MarqueeSelectionMode)
+	{
+	case EMarqueeSelectionMode::Crossing:
+		return false;
+	case EMarqueeSelectionMode::Window:
+		return true;
+	case EMarqueeSelectionMode::CrossLeft:
+		return Start.X < End.X;
+	case EMarqueeSelectionMode::CrossRight:
+		return Start.X > End.X;
+	default:
+		return false;
+	}
+}
diff --git a/Engine/Source/Editor/Experimental/EditorInteractiveToolsFramework/Private/EditorDragTools/MeasureToolInteraction.cpp b/Engine/Source/Editor/Experimental/EditorInteractiveToolsFramework/Private/EditorDragTools/MeasureToolInteraction.cpp
new file mode 100644
--- /dev/null
+++ b/Engine/Source/Editor/Experimental/EditorInteractiveToolsFramework/Private/EditorDragTools/MeasureToolInteraction.cpp
@@ -0,0 +1,203 @@
+// Copyright Epic Games, Inc. All Rights Reserved.
+
+#include "EditorDragTools/MeasureToolInteraction.h"
+#include "CanvasItem.h"
+#include "CanvasTypes.h"
+#include "Editor.h"
+#include "EditorDragTools/EditorViewportClientProxy.h"
+#include "EditorViewportClient.h"
+#include "EngineGlobals.h"
+#include "SceneView.h"
+#include "Settings/EditorStyleSettings.h"
+#include "Settings/LevelEditorViewportSettings.h"
+#include "SnappingUtils.h"
+#include "BaseBehaviors/ClickDragBehavior.h"
+#include "ViewportInteractions/ViewportInteractionsBehaviorSource.h"
+
+UMeasureToolInteraction::UMeasureToolInteraction()
+	: PixelStart()
+	, PixelEnd()
+{
+	ToolType = UE::Editor::ViewportInteractions::Measure;
+	InteractionName = TEXT("Measure Tool");
+
+	bUseSnapping = true;
+	bConvertDelta = false;
+
+	UViewportClickDragBehavior* MouseBehavior = NewObject<UViewportClickDragBehavior>();
+	MouseBehavior->Initialize(this);
+	MouseBehavior->SetDefaultPriority(UE::Editor::DragTools::MEASURE_PRIORITY);
+	MouseBehavior->EnableButton(EKeys::MiddleMouseButton);
+
+	RegisterInputBehavior(MouseBehavior);
+}
+
+void UMeasureToolInteraction::Render(const FSceneView* InView, FCanvas* InCanvas, EViewInteractionState InteractionState)
+{
+	if (!IsDragging() || !InCanvas || !InView)
+	{
+		return;
+	}
+
+	if (InteractionState != EViewInteractionState::Focused)
+	{
+		return;
+	}
+
+	FEditorViewportClient* const EditorViewportClient = GetEditorViewportClient();
+	if (!EditorViewportClient)
+	{
+		return;
+	}
+
+	if (!EditorViewportClient->Viewport)
+	{
+		return;
+	}
+
+	const float OrthoUnitsPerPixel = EditorViewportClient->GetOrthoUnitsPerPixel(EditorViewportClient->Viewport);
+	const float Length =
+		FMath::RoundToFloat(FVector2f{ PixelEnd - PixelStart }.Size() * OrthoUnitsPerPixel * EditorViewportClient->GetDPIScale());
+
+	if (InCanvas && Length >= 1.0f)
+	{
+		const FLinearColor ToolColor = GetDefault<UEditorStyleSettings>()->ViewportToolOverlayColor;
+		FCanvasLineItem LineItem(PixelStart, PixelEnd);
+		LineItem.SetColor(ToolColor);
+		InCanvas->DrawItem(LineItem);
+
+		const FVector2D PixelMid = FVector2D(PixelStart + ((PixelEnd - PixelStart) / 2));
+
+		// Calculate number of decimal places to display, based on the current viewport zoom
+		float Divisor = 1.0f;
+		int32 DecimalPlaces = 0;
+		const float OrderOfMagnitude = FMath::LogX(10.0f, OrthoUnitsPerPixel);
+
+		switch (GetDefault<ULevelEditorViewportSettings>()->MeasuringToolUnits)
+		{
+		case MeasureUnits_Meters:
+			Divisor = 100.0f;
+			// Max one decimal place allowed for meters.
+			DecimalPlaces = FMath::Clamp(FMath::FloorToInt(1.5f - OrderOfMagnitude), 0, 1);
+			break;
+
+		case MeasureUnits_Kilometers:
+			Divisor = 100000.0f;
+			// Max two decimal places allowed for kilometers.
+			DecimalPlaces = FMath::Clamp(FMath::FloorToInt(4.5f - OrderOfMagnitude), 0, 2);
+			break;
+		}
+
+		FNumberFormattingOptions Options;
+		Options.UseGrouping = false;
+		Options.MinimumFractionalDigits = DecimalPlaces;
+		Options.MaximumFractionalDigits = DecimalPlaces;
+
+		const FText LengthStr = FText::AsNumber(Length / Divisor, &Options);
+
+		FCanvasTextItem TextItem(
+			FVector2D(FMath::FloorToFloat(PixelMid.X), FMath::FloorToFloat(PixelMid.Y)), LengthStr, GEngine->GetSmallFont(), ToolColor
+		);
+		TextItem.bCentreX = true;
+		InCanvas->DrawItem(TextItem);
+	}
+}
+
+FInputRayHit UMeasureToolInteraction::CanBeginClickDragSequence(const FInputDeviceRay& InPressPos)
+{
+	if (FEditorViewportClient* const EditorViewportClient = GetEditorViewportClient())
+	{
+		if (IsCurrentModeSupported())
+		{
+			// Measure Tool is for Orthographic viewports only
+			if (EditorViewportClient->IsOrtho())
+			{
+				return FInputRayHit(TNumericLimits<float>::Max());
+			}
+		}
+	}
+
+	return FInputRayHit();
+}
+
+void UMeasureToolInteraction::OnDragStart(const FInputDeviceRay& InPressPos)
+{
+	if (UViewportInteractionsBehaviorSource* BehaviorSource = GetViewportInteractionsBehaviorSource())
+	{
+		BehaviorSource->SetMouseCursorOverride(EMouseCursor::Default);
+	}
+
+	FEditorViewportClient* const EditorViewportClient = GetEditorViewportClient();
+	if (EditorViewportClient && EditorViewportClient->Viewport)
+	{
+		OnActivateTool().Broadcast();
+
+		bIsDragging = true;
+
+		FIntPoint MousePos;
+		EditorViewportClient->Viewport->GetMousePos(MousePos);
+
+		// Take into account DPI scale for drawing lines properly when scale is not 1.0
+		Start = FVector(MousePos.X, MousePos.Y, 0) / EditorViewportClient->GetDPIScale();
+
+		// Snap to constraints.
+		if (bUseSnapping)
+		{
+			const float GridSize = GEditor->GetGridSize();
+			const FVector GridBase(GridSize, GridSize, GridSize);
+			FSnappingUtils::SnapPointToGrid(Start, GridBase);
+		}
+
+		End = Start;
+
+		PixelStart = GetSnappedPixelPos(FVector2D(InPressPos.ScreenPosition));
+		PixelEnd = PixelStart;
+	}
+}
+
+void UMeasureToolInteraction::OnDrag(const FInputDeviceRay& InDragPos)
+{
+	if (FEditorViewportClient* const EditorViewportClient = GetEditorViewportClient())
+	{
+		if (EditorViewportClient->Viewport)
+		{
+			FIntPoint MousePos;
+			EditorViewportClient->Viewport->GetMousePos(MousePos);
+			PixelEnd = GetSnappedPixelPos(FVector2D(MousePos));
+		}
+	}
+}
+
+FVector2D UMeasureToolInteraction::GetSnappedPixelPos(FVector2D InPixelPos) const
+{
+	FEditorViewportClient* const EditorViewportClient = GetEditorViewportClient();
+	if (!EditorViewportClient)
+	{
+		return FVector2D::Zero();
+	}
+
+	FSceneViewFamilyContext ViewFamily(
+		FSceneViewFamily::ConstructionValues(EditorViewportClient->Viewport, EditorViewportClient->GetScene(), EditorViewportClient->EngineShowFlags)
+			.SetRealtimeUpdate(EditorViewportClient->IsRealtime())
+	);
+
+	FSceneView* View = EditorViewportClient->CalcSceneView(&ViewFamily);
+
+	// Put the mouse pos in world space
+	FVector2f PixelPosFloat{ InPixelPos };
+	FVector WorldPos = View->ScreenToWorld(View->PixelToScreen(PixelPosFloat.X, PixelPosFloat.Y, 0.5f));
+
+	// Snap the world position
+	const float GridSize = GEditor->GetGridSize();
+	const FVector GridBase(GridSize, GridSize, GridSize);
+	FSnappingUtils::SnapPointToGrid(WorldPos, GridBase);
+
+	// And back into pixel-space (might fail, in which case we return the original)
+	View->WorldToPixel(WorldPos, InPixelPos);
+
+	// The canvas we're going to render to factors in dpi scale to final position.
+	// Since we're basing our position off mouse coordinates it will already be pixel accurate and therefore we must back out the scale the canvas will apply
+	InPixelPos /= EditorViewportClient->GetDPIScale();
+
+	return InPixelPos;
+}
diff --git a/Engine/Source/Editor/Experimental/EditorInteractiveToolsFramework/Private/EditorDragTools/MoveCameraWithObjectInteraction.cpp b/Engine/Source/Editor/Experimental/EditorInteractiveToolsFramework/Private/EditorDragTools/MoveCameraWithObjectInteraction.cpp
new file mode 100644
--- /dev/null
+++ b/Engine/Source/Editor/Experimental/EditorInteractiveToolsFramework/Private/EditorDragTools/MoveCameraWithObjectInteraction.cpp
@@ -0,0 +1,176 @@
+// Copyright Epic Games, Inc. All Rights Reserved.
+
+#include "EditorDragTools/MoveCameraWithObjectInteraction.h"
+
+#include "ContextObjectStore.h"
+#include "EditorDragTools/DragToolInteraction.h"
+#include "EditorGizmos/EditorTransformGizmo.h"
+#include "EditorGizmos/EditorTransformGizmoUtil.h"
+#include "EditorInteractiveGizmoManager.h"
+#include "EditorModeManager.h"
+#include "InteractiveGizmo.h"
+#include "InteractiveGizmoManager.h"
+#include "InteractiveToolManager.h"
+#include "LevelEditorViewport.h"
+#include "ViewportClientNavigationHelper.h"
+#include "Tools/EdModeInteractiveToolsContext.h"
+#include "ViewportInteractions/ViewportInteractionsBehaviorSource.h"
+
+UMoveCameraWithObjectInteraction::UMoveCameraWithObjectInteraction()
+{
+	ToolType = UE::Editor::ViewportInteractions::DragMovesCamera;
+	InteractionName = TEXT("Drag Moves Camera");
+	GroupName = UE::Editor::ViewportInteractions::DragToolsGroupName;
+}
+
+UMoveCameraWithObjectInteraction::~UMoveCameraWithObjectInteraction()
+{
+	UEditorInteractiveGizmoManager::OnUsesNewTRSGizmosChangedDelegate().RemoveAll(this);
+
+	if (ToolManager)
+	{
+		if (UContextObjectStore* ContextStore = ToolManager->GetContextObjectStore())
+		{
+			if (UEditorTransformGizmoContextObject* ContextObject = ContextStore->FindContext<UEditorTransformGizmoContextObject>())
+			{
+				ContextObject->OnGizmoCreatedDelegate().RemoveAll(this);
+			}
+		}
+	}
+	
+	if (UViewportInteractionsBehaviorSource* Source = GetViewportInteractionsBehaviorSource())
+	{
+		Source->OnShiftKeyStateChanged().RemoveAll(this);
+	}
+
+	UnregisterProxyDelegates();
+}
+
+void UMoveCameraWithObjectInteraction::RegisterProxyDelegates()
+{
+	if (TransformGizmo)
+	{
+		// Register Proxy Delegates
+		if (TObjectPtr<UTransformProxy> TransformProxy = TransformGizmo->ActiveTarget)
+		{
+			OnBeginPivotEditDelegate =
+				TransformProxy->OnBeginTransformEdit.AddUObject(this, &UMoveCameraWithObjectInteraction::OnGizmoMovementBegin);
+
+			OnEndPivotEditDelegate =
+				TransformProxy->OnEndTransformEdit.AddUObject(this, &UMoveCameraWithObjectInteraction::OnGizmoMovementEnd);
+		}
+	}
+}
+
+void UMoveCameraWithObjectInteraction::UnregisterProxyDelegates()
+{
+	if (TransformGizmo)
+	{
+		// Unregister Proxy delegates
+		if (TObjectPtr<UTransformProxy> TransformProxy = TransformGizmo->ActiveTarget)
+		{
+			TransformProxy->OnBeginTransformEdit.Remove(OnBeginPivotEditDelegate);
+			TransformProxy->OnEndTransformEdit.Remove(OnEndPivotEditDelegate);
+		}
+
+		TransformGizmo = nullptr;
+	}
+}
+
+void UMoveCameraWithObjectInteraction::OnGizmoCreatedDelegate(UTransformGizmo* InTransformGizmo)
+{
+	if (!TransformGizmo)
+	{
+		TransformGizmo = Cast<UEditorTransformGizmo>(InTransformGizmo);
+
+		RegisterProxyDelegates();
+	}
+}
+
+void UMoveCameraWithObjectInteraction::OnUsesNewTRSGizmosChanged(bool bInUseNewTRSGizmos)
+{
+	if (!bInUseNewTRSGizmos)
+	{
+		UnregisterProxyDelegates();
+	}
+}
+
+void UMoveCameraWithObjectInteraction::Initialize(UViewportInteractionsBehaviorSource* const InViewportInteractionsBehaviorSource)
+{
+	Super::Initialize(InViewportInteractionsBehaviorSource);
+	
+	GetViewportInteractionsBehaviorSource()->OnShiftKeyStateChanged().AddUObject(this, &UMoveCameraWithObjectInteraction::OnShiftKeyChanged);
+
+	if (const FEditorViewportClient* const EditorViewportClient = GetEditorViewportClient())
+	{
+		if (const FEditorModeTools* const ModeTools = EditorViewportClient->GetModeTools())
+		{
+			if (const UModeManagerInteractiveToolsContext* const InteractiveToolsContext = ModeTools->GetInteractiveToolsContext())
+			{
+				ToolManager = InteractiveToolsContext->ToolManager;
+			}
+		}
+	}
+
+	// New TRS Gizmos are already enabled - let's retrieve the Transform Gizmo, from which we register to drag Begin and End delegates
+	if (UEditorInteractiveGizmoManager::UsesNewTRSGizmos())
+	{
+		if (UInteractiveGizmoManager* PairedGizmoManager = ToolManager->GetPairedGizmoManager())
+		{
+			const TArray<UInteractiveGizmo*> Gizmos = PairedGizmoManager->FindAllGizmosOfType("EditorTransformGizmoBuilder");
+
+			if (!Gizmos.IsEmpty())
+			{
+				TransformGizmo = Cast<UEditorTransformGizmo>(Gizmos[0]);
+			}
+		}
+
+		RegisterProxyDelegates();
+	}
+
+	// If New TRS Gizmos are not enabled, we need to be able to know as soon as they are enabled
+	// This allows to retrieve the Transform Gizmo, from which we register to drag Begin and End delegates
+	if (UContextObjectStore* ContextStore = ToolManager->GetContextObjectStore())
+	{
+		if (UEditorTransformGizmoContextObject* ContextObject =
+				ContextStore->FindContext<UEditorTransformGizmoContextObject>())
+		{
+			ContextObject->OnGizmoCreatedDelegate().AddUObject(this, &UMoveCameraWithObjectInteraction::OnGizmoCreatedDelegate);
+		}
+	}
+
+	// In case new TRS Gizmos are disabled, we want to know that, so we can stop listening to Drag delegates
+	UEditorInteractiveGizmoManager::OnUsesNewTRSGizmosChangedDelegate().AddUObject(
+		this, &UMoveCameraWithObjectInteraction::OnUsesNewTRSGizmosChanged
+	);
+}
+
+void UMoveCameraWithObjectInteraction::OnShiftKeyChanged(bool bDown)
+{
+	bIsModifierActive = bDown;
+	UpdateCameraFollows();
+}
+
+void UMoveCameraWithObjectInteraction::OnGizmoMovementBegin(UTransformProxy* InTransformProxy)
+{
+	bGizmoIsDragged = true;
+	UpdateCameraFollows();
+}
+
+void UMoveCameraWithObjectInteraction::OnGizmoMovementEnd(UTransformProxy* InTransformProxy)
+{
+	bGizmoIsDragged = false;
+	UpdateCameraFollows();
+}
+
+void UMoveCameraWithObjectInteraction::UpdateCameraFollows() const
+{
+	if (FEditorViewportClient* ViewportClient = GetEditorViewportClient())
+	{
+		if (FViewportClientNavigationHelper* NavigationHelper = ViewportClient->GetViewportNavigationHelper())
+		{
+			NavigationHelper->bCameraFollowsGizmo = bGizmoIsDragged && bIsModifierActive;
+		}
+	}
+}
+
diff --git a/Engine/Source/Editor/Experimental/EditorInteractiveToolsFramework/Private/EditorDragTools/ViewportChangeInteraction.cpp b/Engine/Source/Editor/Experimental/EditorInteractiveToolsFramework/Private/EditorDragTools/ViewportChangeInteraction.cpp
new file mode 100644
--- /dev/null
+++ b/Engine/Source/Editor/Experimental/EditorInteractiveToolsFramework/Private/EditorDragTools/ViewportChangeInteraction.cpp
@@ -0,0 +1,196 @@
+﻿// Copyright Epic Games, Inc. All Rights Reserved.
+
+#include "EditorDragTools/ViewportChangeInteraction.h"
+#include "BaseBehaviors/ClickDragBehavior.h"
+#include "CanvasItem.h"
+#include "CanvasTypes.h"
+#include "EditorViewportClient.h"
+#include "Framework/Commands/InputChord.h"
+#include "Settings/EditorStyleSettings.h"
+#include "SnappingUtils.h"
+
+UViewportChangeInteraction::UViewportChangeInteraction()
+{
+	ToolType = UE::Editor::ViewportInteractions::ViewportChange;
+	InteractionName = TEXT("Viewport Change");
+
+	ViewOptionOffset = FVector2D::ZeroVector;
+	bUseSnapping = true;
+	bConvertDelta = false;
+}
+
+void UViewportChangeInteraction::Render(const FSceneView* View, FCanvas* Canvas, EViewInteractionState InteractionState)
+{
+	if (!IsDragging() || !Canvas || !View)
+	{
+		return;
+	}
+
+	if (InteractionState != EViewInteractionState::Focused)
+	{
+		return;
+	}
+
+	FCanvasLineItem LineItem(Start, End);
+	Canvas->DrawItem(LineItem);
+
+	const FLinearColor ToolColor = GetDefault<UEditorStyleSettings>()->ViewportToolOverlayColor;
+	FCanvasTextItem TextItem(
+		FVector2D(FMath::FloorToFloat(End.X), FMath::FloorToFloat(End.Y) + 20),
+		GetDesiredViewportTypeText(),
+		GEngine->GetMediumFont(),
+		ToolColor
+	);
+	TextItem.bCentreX = true;
+	Canvas->DrawItem(TextItem);
+}
+
+FInputRayHit UViewportChangeInteraction::CanBeginClickDragSequence(const FInputDeviceRay& InPressPos)
+{
+	return FInputRayHit(TNumericLimits<float>::Max()); // bHit is true. Depth is max to lose the standard tiebreaker.
+}
+
+void UViewportChangeInteraction::OnDragStart(const FInputDeviceRay& InPressPos)
+{
+	FEditorViewportClient* const EditorViewportClient = GetEditorViewportClient();
+	if (EditorViewportClient && EditorViewportClient->Viewport)
+	{
+		OnActivateTool().Broadcast();
+
+		bIsDragging = true;
+
+		FIntPoint MousePos;
+		EditorViewportClient->Viewport->GetMousePos(MousePos);
+
+		// Take into account DPI scale for drawing lines properly when scale is not 1.0
+		Start = FVector(MousePos.X, MousePos.Y, 0) / EditorViewportClient->GetDPIScale();
+
+		// Snap to constraints.
+		if (bUseSnapping)
+		{
+			const float GridSize = GEditor->GetGridSize();
+			const FVector GridBase(GridSize, GridSize, GridSize);
+			FSnappingUtils::SnapPointToGrid(Start, GridBase);
+		}
+
+		End = Start;
+	}
+}
+
+void UViewportChangeInteraction::OnDrag(const FInputDeviceRay& InDragPos)
+{
+	if (FEditorViewportClient* const EditorViewportClient = GetEditorViewportClient())
+	{
+		FIntPoint MousePos;
+		EditorViewportClient->Viewport->GetMousePos(MousePos);
+
+		// Take into account DPI scale for drawing lines properly when scale is not 1.0
+		End = FVector(MousePos) / EditorViewportClient->GetDPIScale();
+
+		ViewOptionOffset.X = End.X - Start.X;
+		ViewOptionOffset.Y = End.Y - Start.Y;
+	}
+}
+
+void UViewportChangeInteraction::OnDragEnd(const FInputDeviceRay& InReleasePos)
+{
+	if (FEditorViewportClient* const EditorViewportClient = GetEditorViewportClient())
+	{
+		EditorViewportClient->SetViewportType(GetDesiredViewportType());
+	}
+
+	UDragToolInteraction::OnDragEnd(InReleasePos);
+}
+
+ELevelViewportType UViewportChangeInteraction::GetDesiredViewportType() const
+{
+	ELevelViewportType ViewOption = LVT_Perspective;
+
+	if (ViewOptionOffset.Y == 0)
+	{
+		if (ViewOptionOffset.X == 0)
+		{
+			ViewOption = LVT_Perspective;
+		}
+		else if (ViewOptionOffset.X > 0)
+		{
+			ViewOption = LVT_OrthoRight; // Right
+		}
+		else
+		{
+			ViewOption = LVT_OrthoLeft; // Left
+		}
+	}
+	else
+	{
+		double OffsetRatio = ViewOptionOffset.X / ViewOptionOffset.Y;
+		double DragAngle = FMath::RadiansToDegrees(FMath::Atan(OffsetRatio));
+
+		if (ViewOptionOffset.Y >= 0)
+		{
+			if (DragAngle >= -15.f && DragAngle <= 15.f)
+			{
+				ViewOption = LVT_OrthoBottom; // Bottom
+			}
+			else if (DragAngle > 75.f)
+			{
+				ViewOption = LVT_OrthoRight; // Right
+			}
+			else if (DragAngle < -75.f)
+			{
+				ViewOption = LVT_OrthoLeft; // Left
+			}
+		}
+		else
+		{
+			if (DragAngle >= -15.f && DragAngle < 15.f)
+			{
+				ViewOption = LVT_OrthoTop; // Top
+			}
+			else if (DragAngle >= 15.f && DragAngle < 75.f)
+			{
+				ViewOption = LVT_OrthoFront; // Front
+			}
+			else if (DragAngle >= -75.f && DragAngle < -15.f)
+			{
+				ViewOption = LVT_OrthoBack; // Back
+			}
+			else if (DragAngle >= 75.f)
+			{
+				ViewOption = LVT_OrthoLeft; // Left
+			}
+			else if (DragAngle <= -75.f)
+			{
+				ViewOption = LVT_OrthoRight; // Right
+			}
+		}
+	}
+
+	return ViewOption;
+}
+
+FText UViewportChangeInteraction::GetDesiredViewportTypeText() const
+{
+	switch (GetDesiredViewportType())
+	{
+	case LVT_Perspective:
+		return FText::FromString("Perspective");
+	case LVT_OrthoFreelook:
+		return FText::FromString("Free Look");
+	case LVT_OrthoTop:
+		return FText::FromString("Top");
+	case LVT_OrthoLeft:
+		return FText::FromString("Left");
+	case LVT_OrthoFront:
+		return FText::FromString("Front");
+	case LVT_OrthoBack:
+		return FText::FromString("Back");
+	case LVT_OrthoBottom:
+		return FText::FromString("Bottom");
+	case LVT_OrthoRight:
+		return FText::FromString("Right");
+	default:;
+	}
+
+	return FText();
+}
diff --git a/Engine/Source/Editor/Experimental/EditorInteractiveToolsFramework/Private/ViewportInteractions/ViewportCameraRotateInteraction.cpp b/Engine/Source/Editor/Experimental/EditorInteractiveToolsFramework/Private/ViewportInteractions/ViewportCameraRotateInteraction.cpp
--- a/Engine/Source/Editor/Experimental/EditorInteractiveToolsFramework/Private/ViewportInteractions/ViewportCameraRotateInteraction.cpp
+++ b/Engine/Source/Editor/Experimental/EditorInteractiveToolsFramework/Private/ViewportInteractions/ViewportCameraRotateInteraction.cpp
@@ -10,6 +10,7 @@
 UViewportCameraRotateInteraction::UViewportCameraRotateInteraction()
 {
 	InteractionName = TEXT("Camera Rotate");
+	GroupName = UE::Editor::ViewportInteractions::CameraAndSelectionGroupName;
 	ToolType = UE::Editor::ViewportInteractions::CameraRotate;
 
 	UKeyInputBehavior* KeyInputBehavior = NewObject<UKeyInputBehavior>();
diff --git a/Engine/Source/Editor/Experimental/EditorInteractiveToolsFramework/Private/ViewportInteractions/ViewportCameraSpeedMouseWheelInteraction.cpp b/Engine/Source/Editor/Experimental/EditorInteractiveToolsFramework/Private/ViewportInteractions/ViewportCameraSpeedMouseWheelInteraction.cpp
--- a/Engine/Source/Editor/Experimental/EditorInteractiveToolsFramework/Private/ViewportInteractions/ViewportCameraSpeedMouseWheelInteraction.cpp
+++ b/Engine/Source/Editor/Experimental/EditorInteractiveToolsFramework/Private/ViewportInteractions/ViewportCameraSpeedMouseWheelInteraction.cpp
@@ -1,4 +1,4 @@
-// Copyright Epic Games, Inc. All Rights Reserved.
+﻿// Copyright Epic Games, Inc. All Rights Reserved.
 
 #include "ViewportInteractions/ViewportCameraSpeedMouseWheelInteraction.h"
 
@@ -10,6 +10,7 @@ UViewportCameraSpeedMouseWheelInteraction::UViewportCameraSpeedMouseWheelInterac
 {
 	ToolType = UE::Editor::ViewportInteractions::CameraTranslate;
 	InteractionName = TEXT("Camera Speed");
+	GroupName = UE::Editor::ViewportInteractions::CameraAndSelectionGroupName;
 
 	// Set camera speed using mouse wheel
 	UMouseWheelInputBehavior* MouseWheelInputBehavior = NewObject<UMouseWheelModifierInputBehavior>();
diff --git a/Engine/Source/Editor/Experimental/EditorInteractiveToolsFramework/Private/ViewportInteractions/ViewportCameraTranslateInteraction.cpp b/Engine/Source/Editor/Experimental/EditorInteractiveToolsFramework/Private/ViewportInteractions/ViewportCameraTranslateInteraction.cpp
--- a/Engine/Source/Editor/Experimental/EditorInteractiveToolsFramework/Private/ViewportInteractions/ViewportCameraTranslateInteraction.cpp
+++ b/Engine/Source/Editor/Experimental/EditorInteractiveToolsFramework/Private/ViewportInteractions/ViewportCameraTranslateInteraction.cpp
@@ -28,6 +28,7 @@ UViewportCameraTranslateInteraction::UViewportCameraTranslateInteraction()
 {
 	ToolType = UE::Editor::ViewportInteractions::CameraTranslate;
 	InteractionName = TEXT("Camera Translate");
+	GroupName = UE::Editor::ViewportInteractions::CameraAndSelectionGroupName;
 
 	// Camera translate
 	UKeyInputBehavior* KeyInputBehavior = NewObject<UKeyInputBehavior>();
diff --git a/Engine/Source/Editor/Experimental/EditorInteractiveToolsFramework/Private/ViewportInteractions/ViewportDragInteraction.cpp b/Engine/Source/Editor/Experimental/EditorInteractiveToolsFramework/Private/ViewportInteractions/ViewportDragInteraction.cpp
--- a/Engine/Source/Editor/Experimental/EditorInteractiveToolsFramework/Private/ViewportInteractions/ViewportDragInteraction.cpp
+++ b/Engine/Source/Editor/Experimental/EditorInteractiveToolsFramework/Private/ViewportInteractions/ViewportDragInteraction.cpp
@@ -26,6 +26,8 @@ UViewportDragInteraction::UViewportDragInteraction()
 	ViewportClickDragBehavior = ClickDragBehavior;
 
 	RegisterInputBehavior(ClickDragBehavior);
+
+	GroupName = UE::Editor::ViewportInteractions::CameraAndSelectionGroupName;
 }
 
 FInputRayHit UViewportDragInteraction::CanBeginClickDragSequence(const FInputDeviceRay& InPressPos)
diff --git a/Engine/Source/Editor/Experimental/EditorInteractiveToolsFramework/Private/ViewportInteractions/ViewportFOVInteraction.cpp b/Engine/Source/Editor/Experimental/EditorInteractiveToolsFramework/Private/ViewportInteractions/ViewportFOVInteraction.cpp
--- a/Engine/Source/Editor/Experimental/EditorInteractiveToolsFramework/Private/ViewportInteractions/ViewportFOVInteraction.cpp
+++ b/Engine/Source/Editor/Experimental/EditorInteractiveToolsFramework/Private/ViewportInteractions/ViewportFOVInteraction.cpp
@@ -11,6 +11,7 @@ UViewportFOVInteraction::UViewportFOVInteraction()
 	ToolType = UE::Editor::ViewportInteractions::FOV;
 
 	InteractionName = TEXT("FOV");
+	GroupName = UE::Editor::ViewportInteractions::CameraAndSelectionGroupName;
 
 	// Separating keys which work only when Mouse Looking from numpad ones and arrows, which always work.
 	// This prevents this IKeyInputBehaviorTarget from capturing unwanted inputs
diff --git a/Engine/Source/Editor/Experimental/EditorInteractiveToolsFramework/Private/ViewportInteractions/ViewportInteraction.cpp b/Engine/Source/Editor/Experimental/EditorInteractiveToolsFramework/Private/ViewportInteractions/ViewportInteraction.cpp
--- a/Engine/Source/Editor/Experimental/EditorInteractiveToolsFramework/Private/ViewportInteractions/ViewportInteraction.cpp
+++ b/Engine/Source/Editor/Experimental/EditorInteractiveToolsFramework/Private/ViewportInteractions/ViewportInteraction.cpp
@@ -112,7 +112,12 @@ bool UViewportInteraction::IsEnabled() const
 
 FName UViewportInteraction::GetInteractionName() const
 {
-	return InteractionName.IsSet() ? InteractionName.GetValue() : StaticClass()->GetFName();
+	return InteractionName.IsNone() ? StaticClass()->GetFName() : InteractionName;
+}
+
+FName UViewportInteraction::GetInteractionGroup() const
+{
+	return GroupName.IsNone() ? UE::Editor::ViewportInteractions::DefaultGroupName : GroupName;
 }
 
 void UViewportInteraction::Initialize(UViewportInteractionsBehaviorSource* const InViewportInteractionsBehaviorSource)
diff --git a/Engine/Source/Editor/Experimental/EditorInteractiveToolsFramework/Private/ViewportInteractions/ViewportInteractionsBehaviorSource.cpp b/Engine/Source/Editor/Experimental/EditorInteractiveToolsFramework/Private/ViewportInteractions/ViewportInteractionsBehaviorSource.cpp
--- a/Engine/Source/Editor/Experimental/EditorInteractiveToolsFramework/Private/ViewportInteractions/ViewportInteractionsBehaviorSource.cpp
+++ b/Engine/Source/Editor/Experimental/EditorInteractiveToolsFramework/Private/ViewportInteractions/ViewportInteractionsBehaviorSource.cpp
@@ -9,6 +9,7 @@
 #include "Tools/EdModeInteractiveToolsContext.h"
 #include "UnrealClient.h"
 #include "ViewportClientNavigationHelper.h"
+#include "EditorDragTools/DragToolInteraction.h"
 #include "ViewportInteractions/ViewportInteraction.h"
 
 namespace UE::Editor::ViewportInteractions
@@ -99,9 +100,6 @@ void LOG(const TCHAR* InMessage)
 
 static constexpr int MODIFIERS_PRIORITY = VIEWPORT_INTERACTIONS_DEFAULT_PRIORITY;
 
-// We use this string to group Behavior Inputs so that we can selectively remove them if needed, without touching other inputs we might add (e.g. the ones added by Initialize)
-static const FString ViewportInteractionsGroupName = TEXT("ViewportInteractions");
-
 } // namespace UE::Editor::ViewportInteractions
 
 UToolsContextCursorAPI* GetCursorAPI(const FEditorViewportClient* InViewportClient)
@@ -127,15 +125,15 @@ void UViewportInteractionsBehaviorSource::OnUpdateModifierState(int InModifierID
 {
 	if (InModifierID == UE::Editor::ViewportInteractions::ShiftKeyMod)
 	{
-		bIsShiftDown = bInIsOn;
+		SetShiftKeyState(bInIsOn);
 	}
 	else if (InModifierID == UE::Editor::ViewportInteractions::CtrlKeyMod)
 	{
-		bIsCtrlDown = bInIsOn;
+		SetCtrlKeyState(bInIsOn);
 	}
 	else if (InModifierID == UE::Editor::ViewportInteractions::AltKeyMod)
 	{
-		bIsAltDown = bInIsOn;
+		SetAltKeyState(bInIsOn);
 	}
 	else if (InModifierID == UE::Editor::ViewportInteractions::LeftMouseButtonMod)
 	{
@@ -153,9 +151,13 @@ void UViewportInteractionsBehaviorSource::OnUpdateModifierState(int InModifierID
 
 void UViewportInteractionsBehaviorSource::OnForceEndCapture()
 {
-	bIsShiftDown = false;
-	bIsCtrlDown = false;
-	bIsAltDown = false;
+	// We reset modifiers only if editor has lost focus
+	if (!FApp::HasFocus())
+	{
+		SetShiftKeyState(false);
+		SetCtrlKeyState(false);
+		SetAltKeyState(false);
+	}
 
 	bIsLeftMouseButtonDown = false;
 	bIsMiddleMouseButtonDown = false;
@@ -188,7 +190,7 @@ void UViewportInteractionsBehaviorSource::Initialize(const UEditorInteractiveToo
 			this, UE::Editor::ViewportInteractions::AltKeyMod, FInputDeviceState::IsAltKeyDown
 		);
 		KeyAsModifierInputBehavior->SetDefaultPriority(UE::Editor::ViewportInteractions::MODIFIERS_PRIORITY);
-		BehaviorSet->Add(KeyAsModifierInputBehavior, this);
+		CoreBehaviors.Add(KeyAsModifierInputBehavior);
 
 		UMouseButtonAsModifierInputBehavior* MouseButtonAsModifierInputBehavior =
 			NewObject<UMouseButtonAsModifierInputBehavior>();
@@ -217,8 +219,10 @@ void UViewportInteractionsBehaviorSource::Initialize(const UEditorInteractiveToo
 			}
 		);
 		MouseButtonAsModifierInputBehavior->SetDefaultPriority(UE::Editor::ViewportInteractions::MODIFIERS_PRIORITY);
-		BehaviorSet->Add(MouseButtonAsModifierInputBehavior, this);
+		CoreBehaviors.Add(MouseButtonAsModifierInputBehavior);
 	}
+	
+	Reset();
 }
 
 void UViewportInteractionsBehaviorSource::RegisterBehaviorSources()
@@ -254,6 +258,32 @@ void UViewportInteractionsBehaviorSource::Tick(float InDeltaTime) const
 	}
 }
 
+void UViewportInteractionsBehaviorSource::RenderTools(
+	FViewportClient* InViewportClient, const FSceneView* InSceneView, FCanvas* InCanvas
+) const
+{
+	const FViewportClient* Focused = GetEditorViewportClient();
+	const FViewportClient* Hovered = GetHoveredEditorViewportClient();
+
+	EViewInteractionState InteractionState = EViewInteractionState::None;
+	if (InViewportClient == Focused)
+	{
+		InteractionState = EViewInteractionState::Focused;
+	}
+	else if (InViewportClient == Hovered)
+	{
+		InteractionState = EViewInteractionState::Hovered;
+	}
+
+	for (const TPair<FName, TObjectPtr<UViewportInteraction>>& Pair : ViewportInteractions)
+	{
+		if (UViewportInteraction* Interaction = Pair.Value)
+		{
+			Interaction->Render(InSceneView, InCanvas, InteractionState);
+		}
+	}
+}
+
 void UViewportInteractionsBehaviorSource::AddInteractions(const TArray<const UClass*>& InInteractions, bool bInReregister)
 {
 	for (const UClass* InteractionClass : InInteractions)
@@ -277,6 +307,12 @@ UViewportInteraction* UViewportInteractionsBehaviorSource::AddInteraction(const
 	}
 
 	UViewportInteraction* Interaction = NewObject<UViewportInteraction>(GetTransientPackage(), InInteractionClass);
+
+	if (ViewportInteractions.Contains(Interaction->GetInteractionName()))
+	{
+		return nullptr;
+	}
+
 	Interaction->Initialize(this);
 
 	UE::Editor::ViewportInteractions::LOG(*FString::Printf(TEXT("Initializing Viewport Interaction \"%s\""), *InInteractionClass->GetName()));
@@ -284,7 +320,7 @@ UViewportInteraction* UViewportInteractionsBehaviorSource::AddInteraction(const
 
 	for (UInputBehavior* const InputBehavior : Interaction->GetInputBehaviors())
 	{
-		BehaviorSet->Add(InputBehavior, Interaction, UE::Editor::ViewportInteractions::ViewportInteractionsGroupName);
+		BehaviorSet->Add(InputBehavior, Interaction, Interaction->GetInteractionGroup().ToString());
 	}
 
 	if (bInReregister)
@@ -300,8 +336,14 @@ UViewportInteraction* UViewportInteractionsBehaviorSource::AddInteraction(const
 void UViewportInteractionsBehaviorSource::Reset()
 {
 	DeregisterBehaviorSources();
-	ViewportInteractions.Empty();
-	BehaviorSet->RemoveByGroup(UE::Editor::ViewportInteractions::ViewportInteractionsGroupName);
+	BehaviorSet->RemoveAll();
+	ViewportInteractions.Reset();
+	ViewportInteractionsStatusMap.Reset();
+	
+	for (const TObjectPtr<UInputBehavior>& Behavior : CoreBehaviors)
+	{
+		BehaviorSet->Add(Behavior.Get(), this);
+	}
 }
 
 void UViewportInteractionsBehaviorSource::SetMouseCursorOverride(EMouseCursor::Type InMouseCursor)
@@ -392,6 +434,35 @@ FEditorViewportClient* UViewportInteractionsBehaviorSource::GetEditorViewportCli
 	return EditorViewportClient;
 }
 
+FEditorViewportClient* UViewportInteractionsBehaviorSource::GetHoveredEditorViewportClient() const
+{
+	FEditorViewportClient* EditorViewportClient = nullptr;
+
+	if (const UEditorInteractiveToolsContext* const InteractiveToolsContext = EditorInteractiveToolsContextWeak.Get())
+	{
+		if (const FEditorModeTools* const ModeManager = InteractiveToolsContext->GetParentEditorModeManager())
+		{
+			EditorViewportClient = ModeManager->GetHoveredViewportClient();
+		}
+
+		if (!EditorViewportClient)
+		{
+			if (UInteractiveToolManager* const ToolManager = InteractiveToolsContext->ToolManager)
+			{
+				if (const IToolsContextQueriesAPI* const ContextQueriesAPI = ToolManager->GetContextQueriesAPI())
+				{
+					if (const FViewport* const Viewport = ContextQueriesAPI->GetHoveredViewport())
+					{
+						EditorViewportClient = static_cast<FEditorViewportClient*>(Viewport->GetClient());
+					}
+				}
+			}
+		}
+	}
+
+	return EditorViewportClient;
+}
+
 void UViewportInteractionsBehaviorSource::SetViewportInteractionActive(FName InViewportInteraction, bool bInActive)
 {
 	ViewportInteractionsStatusMap.Emplace(InViewportInteraction, bInActive);
@@ -406,3 +477,21 @@ bool UViewportInteractionsBehaviorSource::IsViewportInteractionActive(FName InVi
 
 	return false;
 }
+
+void UViewportInteractionsBehaviorSource::SetAltKeyState(bool bInDown)
+{
+	bIsAltDown = bInDown;
+	OnAltKeyStateChangedDelegate.Broadcast(bIsAltDown);
+}
+
+void UViewportInteractionsBehaviorSource::SetCtrlKeyState(bool bInDown)
+{
+	bIsCtrlDown = bInDown;
+	OnCtrlKeyStateChangedDelegate.Broadcast(bIsCtrlDown);
+}
+
+void UViewportInteractionsBehaviorSource::SetShiftKeyState(bool bInDown)
+{
+	bIsShiftDown = bInDown;
+	OnShiftKeyStateChangedDelegate.Broadcast(bIsShiftDown);
+}
diff --git a/Engine/Source/Editor/Experimental/EditorInteractiveToolsFramework/Private/ViewportInteractions/ViewportInteractionsUtilities.cpp b/Engine/Source/Editor/Experimental/EditorInteractiveToolsFramework/Private/ViewportInteractions/ViewportInteractionsUtilities.cpp
--- a/Engine/Source/Editor/Experimental/EditorInteractiveToolsFramework/Private/ViewportInteractions/ViewportInteractionsUtilities.cpp
+++ b/Engine/Source/Editor/Experimental/EditorInteractiveToolsFramework/Private/ViewportInteractions/ViewportInteractionsUtilities.cpp
@@ -2,8 +2,13 @@
 
 #include "ViewportInteractions/ViewportInteractionsUtilities.h"
 
-#include "ViewportInteractions/ViewportCameraSpeedMouseWheelInteraction.h"
+#include "EditorDragTools/BoxSelectInteraction.h"
+#include "EditorDragTools/DuplicateDragSelectionInteraction.h"
+#include "EditorDragTools/FrustumSelectInteraction.h"
+#include "EditorDragTools/MeasureToolInteraction.h"
+#include "EditorDragTools/MoveCameraWithObjectInteraction.h"
 #include "ViewportInteractions/ViewportCameraRotateInteraction.h"
+#include "ViewportInteractions/ViewportCameraSpeedMouseWheelInteraction.h"
 #include "ViewportInteractions/ViewportCameraTranslateInteraction.h"
 #include "ViewportInteractions/ViewportDollyInteraction.h"
 #include "ViewportInteractions/ViewportFOVInteraction.h"
@@ -37,3 +42,12 @@ void UE::Editor::ViewportInteractions::AddDefaultCameraMovementInteractions(
 		);
 	}
 }
+
+void UE::Editor::ViewportInteractions::AddDefaultDragToolsInteractions(UViewportInteractionsBehaviorSource* InInteractionsBehaviorSource)
+{
+	InInteractionsBehaviorSource->AddInteraction<UFrustumSelectInteraction>();
+	InInteractionsBehaviorSource->AddInteraction<UBoxSelectInteraction>();
+	InInteractionsBehaviorSource->AddInteraction<UMeasureToolInteraction>();
+	InInteractionsBehaviorSource->AddInteraction<UDuplicateDragSelectionInteraction>();
+	InInteractionsBehaviorSource->AddInteraction<UMoveCameraWithObjectInteraction>();
+}
diff --git a/Engine/Source/Editor/Experimental/EditorInteractiveToolsFramework/Private/ViewportInteractions/ViewportMoveYawInteraction.cpp b/Engine/Source/Editor/Experimental/EditorInteractiveToolsFramework/Private/ViewportInteractions/ViewportMoveYawInteraction.cpp
--- a/Engine/Source/Editor/Experimental/EditorInteractiveToolsFramework/Private/ViewportInteractions/ViewportMoveYawInteraction.cpp
+++ b/Engine/Source/Editor/Experimental/EditorInteractiveToolsFramework/Private/ViewportInteractions/ViewportMoveYawInteraction.cpp
@@ -5,7 +5,6 @@
 #include "EditorViewportClient.h"
 #include "Settings/LevelEditorViewportSettings.h"
 #include "ViewportClientNavigationHelper.h"
-#include "EditorDragTools/DragToolsBehaviourSource.h"
 #include "ViewportInteractions/ViewportInteractionsBehaviorSource.h"
 
 UViewportMoveYawInteraction::UViewportMoveYawInteraction()
diff --git a/Engine/Source/Editor/Experimental/EditorInteractiveToolsFramework/Private/ViewportInteractions/ViewportZoomInteraction.cpp b/Engine/Source/Editor/Experimental/EditorInteractiveToolsFramework/Private/ViewportInteractions/ViewportZoomInteraction.cpp
--- a/Engine/Source/Editor/Experimental/EditorInteractiveToolsFramework/Private/ViewportInteractions/ViewportZoomInteraction.cpp
+++ b/Engine/Source/Editor/Experimental/EditorInteractiveToolsFramework/Private/ViewportInteractions/ViewportZoomInteraction.cpp
@@ -15,6 +15,7 @@ UViewportZoomInteraction::UViewportZoomInteraction()
 	ToolType = UE::Editor::ViewportInteractions::Zoom;
 
 	InteractionName = TEXT("Zoom");
+	GroupName = UE::Editor::ViewportInteractions::CameraAndSelectionGroupName;
 
 	// Mousewheel zoom
 	UMouseWheelInputBehavior* MouseWheelInputBehavior = NewObject<UMouseWheelInputBehavior>();
diff --git a/Engine/Source/Editor/Experimental/EditorInteractiveToolsFramework/Public/EditorDragTools/BoxSelectInteraction.h b/Engine/Source/Editor/Experimental/EditorInteractiveToolsFramework/Public/EditorDragTools/BoxSelectInteraction.h
new file mode 100644
--- /dev/null
+++ b/Engine/Source/Editor/Experimental/EditorInteractiveToolsFramework/Public/EditorDragTools/BoxSelectInteraction.h
@@ -0,0 +1,64 @@
+// Copyright Epic Games, Inc. All Rights Reserved.
+
+#pragma once
+
+#include "DragToolInteraction.h"
+#include "MarqueeSelectInteraction.h"
+
+#include "BoxSelectInteraction.generated.h"
+
+class FCanvas;
+class UModel;
+
+UCLASS(MinimalAPI, Transient)
+class UBoxSelectInteraction : public UMarqueeSelectInteraction
+{
+	GENERATED_BODY()
+
+public:
+	UBoxSelectInteraction();
+
+	//~ Begin IViewportClickDragBehaviorTarget
+	virtual FInputRayHit CanBeginClickDragSequence(const FInputDeviceRay& InPressPos) override;
+	virtual void OnDragStart(const FInputDeviceRay& InPressPos) override;
+	virtual void OnDrag(const FInputDeviceRay& InDragPos) override;
+	virtual void OnDragEnd(const FInputDeviceRay& InReleasePos) override;
+	virtual void OnTerminateDragSequence() override;
+	//~ End IViewportClickDragBehaviorTarget
+
+protected:
+	virtual TArray<FEditorModeID> GetUnsupportedModes() override;
+
+private:
+	/**
+	 * Returns true if the provided BSP node intersects with the provided frustum
+	 *
+	 * @param InModel				The model containing BSP nodes to check
+	 * @param InNodeIndex				The index to a BSP node in the model.  This node is used for the bounds check.
+	 * @param InBox			The frustum to check against.
+	 * @param bInUseStrictSelection	true if the node must be entirely within the frustum
+	 */
+	static bool IntersectsBox(const UModel& InModel, int32 InNodeIndex, const FBox& InBox, bool bInUseStrictSelection);
+
+	/** Adds a hover effect to the passed in actor */
+	static void AddHoverEffect(AActor& InActor);
+
+	/** Adds a hover effect to the passed in bsp surface */
+	static void AddHoverEffect(UModel& InModel, int32 InSurfIndex);
+
+	/** Removes a hover effect from the passed in actor */
+	static void RemoveHoverEffect(AActor& InActor);
+
+	/** Removes a hover effect from the passed in bsp surface */
+	static void RemoveHoverEffect(UModel& InModel, int32 InSurfIndex);
+
+	/**
+	 * Calculates a box to check actors against
+	 *
+	 * @param OutBox	The created box.
+	 */
+	void CalculateBox(FBox& OutBox) const;
+	
+	/** List of BSP models to check for selection */
+	TArray<UModel*> ModelsToCheck;
+};
diff --git a/Engine/Source/Editor/Experimental/EditorInteractiveToolsFramework/Public/EditorDragTools/DragToolInteraction.h b/Engine/Source/Editor/Experimental/EditorInteractiveToolsFramework/Public/EditorDragTools/DragToolInteraction.h
new file mode 100644
--- /dev/null
+++ b/Engine/Source/Editor/Experimental/EditorInteractiveToolsFramework/Public/EditorDragTools/DragToolInteraction.h
@@ -0,0 +1,132 @@
+﻿// Copyright Epic Games, Inc. All Rights Reserved.
+
+#pragma once
+
+#include "BaseBehaviors/BehaviorTargetInterfaces.h"
+#include "Editor.h"
+#include "Behaviors/ViewportClickDragBehavior.h"
+#include "ViewportInteractions/ViewportInteraction.h"
+
+#include "DragToolInteraction.generated.h"
+
+class IEditorViewportClientProxy;
+struct FWorldSelectionElementArgs;
+class FCanvas;
+class FEditorModeTools;
+class FEditorViewportClient;
+class FPrimitiveDrawInterface;
+class FSceneView;
+class FUICommandInfo;
+
+/**
+ * ITF related helper functions
+ */
+namespace UE::Editor::DragTools
+{
+
+DECLARE_MULTICAST_DELEGATE(FOnEditorDragToolsToggleDelegate)
+DECLARE_MULTICAST_DELEGATE(FOnOnViewportChangeToolToggleDelegate)
+
+/**
+ * Returns true if ITF-based drag tools should be used.
+ * If an ITF version of a drag tool is not available yet, legacy will be used.
+ */
+EDITORINTERACTIVETOOLSFRAMEWORK_API bool UseEditorDragTools();
+EDITORINTERACTIVETOOLSFRAMEWORK_API bool IsViewportChangeToolEnabled();
+
+EDITORINTERACTIVETOOLSFRAMEWORK_API FOnEditorDragToolsToggleDelegate& OnEditorDragToolsActivated();
+EDITORINTERACTIVETOOLSFRAMEWORK_API FOnEditorDragToolsToggleDelegate& OnEditorDragToolsDeactivated();
+
+EDITORINTERACTIVETOOLSFRAMEWORK_API FOnOnViewportChangeToolToggleDelegate& OnViewportChangeToolActivated();
+EDITORINTERACTIVETOOLSFRAMEWORK_API FOnOnViewportChangeToolToggleDelegate& OnViewportChangeToolDeactivated();
+
+// In case New TRS Gizmos are enabled, we need Mouse Drag + Key modifier Drag Tools Behaviors to be processed before
+// the Gizmo ones, and also before viewport interactions - otherwise Drag Behaviors will never be triggered. See:
+// * UTransformGizmo::SetupBehaviors()			[DEFAULT_GIZMO_PRIORITY]
+// * UTransformGizmo::SetupIndirectBehaviors()	[DEFAULT_GIZMO_PRIORITY - 1]
+
+static constexpr int HANDLED_BEFORE_VIEWPORT_INTERACTIONS =
+	UE::Editor::ViewportInteractions::VIEWPORT_INTERACTIONS_DEFAULT_PRIORITY - 1;
+
+// Conversely, tools which are activated with just mouse buttons input (no key modifiers, mouse drag behavior only) need to be processed after Gizmo Behaviors
+static constexpr int HANDLED_AFTER_GIZMO_PRIORITY = FInputCapturePriority::DEFAULT_GIZMO_PRIORITY + 1;
+
+static constexpr int DRAG_TOOLS_DEFAULT_PRIORITY = HANDLED_BEFORE_VIEWPORT_INTERACTIONS - 10;
+static constexpr int DRAG_TOOLS_LOW_PRIORITY = HANDLED_BEFORE_VIEWPORT_INTERACTIONS - 1;
+
+static constexpr int DUPLICATE_SELECTION_PRIORITY = DRAG_TOOLS_DEFAULT_PRIORITY;
+static constexpr int FRUSTUM_SELECT_PRIORITY = DRAG_TOOLS_DEFAULT_PRIORITY;
+static constexpr int VIEWPORT_CHANGE_PRIORITY = DRAG_TOOLS_DEFAULT_PRIORITY;
+static constexpr int BOX_SELECT_PRIORITY = HANDLED_AFTER_GIZMO_PRIORITY;
+static constexpr int MEASURE_PRIORITY = DRAG_TOOLS_LOW_PRIORITY;
+
+} // namespace UE::Editor::DragTools
+
+/**
+ * The base class that all drag tools inherit from.
+ * The drag tools implement special behaviors for the user clicking and dragging in a viewport.
+ */
+UCLASS(MinimalAPI, Transient)
+class UDragToolInteraction : public UViewportInteraction, public IViewportClickDragBehaviorTarget
+{
+	GENERATED_BODY()
+
+public:
+	UDragToolInteraction();
+
+	virtual void Initialize(UViewportInteractionsBehaviorSource* const InViewportInteractionsBehaviorSource) override;
+
+	/** @return true if we are dragging */
+	bool IsDragging() const
+	{
+		return bIsDragging;
+	}
+
+	/** Does this drag tool need to have the mouse movement converted to the viewport orientation? */
+	bool bConvertDelta;
+
+	//~ Begin IViewportClickDragBehaviorTarget
+	virtual FInputRayHit CanBeginClickDragSequence(const FInputDeviceRay& InPressPos) override
+	{
+		return FInputRayHit();
+	}
+
+	virtual void OnDragStart(const FInputDeviceRay& InDragStartPos) override {}
+	virtual void OnDrag(const FInputDeviceRay& InDragPos) override {}
+	virtual void OnDragEnd(const FInputDeviceRay& InDragEndPos) override;
+	virtual void OnForceEndCapture() override;
+	virtual void OnTerminateDragSequence() override;
+	virtual void OnStateUpdated(const FInputDeviceState& InInputDeviceState) override;
+	//~ End IViewportClickDragBehaviorTarget
+
+	DECLARE_MULTICAST_DELEGATE(FOnToolStateChange);
+	FOnToolStateChange& OnActivateTool()
+	{
+		return OnToolActivatedDelegate;
+	}
+	FOnToolStateChange& OnDeactivateTool()
+	{
+		return OnToolDeactivatedDelegate;
+	}
+
+protected:
+	virtual TArray<FEditorModeID> GetUnsupportedModes() { return {}; }
+
+	bool IsCurrentModeSupported();
+
+	FEditorModeTools* ModeTools;
+
+	/** The start/end location of the current drag. */
+	FVector Start, End;
+
+	/** If true, the drag tool wants to be passed grid snapped values. */
+	bool bUseSnapping;
+
+	FInputDeviceState InputState;
+	bool bIsDragging;
+
+	FOnToolStateChange OnToolActivatedDelegate;
+	FOnToolStateChange OnToolDeactivatedDelegate;
+
+	IEditorViewportClientProxy* EditorViewportClientProxy;
+};
diff --git a/Engine/Source/Editor/Experimental/EditorInteractiveToolsFramework/Public/EditorDragTools/DragToolsBehaviourSource.h b/Engine/Source/Editor/Experimental/EditorInteractiveToolsFramework/Public/EditorDragTools/DragToolsBehaviourSource.h
deleted file mode 100644
diff --git a/Engine/Source/Editor/Experimental/EditorInteractiveToolsFramework/Public/EditorDragTools/DuplicateDragSelectionInteraction.h b/Engine/Source/Editor/Experimental/EditorInteractiveToolsFramework/Public/EditorDragTools/DuplicateDragSelectionInteraction.h
new file mode 100644
--- /dev/null
+++ b/Engine/Source/Editor/Experimental/EditorInteractiveToolsFramework/Public/EditorDragTools/DuplicateDragSelectionInteraction.h
@@ -0,0 +1,65 @@
+// Copyright Epic Games, Inc. All Rights Reserved.
+
+#pragma once
+
+#include "BaseGizmos/TransformProxy.h"
+#include "EditorGizmos/TransformGizmo.h"
+#include "ViewportInteractions/ViewportInteraction.h"
+
+#include "DuplicateDragSelectionInteraction.generated.h"
+
+class FLevelEditorViewportClient;
+class FScopedTransaction;
+class UEditorTransformGizmo;
+class UInteractiveToolManager;
+class UTransformProxy;
+
+/**
+ * Listens for key press of a specified modifier key, and if Level Viewport Gizmo is dragging, it duplicates current selection.
+ * Currently used with Alt
+ */
+UCLASS(MinimalAPI, Transient)
+class UDuplicateDragSelectionInteraction
+	: public UViewportInteraction
+
+{
+	GENERATED_BODY()
+
+public:
+	UDuplicateDragSelectionInteraction();
+
+	virtual ~UDuplicateDragSelectionInteraction() override;
+
+protected:
+	void OnGizmoMovementBegin(UTransformProxy* InTransformProxy);
+	void OnGizmoMovementEnd(UTransformProxy* InTransformProxy);
+
+	void RegisterProxyDelegates();
+	void Reset();
+
+	void OnGizmoCreatedDelegate(UTransformGizmo* InTransformGizmo);
+	void OnUsesNewTRSGizmosChanged(bool bInUseNewTRSGizmos);
+
+	virtual void Initialize(UViewportInteractionsBehaviorSource* const InViewportInteractionsBehaviorSource) override;
+
+	/**
+	 * Duplicate the current selection
+	 * @param bInSelectNewElements if newly created elements should be selected right after duplication
+	 */
+	void DuplicateSelection(bool bInSelectNewElements) const;
+
+private:
+	void OnModifierKeyChanged(bool bInIsDown);
+
+	UPROPERTY(Transient)
+	TObjectPtr<UInteractiveToolManager> ToolManager;
+
+	UPROPERTY(Transient)
+	TObjectPtr<UEditorTransformGizmo> TransformGizmo;
+
+	bool bGizmoIsDragged = false;
+	bool bModifierKeyIsPressed = false;
+
+	/** Keeps the duplicate + transform transaction alive between gizmo movement begin and end */
+	TSharedPtr<FScopedTransaction> DuplicateTransaction;
+};
diff --git a/Engine/Source/Editor/Experimental/EditorInteractiveToolsFramework/Public/EditorDragTools/EditorBoxSelect.h b/Engine/Source/Editor/Experimental/EditorInteractiveToolsFramework/Public/EditorDragTools/EditorBoxSelect.h
deleted file mode 100644
diff --git a/Engine/Source/Editor/Experimental/EditorInteractiveToolsFramework/Public/EditorDragTools/EditorDragToolBehaviorTarget.h b/Engine/Source/Editor/Experimental/EditorInteractiveToolsFramework/Public/EditorDragTools/EditorDragToolBehaviorTarget.h
deleted file mode 100644
diff --git a/Engine/Source/Editor/Experimental/EditorInteractiveToolsFramework/Public/EditorDragTools/EditorDuplicateDragSelection.h b/Engine/Source/Editor/Experimental/EditorInteractiveToolsFramework/Public/EditorDragTools/EditorDuplicateDragSelection.h
deleted file mode 100644
diff --git a/Engine/Source/Editor/Experimental/EditorInteractiveToolsFramework/Public/EditorDragTools/EditorFrustumSelect.h b/Engine/Source/Editor/Experimental/EditorInteractiveToolsFramework/Public/EditorDragTools/EditorFrustumSelect.h
deleted file mode 100644
diff --git a/Engine/Source/Editor/Experimental/EditorInteractiveToolsFramework/Public/EditorDragTools/EditorMarqueeSelect.h b/Engine/Source/Editor/Experimental/EditorInteractiveToolsFramework/Public/EditorDragTools/EditorMarqueeSelect.h
deleted file mode 100644
diff --git a/Engine/Source/Editor/Experimental/EditorInteractiveToolsFramework/Public/EditorDragTools/EditorMeasureTool.h b/Engine/Source/Editor/Experimental/EditorInteractiveToolsFramework/Public/EditorDragTools/EditorMeasureTool.h
deleted file mode 100644
diff --git a/Engine/Source/Editor/Experimental/EditorInteractiveToolsFramework/Public/EditorDragTools/EditorMoveCameraWithObject.h b/Engine/Source/Editor/Experimental/EditorInteractiveToolsFramework/Public/EditorDragTools/EditorMoveCameraWithObject.h
deleted file mode 100644
diff --git a/Engine/Source/Editor/Experimental/EditorInteractiveToolsFramework/Public/EditorDragTools/EditorViewportChange.h b/Engine/Source/Editor/Experimental/EditorInteractiveToolsFramework/Public/EditorDragTools/EditorViewportChange.h
deleted file mode 100644
diff --git a/Engine/Source/Editor/Experimental/EditorInteractiveToolsFramework/Public/EditorDragTools/FrustumSelectInteraction.h b/Engine/Source/Editor/Experimental/EditorInteractiveToolsFramework/Public/EditorDragTools/FrustumSelectInteraction.h
new file mode 100644
--- /dev/null
+++ b/Engine/Source/Editor/Experimental/EditorInteractiveToolsFramework/Public/EditorDragTools/FrustumSelectInteraction.h
@@ -0,0 +1,99 @@
+// Copyright Epic Games, Inc. All Rights Reserved.
+
+#pragma once
+
+#include "DragToolInteraction.h"
+#include "MarqueeSelectInteraction.h"
+
+#include "FrustumSelectInteraction.generated.h"
+
+#define UE_API EDITORINTERACTIVETOOLSFRAMEWORK_API
+
+
+class FCanvas;
+class FEditorViewportClient;
+class FSceneView;
+class UModel;
+struct FConvexVolume;
+
+/**
+ * Draws a box in the current viewport and when the mouse button is released,
+ * it selects/unselects everything inside of it.
+ */
+UCLASS(MinimalAPI, Transient)
+class UFrustumSelectInteraction : public UMarqueeSelectInteraction
+{
+	GENERATED_BODY()
+
+public:
+	UFrustumSelectInteraction();
+
+	//~ Begin IViewportClickDragBehaviorTarget
+	virtual FInputRayHit CanBeginClickDragSequence(const FInputDeviceRay& InPressPos) override;
+	virtual void OnDragStart(const FInputDeviceRay& InPressPos) override;
+	virtual void OnDrag(const FInputDeviceRay& InDragPos) override;
+	virtual void OnDragEnd(const FInputDeviceRay& InReleasePos) override;
+	virtual void OnTerminateDragSequence() override;
+	//~ End IViewportClickDragBehaviorTarget
+
+protected:
+	virtual TArray<FEditorModeID> GetUnsupportedModes() override;
+
+private:
+	/**
+	 * Returns true if the provided BSP node intersects with the provided frustum
+	 *
+	 * @param InModel				The model containing BSP nodes to check
+	 * @param NodeIndex				The index to a BSP node in the model.  This node is used for the bounds check.
+	 * @param InFrustum				The frustum to check against.
+	 * @param bUseStrictSelection	true if the node must be entirely within the frustum
+	 */
+	static bool IntersectsFrustum(const UModel& InModel, int32 NodeIndex, const FConvexVolume& InFrustum, bool bUseStrictSelection);
+
+	/**
+	 * Calculates a frustum to check actors against
+	 *
+	 * @param InView			Information about scene projection
+	 * @param OutFrustum		The created frustum
+	 * @param bUseBoxFrustum	If true a frustum out of the current dragged box will be created.  false will use the view frustum.
+	 */
+	void CalculateFrustum(const FSceneView* InView, FConvexVolume& OutFrustum, bool bUseBoxFrustum) const;
+};
+
+namespace UE::EditorDragTools
+{
+
+/**
+ * FViewportFrustum calculates a frustum box from 2D start and end positions in a viewport.  
+ */
+
+struct FViewportFrustum
+{
+public:
+	/** 
+	 * @param InViewportClient	Information about the viewport client (view location)
+	 * @param InView			Information about scene projection
+	 * @param InStart			The 2D screen space start position of the dragged box
+	 * @param InEnd				The 2D screen space end position of the dragged box
+	 * @param InUseBoxFrustum	If true a frustum out of the current dragged box will be created (false will use the view frustum)
+	 */
+	UE_API FViewportFrustum(const FEditorViewportClient& InViewportClient, const FSceneView& InView,
+		const FVector& InStart, const FVector& InEnd, const bool InUseBoxFrustum);
+
+	/** 
+	 * @param OutFrustum		The created frustum
+	 */
+	UE_API void Calculate(FConvexVolume& OutFrustum) const;
+
+private:
+
+	const FEditorViewportClient& ViewportClient;
+	const FSceneView& View;
+	const FVector& Start;
+	const FVector& End;
+	bool bUseBoxFrustum = false;
+};
+
+}
+
+#undef UE_API
diff --git a/Engine/Source/Editor/Experimental/EditorInteractiveToolsFramework/Public/EditorDragTools/MarqueeSelectInteraction.h b/Engine/Source/Editor/Experimental/EditorInteractiveToolsFramework/Public/EditorDragTools/MarqueeSelectInteraction.h
new file mode 100644
--- /dev/null
+++ b/Engine/Source/Editor/Experimental/EditorInteractiveToolsFramework/Public/EditorDragTools/MarqueeSelectInteraction.h
@@ -0,0 +1,34 @@
+﻿// Copyright Epic Games, Inc. All Rights Reserved.
+
+#pragma once
+#include "DragToolInteraction.h"
+
+#include "MarqueeSelectInteraction.generated.h"
+
+class FCanvas;
+class FEditorViewportClient;
+class FSceneView;
+class UModel;
+
+/**
+ * Base class for tools that use a marquee selection behavior
+ */
+UCLASS(MinimalAPI, Transient)
+class UMarqueeSelectInteraction : public UDragToolInteraction
+{
+	GENERATED_BODY()
+
+public:
+	UMarqueeSelectInteraction();
+
+	virtual void Render(const FSceneView* View, FCanvas* Canvas, EViewInteractionState InteractionState) override;
+	
+protected:
+
+	/**
+	 * @return Whether selection should only include items completely contained by the box.  
+	 */
+	bool IsWindowSelection() const;
+
+	TWeakObjectPtr<UViewportClickDragBehavior> ClickDragInputBehavior;
+};
diff --git a/Engine/Source/Editor/Experimental/EditorInteractiveToolsFramework/Public/EditorDragTools/MeasureToolInteraction.h b/Engine/Source/Editor/Experimental/EditorInteractiveToolsFramework/Public/EditorDragTools/MeasureToolInteraction.h
new file mode 100644
--- /dev/null
+++ b/Engine/Source/Editor/Experimental/EditorInteractiveToolsFramework/Public/EditorDragTools/MeasureToolInteraction.h
@@ -0,0 +1,35 @@
+// Copyright Epic Games, Inc. All Rights Reserved.
+
+#pragma once
+
+#include "DragToolInteraction.h"
+
+#include "MeasureToolInteraction.generated.h"
+
+class FCanvas;
+
+UCLASS(MinimalAPI, Transient)
+class UMeasureToolInteraction : public UDragToolInteraction
+{
+	GENERATED_BODY()
+
+public:
+	UMeasureToolInteraction();
+
+	virtual void Render(const FSceneView* InView,FCanvas* InCanvas, EViewInteractionState InteractionState) override;
+
+	//~ Begin IViewportClickDragBehaviorTarget
+	virtual FInputRayHit CanBeginClickDragSequence(const FInputDeviceRay& InPressPos) override;
+	virtual void OnDragStart(const FInputDeviceRay& InPressPos) override;
+	virtual void OnDrag(const FInputDeviceRay& InDragPos) override;
+	//~ End IViewportClickDragBehaviorTarget
+
+private:
+	/**
+	 * Gets the world-space snapped position of the specified pixel position
+	 */
+	FVector2D GetSnappedPixelPos(FVector2D InPixelPos) const;
+
+	/** Pixel-space positions for start and end */
+	FVector2D PixelStart, PixelEnd;
+};
diff --git a/Engine/Source/Editor/Experimental/EditorInteractiveToolsFramework/Public/EditorDragTools/MoveCameraWithObjectInteraction.h b/Engine/Source/Editor/Experimental/EditorInteractiveToolsFramework/Public/EditorDragTools/MoveCameraWithObjectInteraction.h
new file mode 100644
--- /dev/null
+++ b/Engine/Source/Editor/Experimental/EditorInteractiveToolsFramework/Public/EditorDragTools/MoveCameraWithObjectInteraction.h
@@ -0,0 +1,62 @@
+// Copyright Epic Games, Inc. All Rights Reserved.
+
+#pragma once
+
+#include "BaseGizmos/TransformProxy.h"
+#include "EditorGizmos/TransformGizmo.h"
+#include "ViewportInteractions/ViewportInteraction.h"
+
+#include "MoveCameraWithObjectInteraction.generated.h"
+
+#define UE_API EDITORINTERACTIVETOOLSFRAMEWORK_API
+
+class FLevelEditorViewportClient;
+class UEditorTransformGizmo;
+class UInteractiveToolManager;
+class UTransformProxy;
+
+/**
+ * Listens for key press of a specified modifier key, and if Level Viewport Gizmo is dragging, it will move the viewport camera along.
+ * Currently used with Shift
+ */
+UCLASS(MinimalAPI, Transient)
+class UMoveCameraWithObjectInteraction
+	: public UViewportInteraction
+{
+	GENERATED_BODY()
+
+public:
+	UE_API UMoveCameraWithObjectInteraction();
+	UE_API virtual ~UMoveCameraWithObjectInteraction() override;
+
+	virtual void Initialize(UViewportInteractionsBehaviorSource* const InViewportInteractionsBehaviorSource) override;
+
+protected:
+
+	UE_API void OnShiftKeyChanged(bool bDown);
+	UE_API void OnGizmoMovementBegin(UTransformProxy* InTransformProxy);
+	UE_API void OnGizmoMovementEnd(UTransformProxy* InTransformProxy);
+
+	UE_API void RegisterProxyDelegates();
+	UE_API void UnregisterProxyDelegates();
+
+	UE_API void OnGizmoCreatedDelegate(UTransformGizmo* InTransformGizmo);
+	UE_API void OnUsesNewTRSGizmosChanged(bool bInUseNewTRSGizmos);
+
+private:
+	void UpdateCameraFollows() const;
+
+	UPROPERTY(Transient)
+	TObjectPtr<UInteractiveToolManager> ToolManager;
+	UPROPERTY(Transient)
+	TObjectPtr<UEditorTransformGizmo> TransformGizmo;
+
+	bool bGizmoIsDragged = false;
+	bool bIsModifierActive = false;
+
+	FDelegateHandle OnBeginPivotEditDelegate;
+	FDelegateHandle OnEndPivotEditDelegate;
+	FDelegateHandle OnTransformChangedDelegate;
+};
+
+#undef UE_API
diff --git a/Engine/Source/Editor/Experimental/EditorInteractiveToolsFramework/Public/EditorDragTools/ViewportChangeInteraction.h b/Engine/Source/Editor/Experimental/EditorInteractiveToolsFramework/Public/EditorDragTools/ViewportChangeInteraction.h
new file mode 100644
--- /dev/null
+++ b/Engine/Source/Editor/Experimental/EditorInteractiveToolsFramework/Public/EditorDragTools/ViewportChangeInteraction.h
@@ -0,0 +1,36 @@
+// Copyright Epic Games, Inc. All Rights Reserved.
+
+#pragma once
+
+#include "Editor/UnrealEdTypes.h"
+#include "DragToolInteraction.h"
+
+#include "ViewportChangeInteraction.generated.h"
+
+class FCanvas;
+class FEditorViewportClient;
+class FLevelEditorViewportClient;
+class FSceneView;
+
+UCLASS(MinimalAPI, Transient)
+class UViewportChangeInteraction : public UDragToolInteraction
+{
+	GENERATED_BODY()
+
+public:
+	UViewportChangeInteraction();
+	virtual void Render(const FSceneView* View, FCanvas* Canvas, EViewInteractionState InteractionState) override;
+
+	//~ Begin IViewportClickDragBehaviorTarget
+	virtual FInputRayHit CanBeginClickDragSequence(const FInputDeviceRay& InPressPos) override;
+	virtual void OnDragStart(const FInputDeviceRay& InPressPos) override;
+	virtual void OnDrag(const FInputDeviceRay& InDragPos) override;
+	virtual void OnDragEnd(const FInputDeviceRay& InReleasePos) override;
+	//~ End IViewportClickDragBehaviorTarget
+
+private:
+	ELevelViewportType GetDesiredViewportType() const;
+	FText GetDesiredViewportTypeText() const;
+
+	FVector2D ViewOptionOffset;
+};
diff --git a/Engine/Source/Editor/Experimental/EditorInteractiveToolsFramework/Public/ViewportInteractions/ViewportInteraction.h b/Engine/Source/Editor/Experimental/EditorInteractiveToolsFramework/Public/ViewportInteractions/ViewportInteraction.h
--- a/Engine/Source/Editor/Experimental/EditorInteractiveToolsFramework/Public/ViewportInteractions/ViewportInteraction.h
+++ b/Engine/Source/Editor/Experimental/EditorInteractiveToolsFramework/Public/ViewportInteractions/ViewportInteraction.h
@@ -42,6 +42,8 @@ static constexpr int MULTIBUTTON_DRAG_PRIORITY = FInputCapturePriority::DEFAULT_
  * Use by FEditorViewportClient to check whether a specific interaction is handled by an ITF-based interaction, or if it
  * should be handled by legacy logic
  */
+
+// Camera Movement
 static FName Zoom = TEXT("Zoom");
 static FName Orbit = TEXT("Orbit");
 static FName CameraDrag= TEXT("CameraDrag");
@@ -49,10 +51,25 @@ static FName CameraTranslate = TEXT("CameraTranslate");
 static FName CameraRotate = TEXT("CameraRotate");
 static FName FOV = TEXT("FOV");
 static FName Commands = TEXT("Commands");
+
+// Clicking
 static FName SingleClick = TEXT("SingleClick");
 static FName DoubleClick = TEXT("DoubleClick");
 static FName ContextMenu = TEXT("ContextMenu");
 
+// Drag Tools
+static FName BoxSelect = TEXT("BoxSelect");
+static FName FrustumSelect = TEXT("FrustumSelect");
+static FName Measure = TEXT("Measure");
+static FName ViewportChange = TEXT("ViewportChange");
+static FName DuplicateDrag = TEXT("DuplicateDrag");
+static FName DragMovesCamera = TEXT("DragMovesCamera");
+
+// We use these strings to group Behavior Inputs so that we can selectively remove them if needed, without touching other inputs we might add (e.g. the ones added by Initialize)
+static const FName DefaultGroupName = TEXT("ViewportInteractions");
+static const FName CameraAndSelectionGroupName = TEXT("CameraAndSelection");
+static const FName DragToolsGroupName = TEXT("DragTools");
+
 } // namespace UE::Editor::ViewportInteractions
 
 /**
@@ -83,9 +100,11 @@ public:
 	UE_API void SetEnabled(bool bInEnabled);
 	UE_API bool IsEnabled() const;
 
-	virtual void Tick(float InDeltaTime) const
-	{
-	}
+	/** Called every frame. */
+	virtual void Tick(float InDeltaTime) const {}
+	
+	/** Called every frame to render tools to a 2D canvas */
+	virtual void Render(const FSceneView* InView, FCanvas* InCanvas, EViewInteractionState InteractionState) {}
 
 	TArray<TObjectPtr<UInputBehavior>> GetInputBehaviors()
 	{
@@ -94,6 +113,11 @@ public:
 
 	UE_API FName GetInteractionName() const;
 
+	/**
+	 * The Group this interaction belongs to.
+	 */
+	UE_API FName GetInteractionGroup() const;
+
 	/**
 	 * Initialization code which can be expanded by derived classes
 	 */
@@ -136,9 +160,15 @@ protected:
 
 	/**
 	 * Name for this interaction, used by Viewport Interactions Behavior Source to identify it.
-	 * If none is set, static class name will be used.
+	 * If NAME_None, static class name will be used.
+	 */
+	FName InteractionName = NAME_None;
+
+	/**
+	 * The Group this interaction belongs to. It can be specified to allow the selective removal/unregistering of interactions.
+	 * If NAME_None, a default group will be used.
 	 */
-	TOptional<FName> InteractionName;
+	FName GroupName = NAME_None;
 
 private:
 	UE_API void OnUserDefinedChordChanged(const FUICommandInfo& InCommandInfo);
diff --git a/Engine/Source/Editor/Experimental/EditorInteractiveToolsFramework/Public/ViewportInteractions/ViewportInteractionsBehaviorSource.h b/Engine/Source/Editor/Experimental/EditorInteractiveToolsFramework/Public/ViewportInteractions/ViewportInteractionsBehaviorSource.h
--- a/Engine/Source/Editor/Experimental/EditorInteractiveToolsFramework/Public/ViewportInteractions/ViewportInteractionsBehaviorSource.h
+++ b/Engine/Source/Editor/Experimental/EditorInteractiveToolsFramework/Public/ViewportInteractions/ViewportInteractionsBehaviorSource.h
@@ -53,14 +53,16 @@ EDITORINTERACTIVETOOLSFRAMEWORK_API void LOG(const TCHAR* InMessage);
 
 DECLARE_MULTICAST_DELEGATE_OneParam(FOnMouseLookingChanged, bool /* bIsMouseLooking*/)
 
-	/**
-	 * This class hosts a list of ITF Viewport Interactions and a UBehaviorSet containing the UInteractionBehavior
-	 * instances required by those Viewport Interactions.
-	 * It also acts as a behavior target, registering the status of Shift, Alt, Ctrl modifier keys. That and other
-	 * pieces of information can be accessed externally, so that e.g. Viewport Interactions can know about modifier
-	 * states without having to implement themselves what would be duplicate logic to handle those inputs. This class
-	 * can also be used to know whether the Viewport camera is being moved using the mouse,or to mark it as such.
-	 */
+    DECLARE_MULTICAST_DELEGATE_OneParam(FOnModifierKeyChanged, bool /* bIsDown*/)
+
+    /**
+     * This class hosts a list of ITF Viewport Interactions and a UBehaviorSet containing the UInteractionBehavior
+     * instances required by those Viewport Interactions.
+     * It also acts as a behavior target, registering the status of Shift, Alt, Ctrl modifier keys. That and other
+     * pieces of information can be accessed externally, so that e.g. Viewport Interactions can know about modifier
+     * states without having to implement themselves what would be duplicate logic to handle those inputs. This class
+     * can also be used to know whether the Viewport camera is being moved using the mouse,or to mark it as such.
+     */
 	UCLASS(Transient, MinimalAPI)
 class UViewportInteractionsBehaviorSource final
 	: public UObject
@@ -102,6 +104,11 @@ public:
 	 */
 	EDITORINTERACTIVETOOLSFRAMEWORK_API void Tick(float InDeltaTime) const;
 
+	/**
+	 * Renders the active tool on the specified View/Canvas
+	 */
+	EDITORINTERACTIVETOOLSFRAMEWORK_API void RenderTools(FViewportClient* InViewportClient, const FSceneView* InSceneView, FCanvas* InCanvas) const;
+
 	/**
 	 * Add the specified interactions to this Behavior Source
 	 */
@@ -219,13 +226,16 @@ public:
 
 	FEditorViewportClient* GetEditorViewportClient() const;
 
+	FEditorViewportClient* GetHoveredEditorViewportClient() const;
+
 	/**
 	 * Marks the specified interaction type as active (meaning a viewport interaction of that type is currently active)
+	 * Internally keeps count of how many interactions of this type are active
 	 */
 	void SetViewportInteractionActive(FName InViewportInteraction, bool bInActive);
 
 	/**
-	 * Checks if there is any active viewport interaction of the specified type
+	 * Checks if there is any active viewport interaction of the specified type (at least one)
 	 */
 	bool EDITORINTERACTIVETOOLSFRAMEWORK_API IsViewportInteractionActive(FName InViewportInteraction);
 
@@ -234,12 +244,31 @@ public:
 		return OnMouseLookingChangedDelegate;
 	}
 
+	FOnModifierKeyChanged& OnAltKeyStateChanged()
+	{
+		return OnAltKeyStateChangedDelegate;
+	}
+
+	FOnModifierKeyChanged& OnCtrlKeyStateChanged()
+	{
+		return OnCtrlKeyStateChangedDelegate;
+	}
+
+	FOnModifierKeyChanged& OnShiftKeyStateChanged()
+	{
+		return OnShiftKeyStateChangedDelegate;
+	}
+
 private:
 	/**
 	 * All Input Behaviors from registered Viewport Interactions should be added to this Behavior Set (done with ::AddInteraction() )
 	 */
 	UPROPERTY(Transient)
 	TObjectPtr<UInputBehaviorSet> BehaviorSet;
+	
+	/** Behaviors that are garuanteed to always be present */
+	UPROPERTY(Transient)
+	TArray<TObjectPtr<UInputBehavior>> CoreBehaviors;
 
 	/**
 	 * A collection of ITF-based Viewport Interactions
@@ -248,6 +277,10 @@ private:
 	UPROPERTY(Transient)
 	TMap<FName, TObjectPtr<UViewportInteraction>> ViewportInteractions;
 
+	void SetAltKeyState(bool bInDown);
+	void SetCtrlKeyState(bool bInDown);
+	void SetShiftKeyState(bool bInDown);
+
 	bool bIsShiftDown;
 	bool bIsCtrlDown;
 	bool bIsAltDown;
@@ -268,4 +301,8 @@ private:
 	TMap<FName, bool> ViewportInteractionsStatusMap;
 
 	FOnMouseLookingChanged OnMouseLookingChangedDelegate;
+
+	FOnModifierKeyChanged OnAltKeyStateChangedDelegate;
+	FOnModifierKeyChanged OnCtrlKeyStateChangedDelegate;
+	FOnModifierKeyChanged OnShiftKeyStateChangedDelegate;
 };
diff --git a/Engine/Source/Editor/Experimental/EditorInteractiveToolsFramework/Public/ViewportInteractions/ViewportInteractionsUtilities.h b/Engine/Source/Editor/Experimental/EditorInteractiveToolsFramework/Public/ViewportInteractions/ViewportInteractionsUtilities.h
--- a/Engine/Source/Editor/Experimental/EditorInteractiveToolsFramework/Public/ViewportInteractions/ViewportInteractionsUtilities.h
+++ b/Engine/Source/Editor/Experimental/EditorInteractiveToolsFramework/Public/ViewportInteractions/ViewportInteractionsUtilities.h
@@ -13,4 +13,6 @@ namespace UE::Editor::ViewportInteractions
  * e.g. Movement, Rotation, View Angle, Orbit, Panning, FOV, etc
  */
 EDITORINTERACTIVETOOLSFRAMEWORK_API void AddDefaultCameraMovementInteractions(UViewportInteractionsBehaviorSource* InInteractionsBehaviorSource);
+
+EDITORINTERACTIVETOOLSFRAMEWORK_API void AddDefaultDragToolsInteractions(UViewportInteractionsBehaviorSource* InInteractionsBehaviorSource);
 }
diff --git a/Engine/Source/Editor/LevelEditor/Private/SLevelEditor.cpp b/Engine/Source/Editor/LevelEditor/Private/SLevelEditor.cpp
--- a/Engine/Source/Editor/LevelEditor/Private/SLevelEditor.cpp
+++ b/Engine/Source/Editor/LevelEditor/Private/SLevelEditor.cpp
@@ -105,8 +105,7 @@ SLevelEditor::SLevelEditor()
 	if (UModeManagerInteractiveToolsContext* const ModeManagerInteractiveToolsContext =
 			SLevelEditor::GetEditorModeManager().GetInteractiveToolsContext())
 	{
-		ModeManagerInteractiveToolsContext->SetDragToolsEnabled(true);
-		ModeManagerInteractiveToolsContext->SetInputToolsEnabled(true);
+		ModeManagerInteractiveToolsContext->SetViewportInteractionsEnabled(true);
 
 		ModeManagerInteractiveToolsContext->OnRegisterViewportInteractions().AddRaw(this, &SLevelEditor::RegisterViewportInteractions
 	);
@@ -773,8 +772,12 @@ void SLevelEditor::RegisterViewportInteractions() const
 				ModeManagerInteractiveToolsContext->GetViewportInteractionsBehaviorSource())
 		{
 			InputToolsSource->Reset();
-
-			UE::Editor::ViewportInteractions::AddDefaultCameraMovementInteractions(InputToolsSource);
+			
+			if (ModeManagerInteractiveToolsContext->UsesViewportInteractions())
+			{
+				UE::Editor::ViewportInteractions::AddDefaultCameraMovementInteractions(InputToolsSource);
+				UE::Editor::ViewportInteractions::AddDefaultDragToolsInteractions(InputToolsSource);
+			}
 
 			// Left Mouse Click
 			InputToolsSource->AddInteraction<ULevelViewportClickSelection>();
diff --git a/Engine/Source/Editor/LevelEditor/Private/ViewportInteractions/LevelViewportClickSelection.cpp b/Engine/Source/Editor/LevelEditor/Private/ViewportInteractions/LevelViewportClickSelection.cpp
--- a/Engine/Source/Editor/LevelEditor/Private/ViewportInteractions/LevelViewportClickSelection.cpp
+++ b/Engine/Source/Editor/LevelEditor/Private/ViewportInteractions/LevelViewportClickSelection.cpp
@@ -100,7 +100,7 @@ ULevelViewportClickSelection::ULevelViewportClickSelection()
 
 void ULevelViewportClickSelection::ProcessClick_Internal(const FViewportClick& InViewportClick, HHitProxy* InHitProxy, FSceneView& View)
 {
-	UE::Editor::ViewportInteractions::LOG(*InteractionName->ToString());
+	UE::Editor::ViewportInteractions::LOG(*InteractionName.ToString());
 	UE::Editor::LevelViewportInteractions::ProcessClick(GetEditorViewportClient(), InViewportClick, InHitProxy);
 }
 
diff --git a/Engine/Source/Editor/LevelEditor/Private/ViewportToolbar/LevelEditorSubmenus.cpp b/Engine/Source/Editor/LevelEditor/Private/ViewportToolbar/LevelEditorSubmenus.cpp
--- a/Engine/Source/Editor/LevelEditor/Private/ViewportToolbar/LevelEditorSubmenus.cpp
+++ b/Engine/Source/Editor/LevelEditor/Private/ViewportToolbar/LevelEditorSubmenus.cpp
@@ -3187,7 +3187,7 @@ void ExtendTransformSubmenu(FName InTransformSubmenuName)
 			MarqueeMode.Visibility = TAttribute<bool>::CreateLambda([]()
 			{
 				// ITF Drag Tools are automatically enabled by new TRS Gizmos
-				return UEditorInteractiveGizmoManager::UsesNewTRSGizmos() && GLevelEditorModeTools().GetInteractiveToolsContext()->UsesDragTools();
+				return UEditorInteractiveGizmoManager::UsesNewTRSGizmos() && GLevelEditorModeTools().GetInteractiveToolsContext()->UsesViewportInteractions();
 			});
 		}
 	}
diff --git a/Engine/Source/Editor/Persona/Private/IPersonaEditorModeManager.cpp b/Engine/Source/Editor/Persona/Private/IPersonaEditorModeManager.cpp
--- a/Engine/Source/Editor/Persona/Private/IPersonaEditorModeManager.cpp
+++ b/Engine/Source/Editor/Persona/Private/IPersonaEditorModeManager.cpp
@@ -65,7 +65,6 @@ IPersonaEditorModeManager::IPersonaEditorModeManager() : FAssetEditorModeManager
 	UModeManagerInteractiveToolsContext* ToolsContext = GetInteractiveToolsContext();
 	UContextObjectStore* ContextObjectStore = ToolsContext->ContextObjectStore;
 	ContextObjectStore->AddContextObject(PersonaModeManagerContext.Get());
-	ToolsContext->SetDragToolsEnabled(true);
 }
 
 IPersonaEditorModeManager::~IPersonaEditorModeManager()
diff --git a/Engine/Source/Editor/UnrealEd/Private/DragTool_FrustumSelect.cpp b/Engine/Source/Editor/UnrealEd/Private/DragTool_FrustumSelect.cpp
--- a/Engine/Source/Editor/UnrealEd/Private/DragTool_FrustumSelect.cpp
+++ b/Engine/Source/Editor/UnrealEd/Private/DragTool_FrustumSelect.cpp
@@ -25,7 +25,7 @@
 #include "Subsystems/BrushEditingSubsystem.h"
 #include "LevelEditorSubsystem.h"
 #include "SceneView.h"
-#include "EditorDragTools/EditorFrustumSelect.h"
+#include "EditorDragTools/FrustumSelectInteraction.h"
 
 ///////////////////////////////////////////////////////////////////////////////
 //
diff --git a/Engine/Source/Editor/UnrealEd/Private/LevelEditorViewport.cpp b/Engine/Source/Editor/UnrealEd/Private/LevelEditorViewport.cpp
--- a/Engine/Source/Editor/UnrealEd/Private/LevelEditorViewport.cpp
+++ b/Engine/Source/Editor/UnrealEd/Private/LevelEditorViewport.cpp
@@ -115,10 +115,11 @@
 #include "Rendering/StaticLightingSystemInterface.h"
 #include "TextureResource.h"
 #include "ViewportClientNavigationHelper.h"
-#include "EditorDragTools/EditorDragToolBehaviorTarget.h"
 #include "Subsystems/PlacementSubsystem.h"
 #include "Subsystems/EditorElementSubsystem.h"
 #include "ViewportSelectionUtilities.h"
+#include "ViewportInteractions/ViewportInteraction.h"
+#include "ViewportInteractions/ViewportInteractionsBehaviorSource.h"
 
 DEFINE_LOG_CATEGORY(LogEditorViewport);
 
@@ -3282,7 +3283,7 @@ bool FLevelEditorViewportClient::InputWidgetDelta(FViewport* InViewport, EAxisLi
 				const bool MiddleMouseButtonDown = InViewport->KeyState(EKeys::MiddleMouseButton);
 
 				// Some functionality can be skipped if new gizmos and ITF drag tools are active
-				const bool bUseLegacyGizmoDragTools = !(ModeTools->UsesViewportITF() && UE::Editor::DragTools::UseEditorDragTools());
+				const bool bUseLegacyGizmoDragTools = !(ModeTools->UsesViewportITF() && UE::Editor::ViewportInteractions::UseEditorViewportInteractions());
 
 				// If duplicate dragging . . .
 				if (IsAltPressed() && (LeftMouseButtonDown || RightMouseButtonDown) && bUseLegacyGizmoDragTools)
diff --git a/Engine/Source/Editor/UnrealEd/Private/MouseDeltaTracker.cpp b/Engine/Source/Editor/UnrealEd/Private/MouseDeltaTracker.cpp
--- a/Engine/Source/Editor/UnrealEd/Private/MouseDeltaTracker.cpp
+++ b/Engine/Source/Editor/UnrealEd/Private/MouseDeltaTracker.cpp
@@ -10,10 +10,10 @@
 #include "EditorDragTools.h"
 #include "SnappingUtils.h"
 #include "UnrealWidget.h"
-#include "EditorDragTools/EditorDragToolBehaviorTarget.h"
 #include "EditorInteractiveGizmoManager.h"
 #include "Misc/AxisDisplayInfo.h"
 #include "EditorModeManager.h"
+#include "ViewportInteractions/ViewportInteractionsBehaviorSource.h"
 
 #define LOCTEXT_NAMESPACE "MouseDeltaTracker"
 
@@ -345,9 +345,18 @@ bool FMouseDeltaTracker::EndTracking(FEditorViewportClient* InViewportClient)
 
 void FMouseDeltaTracker::ConditionalBeginUsingDragTool( FEditorViewportClient* InViewportClient )
 {
-	if (UE::Editor::DragTools::UseEditorDragTools())
+	if (FEditorModeTools* ModeTools = InViewportClient->GetModeTools())
 	{
-		return;
+		if (UE::Editor::ViewportInteractions::UseEditorViewportInteractions())
+		{
+			if (UModeManagerInteractiveToolsContext* InputToolsBehaviorSource = ModeTools->GetInteractiveToolsContext())
+			{
+				if (InputToolsBehaviorSource->UsesViewportInteractions())
+				{
+					return;
+				}
+			}
+		}
 	}
 
 	const bool LeftMouseButtonDown = InViewportClient->Viewport->KeyState(EKeys::LeftMouseButton);
diff --git a/Engine/Source/Editor/UnrealEd/Private/Tools/EdModeInteractiveToolsContext.cpp b/Engine/Source/Editor/UnrealEd/Private/Tools/EdModeInteractiveToolsContext.cpp
--- a/Engine/Source/Editor/UnrealEd/Private/Tools/EdModeInteractiveToolsContext.cpp
+++ b/Engine/Source/Editor/UnrealEd/Private/Tools/EdModeInteractiveToolsContext.cpp
@@ -38,8 +38,7 @@
 
 #include "BaseGizmos/GizmoRenderingUtil.h"
 #include "UnrealClient.h"
-#include "EditorDragTools/DragToolsBehaviourSource.h"
-#include "EditorDragTools/EditorDragToolBehaviorTarget.h"
+#include "EditorDragTools/DragToolInteraction.h"
 #include "UObject/ObjectSaveContext.h"
 #include "ViewportInteractions/ViewportInteractionsBehaviorSource.h"
 
@@ -933,9 +932,9 @@ void UModeManagerInteractiveToolsContext::DrawHUD(FViewportClient* ViewportClien
 		EdModeContext->DrawHUD(ViewportClient, Viewport, View, Canvas);
 	}
 
-	if (DragToolsBehaviorSource)
+	if (ViewportInteractionsBehaviorSource)
 	{
-		DragToolsBehaviorSource->RenderTools(ViewportClient, View, Canvas);
+		ViewportInteractionsBehaviorSource->RenderTools(ViewportClient, View, Canvas);
 	}
 }
 
@@ -1353,10 +1352,6 @@ void UModeManagerInteractiveToolsContext::Initialize(IToolsContextQueriesAPI* Qu
 
 	GetParentEditorModeManager()->OnEditorModeIDChanged().AddUObject(this, &UModeManagerInteractiveToolsContext::OnEditorModeChanged);
 
-	// Delegates used to handle DragTools.EnableITFTools CVAR changes
-	UE::Editor::DragTools::OnEditorDragToolsActivated().AddUObject(this, &UModeManagerInteractiveToolsContext::RegisterDragTools);
-	UE::Editor::DragTools::OnEditorDragToolsDeactivated().AddUObject(this, &UModeManagerInteractiveToolsContext::UnregisterDragTools);
-
 	// Delegates used to handle InputTools.EnableITFTools CVAR changes
 	UE::Editor::ViewportInteractions::OnEditorViewportInteractionsActivated().AddUObject(this, &UModeManagerInteractiveToolsContext::RegisterViewportInteractions);
 	UE::Editor::ViewportInteractions::OnEditorViewportInteractionsDeactivated().AddUObject(this, &UModeManagerInteractiveToolsContext::UnregisterViewportInteractions);
@@ -1364,7 +1359,6 @@ void UModeManagerInteractiveToolsContext::Initialize(IToolsContextQueriesAPI* Qu
 
 void UModeManagerInteractiveToolsContext::Shutdown()
 {
-	UnregisterDragTools();
 	UnregisterViewportInteractions();
 
 	GetParentEditorModeManager()->OnEditorModeIDChanged().RemoveAll(this);
@@ -1440,65 +1434,23 @@ void UModeManagerInteractiveToolsContext::SetBypassAltModifier(bool bInBypassAlt
 	bBypassAltModifier = bInBypassAltModifier;
 }
 
-UViewportInteractionsBehaviorSource* UModeManagerInteractiveToolsContext::GetViewportInteractionsBehaviorSource()
+bool UModeManagerInteractiveToolsContext::UsesViewportInteractions() const
 {
-	if (!ViewportInteractionsBehaviorSource && UsesViewportInteractions() && UE::Editor::ViewportInteractions::UseEditorViewportInteractions())
-	{
-		RegisterViewportInteractions();
-	}
-
-	return ViewportInteractionsBehaviorSource.Get();
+	return bUsesViewportInteractions && UE::Editor::ViewportInteractions::UseEditorViewportInteractions();
 }
 
-void UModeManagerInteractiveToolsContext::RegisterDragTools()
+UViewportInteractionsBehaviorSource* UModeManagerInteractiveToolsContext::GetViewportInteractionsBehaviorSource()
 {
-	FEditorModeTools* EditorModeTools = GetParentEditorModeManager();
-	if (!EditorModeTools)
-	{
-		return;
-	}
-
-	// Ignore uninitialized Mode Tools missing a Viewport Client
-	if (!EditorModeTools->GetFocusedViewportClient())
+	if (!ViewportInteractionsBehaviorSource && UsesViewportInteractions())
 	{
-		return;
-	}
-
-	// This Mode Manager does not require Drag Tools
-	if (!UsesDragTools())
-	{
-		return;
+		RegisterViewportInteractions();
 	}
 
-	// Not initialized yet
-	if (!DragToolsBehaviorSource)
-	{
-		// Alt modifier input required by Drag Tools, but currently skipped by ITF in Tool Context.
-		// Let's make sure Alt keys are processed
-		SetBypassAltModifier(false);
-
-		DragToolsBehaviorSource = NewObject<UDragToolsBehaviorSource>();
-		DragToolsBehaviorSource->Initialize(this);
-		DragToolsBehaviorSource->RegisterSource();
-	}
+	return ViewportInteractionsBehaviorSource.Get();
 }
 
 void UModeManagerInteractiveToolsContext::RegisterViewportInteractions()
 {
-	UE_LOG(LogTemp,Log,TEXT("RegisterViewportInteractions"));
-	FEditorModeTools* EditorModeTools = GetParentEditorModeManager();
-	if (!EditorModeTools)
-	{
-		return;
-	}
-
-	// This Mode Manager does not require Input Tools
-	if (!UsesViewportInteractions())
-	{
-		return;
-	}
-
-	// Not initialized yet
 	if (!ViewportInteractionsBehaviorSource)
 	{
 		SetBypassAltModifier(false);
@@ -1511,19 +1463,6 @@ void UModeManagerInteractiveToolsContext::RegisterViewportInteractions()
 	}
 }
 
-void UModeManagerInteractiveToolsContext::UnregisterDragTools()
-{
-	if (DragToolsBehaviorSource)
-	{
-		// Revert to previous state, with Alt Modifier not being processed
-		SetBypassAltModifier(true);
-		DragToolsBehaviorSource->DeregisterSource();
-		DragToolsBehaviorSource = nullptr;
-
-		UToolMenus::UnregisterOwner(this);
-	}
-}
-
 void UModeManagerInteractiveToolsContext::UnregisterViewportInteractions()
 {
 	if (ViewportInteractionsBehaviorSource)
@@ -1532,23 +1471,13 @@ void UModeManagerInteractiveToolsContext::UnregisterViewportInteractions()
 		SetBypassAltModifier(true);
 		ViewportInteractionsBehaviorSource->DeregisterBehaviorSources();
 		ViewportInteractionsBehaviorSource = nullptr;
+		
+		UToolMenus::UnregisterOwner(this);
 	}
 }
 
 void UModeManagerInteractiveToolsContext::OnEditorModeChanged(const FEditorModeID& InModeID, bool bInIsEnteringMode)
 {
-	if (UE::Editor::DragTools::UseEditorDragTools())
-	{
-		if (bInIsEnteringMode)
-		{
-			RegisterDragTools();
-		}
-		else
-		{
-			UnregisterDragTools();
-		}
-	}
-
 	if (UE::Editor::ViewportInteractions::UseEditorViewportInteractions())
 	{
 		if (bInIsEnteringMode)
diff --git a/Engine/Source/Editor/UnrealEd/Public/Tools/EdModeInteractiveToolsContext.h b/Engine/Source/Editor/UnrealEd/Public/Tools/EdModeInteractiveToolsContext.h
--- a/Engine/Source/Editor/UnrealEd/Public/Tools/EdModeInteractiveToolsContext.h
+++ b/Engine/Source/Editor/UnrealEd/Public/Tools/EdModeInteractiveToolsContext.h
@@ -39,7 +39,6 @@ class ILevelEditor;
 class IToolsContextQueriesAPI;
 class IToolsContextRenderAPI;
 class IToolsContextTransactionsAPI;
-class UDragToolsBehaviorSource;
 class UEdModeInteractiveToolsContext;
 class UGizmoViewContext;
 class UInputRouter;
@@ -373,15 +372,13 @@ public:
 	/**
 	 * Mark Editor Modes Tools for Drag Tools Support
 	 */
-	void SetDragToolsEnabled(bool bInEnabled)
-	{
-		bUsesDragTools = bInEnabled;
-	}
+	UE_DEPRECATED(5.8, "Drag tools have been merged with VIewport Interactions")
+	UE_API void SetDragToolsEnabled(bool bInEnabled) { SetViewportInteractionsEnabled(bInEnabled); }
 
 	/**
 	 * Enable ITF Input tools support (e.g. Viewport Interactions as camera movement and selection)
 	 */
-	void SetInputToolsEnabled(bool bInEnabled)
+	void SetViewportInteractionsEnabled(bool bInEnabled)
 	{
 		bUsesViewportInteractions = bInEnabled;
 	}
@@ -390,18 +387,13 @@ public:
 	 * Are Drag Tools supported? e.g. Marquee Select
 	 * Enable/Disable support by using SetDragToolsEnabled(...)
 	 */
-	bool UsesDragTools() const
-	{
-		return bUsesDragTools;
-	}
+	UE_DEPRECATED(5.8, "Drag tools have been merged with VIewport Interactions")
+	UE_API bool UsesDragTools() const { return UsesViewportInteractions(); } 
 
 	/**
 	 * Are Viewport Interactions supported? e.g. camera movements, commands, clicking, etc
 	 */
-	bool UsesViewportInteractions() const
-	{
-		return bUsesViewportInteractions;
-	}
+	UE_API bool UsesViewportInteractions() const;
 
 	/**
 	 * Return the current Input Tools Behavior Source - nullptr if not enabled
@@ -419,7 +411,8 @@ protected:
 	/**
 	 * Activating Drag Tools if current Mode Manager requires them
 	 */
-	UE_API void RegisterDragTools();
+	UE_DEPRECATED(5.8, "Drag tools have been merged with viewport interactions.")
+	UE_API void RegisterDragTools() {};
 
 	/**
 	 * Activating Viewport Interactions if current Mode Manager requires them
@@ -429,7 +422,8 @@ protected:
 	/**
 	 * Deactivating Drag Tools
 	 */
-	UE_API void UnregisterDragTools();
+	UE_DEPRECATED(5.8, "Drag tools have been merged with viewport interactions.")
+	UE_API void UnregisterDragTools() {};
 
 	/**
 	 * Deactivating Viewport Interactions
@@ -459,16 +453,9 @@ protected:
 
 	//~End ITF Alt processing Bypass
 
-	/**
-	 *
-	 */
-	UPROPERTY(Transient)
-	TObjectPtr<UDragToolsBehaviorSource> DragToolsBehaviorSource;
-
 	UPROPERTY(Transient)
 	TObjectPtr<UViewportInteractionsBehaviorSource> ViewportInteractionsBehaviorSource;
 
-	bool bUsesDragTools = false;
 	bool bUsesViewportInteractions = false;
 };
 
