[ITF Viewport Interactions] Fix drag tools priority

diff --git a/Engine/Source/Editor/Experimental/EditorInteractiveToolsFramework/Public/EditorDragTools/DragToolInteraction.h b/Engine/Source/Editor/Experimental/EditorInteractiveToolsFramework/Public/EditorDragTools/DragToolInteraction.h
--- a/Engine/Source/Editor/Experimental/EditorInteractiveToolsFramework/Public/EditorDragTools/DragToolInteraction.h
+++ b/Engine/Source/Editor/Experimental/EditorInteractiveToolsFramework/Public/EditorDragTools/DragToolInteraction.h
@@ -51,14 +51,11 @@ static constexpr int HANDLED_BEFORE_VIEWPORT_INTERACTIONS =
 // Conversely, tools which are activated with just mouse buttons input (no key modifiers, mouse drag behavior only) need to be processed after Gizmo Behaviors
 static constexpr int HANDLED_AFTER_GIZMO_PRIORITY = FInputCapturePriority::DEFAULT_GIZMO_PRIORITY + 1;
 
-static constexpr int DRAG_TOOLS_DEFAULT_PRIORITY = HANDLED_BEFORE_VIEWPORT_INTERACTIONS - 10;
-static constexpr int DRAG_TOOLS_LOW_PRIORITY = HANDLED_BEFORE_VIEWPORT_INTERACTIONS - 1;
-
-static constexpr int DUPLICATE_SELECTION_PRIORITY = DRAG_TOOLS_DEFAULT_PRIORITY;
-static constexpr int FRUSTUM_SELECT_PRIORITY = DRAG_TOOLS_DEFAULT_PRIORITY;
-static constexpr int VIEWPORT_CHANGE_PRIORITY = DRAG_TOOLS_DEFAULT_PRIORITY;
+static constexpr int DUPLICATE_SELECTION_PRIORITY = ViewportInteractions::DRAG_TOOLS_DEFAULT_PRIORITY;
+static constexpr int FRUSTUM_SELECT_PRIORITY = ViewportInteractions::DRAG_TOOLS_DEFAULT_PRIORITY;
+static constexpr int VIEWPORT_CHANGE_PRIORITY = ViewportInteractions::DRAG_TOOLS_DEFAULT_PRIORITY;
 static constexpr int BOX_SELECT_PRIORITY = HANDLED_AFTER_GIZMO_PRIORITY;
-static constexpr int MEASURE_PRIORITY = DRAG_TOOLS_LOW_PRIORITY;
+static constexpr int MEASURE_PRIORITY = ViewportInteractions::DRAG_TOOLS_DEFAULT_PRIORITY;
 
 } // namespace UE::Editor::DragTools
 
