[ITF Viewport Interactions] Add method to remove interactions

diff --git a/Engine/Source/Editor/Experimental/EditorInteractiveToolsFramework/Private/ViewportInteractions/ViewportInteractionsBehaviorSource.cpp b/Engine/Source/Editor/Experimental/EditorInteractiveToolsFramework/Private/ViewportInteractions/ViewportInteractionsBehaviorSource.cpp
--- a/Engine/Source/Editor/Experimental/EditorInteractiveToolsFramework/Private/ViewportInteractions/ViewportInteractionsBehaviorSource.cpp
+++ b/Engine/Source/Editor/Experimental/EditorInteractiveToolsFramework/Private/ViewportInteractions/ViewportInteractionsBehaviorSource.cpp
@@ -333,6 +333,25 @@ UViewportInteraction* UViewportInteractionsBehaviorSource::AddInteraction(const
 	return Interaction;
 }
 
+void UViewportInteractionsBehaviorSource::RemoveInteraction(const FName InInteractionName, bool bInReregister)
+{
+	if (UViewportInteraction* Interaction = ViewportInteractions.FindRef(InInteractionName))
+	{
+		UE::Editor::ViewportInteractions::LOG(*FString::Printf(TEXT("Removing Viewport Interaction \"%s\""), *Interaction->GetInteractionName().ToString()));
+
+		ViewportInteractions.Remove(Interaction->GetInteractionName());
+
+		BehaviorSet->RemoveBySource(Interaction);
+
+		if (bInReregister)
+		{
+			// Refreshing registered behaviors
+			DeregisterBehaviorSources();
+			RegisterBehaviorSources();
+		}
+	}
+}
+
 void UViewportInteractionsBehaviorSource::Reset()
 {
 	DeregisterBehaviorSources();
diff --git a/Engine/Source/Editor/Experimental/EditorInteractiveToolsFramework/Public/ViewportInteractions/ViewportInteractionsBehaviorSource.h b/Engine/Source/Editor/Experimental/EditorInteractiveToolsFramework/Public/ViewportInteractions/ViewportInteractionsBehaviorSource.h
--- a/Engine/Source/Editor/Experimental/EditorInteractiveToolsFramework/Public/ViewportInteractions/ViewportInteractionsBehaviorSource.h
+++ b/Engine/Source/Editor/Experimental/EditorInteractiveToolsFramework/Public/ViewportInteractions/ViewportInteractionsBehaviorSource.h
@@ -124,6 +124,11 @@ public:
 		const UClass* InInteractionClass, bool bInReregister = false
 	);
 
+	/**
+	 * Removes the specified interaction from the active ones.
+	 */
+	EDITORINTERACTIVETOOLSFRAMEWORK_API void RemoveInteraction(const FName InInteractionName, bool bInReregister = false);
+
 	/**
 	 * De-registers behaviors, removes viewport interactions.
 	 * Behaviors added by Initialize will be kept around.
