[TRS Gizmos] Allow the 5.8 TRS gizmos to be enabled & disabled without restarting.
- The TRS gizmo builder can fall back to the animation gizmo builder.
- Toggling the TRS gizmo setting does the necessary messaging to toggle back and forth between the animation gizmos and the new 5.8 gizmos.
- Simplifies the logic of `UEditorInteractiveGizmoManager::UpdateActiveEditorGizmos()` and ensures that gizmos that aren't in the right visibility state are corrected.

#rnx

diff --git a/Engine/Source/Editor/Experimental/EditorInteractiveToolsFramework/Private/EditorInteractiveGizmoManager.cpp b/Engine/Source/Editor/Experimental/EditorInteractiveToolsFramework/Private/EditorInteractiveGizmoManager.cpp
--- a/Engine/Source/Editor/Experimental/EditorInteractiveToolsFramework/Private/EditorInteractiveGizmoManager.cpp
+++ b/Engine/Source/Editor/Experimental/EditorInteractiveToolsFramework/Private/EditorInteractiveGizmoManager.cpp
@@ -277,40 +277,32 @@ bool UEditorInteractiveGizmoManager::GetShowEditorGizmosForView(IToolsContextRen
 
 void UEditorInteractiveGizmoManager::UpdateActiveEditorGizmos()
 {
-	const bool bEnableEditorGizmos = UsesNewTRSGizmos();
-	if (!bEnableEditorGizmos)
+	if (!UsesNewTRSGizmos())
 	{
-		if (bShowEditorGizmos)
+		if (UTransformGizmo* Gizmo = FindDefaultTransformGizmo())
 		{
-			if (UTransformGizmo* Gizmo = FindDefaultTransformGizmo())
-			{
-				DestroyGizmo(Gizmo);
-			}
-			DestroyAllEditorGizmos();
+			DestroyGizmo(Gizmo);
 		}
+		DestroyAllEditorGizmos();
 		
 		bShowEditorGizmos = false;
 		return;
 	}
 	
-	
-	const bool bEditorModeToolsSupportsWidgetDrawing = EditorModeManager ? EditorModeManager->GetShowWidget() : true;
-	const bool bNewShowEditorGizmos = bEditorModeToolsSupportsWidgetDrawing && bEnableEditorGizmos;
+	const bool bGizmosShouldBeVisible = EditorModeManager ? EditorModeManager->GetShowWidget() : true;
 
-	if (bShowEditorGizmos != bNewShowEditorGizmos)
+	if (UTransformGizmo* Gizmo = FindDefaultTransformGizmo())
 	{
-		bShowEditorGizmos = bNewShowEditorGizmos;
-
-		if (UTransformGizmo* Gizmo = FindDefaultTransformGizmo())
+		if (Gizmo->bVisible != bGizmosShouldBeVisible)
 		{
-			Gizmo->SetVisibility(bShowEditorGizmos ? bEditorModeToolsSupportsWidgetDrawing : false);
-
-			if (!bEnableEditorGizmos)
-			{
-				DestroyGizmo(Gizmo);	
-			}
+			Gizmo->SetVisibility(bGizmosShouldBeVisible);
 		}
+	}
 
+	if (bShowEditorGizmos != bGizmosShouldBeVisible)
+	{
+		bShowEditorGizmos = bGizmosShouldBeVisible;
+		
 		if (!bShowEditorGizmos)
 		{
 			DestroyAllEditorGizmos();
diff --git a/Engine/Source/Editor/Experimental/EditorInteractiveToolsFramework/Public/EditorGizmos/EditorTransformGizmoBuilder.h b/Engine/Source/Editor/Experimental/EditorInteractiveToolsFramework/Public/EditorGizmos/EditorTransformGizmoBuilder.h
--- a/Engine/Source/Editor/Experimental/EditorInteractiveToolsFramework/Public/EditorGizmos/EditorTransformGizmoBuilder.h
+++ b/Engine/Source/Editor/Experimental/EditorInteractiveToolsFramework/Public/EditorGizmos/EditorTransformGizmoBuilder.h
@@ -32,7 +32,7 @@ public:
 	// If set, this function will be passed to UTransformGizmo instances to override the default material and display size.
 	TFunction<const FGizmoCustomization()> CustomizationFunction;
 
-private:
+protected:
 	static UE_API const UEditorTransformGizmoContextObject* GetTransformGizmoContext(const FToolBuilderState& InSceneState);
 };
 
