[ITF Interactions] Interactions can all specify which modes they do not support

Used to be a Drag Tools feature - moved to Viewport Interactions base class.

diff --git a/Engine/Source/Editor/Experimental/EditorInteractiveToolsFramework/Private/EditorDragTools/BoxSelectInteraction.cpp b/Engine/Source/Editor/Experimental/EditorInteractiveToolsFramework/Private/EditorDragTools/BoxSelectInteraction.cpp
--- a/Engine/Source/Editor/Experimental/EditorInteractiveToolsFramework/Private/EditorDragTools/BoxSelectInteraction.cpp
+++ b/Engine/Source/Editor/Experimental/EditorInteractiveToolsFramework/Private/EditorDragTools/BoxSelectInteraction.cpp
@@ -81,12 +81,9 @@ FInputRayHit UBoxSelectInteraction::CanBeginClickDragSequence(const FInputDevice
 {
 	if (FEditorViewportClient* const EditorViewportClient = GetEditorViewportClient())
 	{
-		if (IsCurrentModeSupported())
+		if (EditorViewportClient->IsOrtho())
 		{
-			if (EditorViewportClient->IsOrtho())
-			{
-				return FInputRayHit(TNumericLimits<float>::Max());
-			}
+			return FInputRayHit(TNumericLimits<float>::Max());
 		}
 	}
 
@@ -160,6 +157,12 @@ void UBoxSelectInteraction::OnDrag(const FInputDeviceRay& InDragPos)
 		return;
 	}
 
+	FEditorModeTools* const ModeTools = EditorViewportClient->GetModeTools();
+	if (!ModeTools)
+	{
+		return;
+	}
+
 	FIntPoint MousePos;
 	EditorViewportClient->Viewport->GetMousePos(MousePos);
 
@@ -237,6 +240,12 @@ void UBoxSelectInteraction::OnDragEnd(const FInputDeviceRay& InReleasePos)
 	{
 		if (FEditorViewportClient* const EditorViewportClient = GetEditorViewportClient())
 		{
+			FEditorModeTools* const ModeTools = EditorViewportClient->GetModeTools();
+			if (!ModeTools)
+			{
+				return;
+			}
+
 			// Workaround to use while ITF and legacy input coexist.
 			// Editor Viewport Client processes ITF input before anything else.
 			// This means that, while ITF and legacy input coexist, simple clicks will be "stolen" by ITF Box Select, which
@@ -402,7 +411,7 @@ void UBoxSelectInteraction::OnTerminateDragSequence()
 	UDragToolInteraction::OnTerminateDragSequence();
 }
 
-TArray<FEditorModeID> UBoxSelectInteraction::GetUnsupportedModes()
+TArray<FEditorModeID> UBoxSelectInteraction::GetUnsupportedModes() const
 {
 	const TArray<FEditorModeID> UnsupportedModes = { FBuiltinEditorModes::EM_Landscape, FBuiltinEditorModes::EM_Foliage };
 	return UnsupportedModes;
diff --git a/Engine/Source/Editor/Experimental/EditorInteractiveToolsFramework/Private/EditorDragTools/DragToolInteraction.cpp b/Engine/Source/Editor/Experimental/EditorInteractiveToolsFramework/Private/EditorDragTools/DragToolInteraction.cpp
--- a/Engine/Source/Editor/Experimental/EditorInteractiveToolsFramework/Private/EditorDragTools/DragToolInteraction.cpp
+++ b/Engine/Source/Editor/Experimental/EditorInteractiveToolsFramework/Private/EditorDragTools/DragToolInteraction.cpp
@@ -103,7 +103,6 @@ FOnOnViewportChangeToolToggleDelegate& OnViewportChangeToolDeactivated()
 
 UDragToolInteraction::UDragToolInteraction()
 	: bConvertDelta(false)
-	, ModeTools(nullptr)
 	, Start()
 	, End()
 	, bUseSnapping(false)
@@ -117,10 +116,12 @@ void UDragToolInteraction::Initialize(UViewportInteractionsBehaviorSource* const
 {
 	UViewportInteraction::Initialize(InViewportInteractionsBehaviorSource);
 
-	if (FEditorViewportClient* EditorViewportClient = GetEditorViewportClient())
+	if (const FEditorViewportClient* const EditorViewportClient = GetEditorViewportClient())
 	{
-		ModeTools = EditorViewportClient->GetModeTools();
-		EditorViewportClientProxy = IEditorViewportClientProxy::CreateViewportClientProxy(ModeTools->GetInteractiveToolsContext());
+		if (const FEditorModeTools* const ModeTools = EditorViewportClient->GetModeTools())
+		{
+			EditorViewportClientProxy = IEditorViewportClientProxy::CreateViewportClientProxy(ModeTools->GetInteractiveToolsContext());
+		}
 	}
 }
 
@@ -169,19 +170,3 @@ void UDragToolInteraction::OnTerminateDragSequence()
 		EditorViewportClient->Invalidate(true, false);
 	}
 }
-
-bool UDragToolInteraction::IsCurrentModeSupported()
-{
-	if (ModeTools)
-	{
-		for (const FEditorModeID UnsupportedMode : GetUnsupportedModes())
-		{
-			if (!ModeTools->EnsureNotInMode(UnsupportedMode))
-			{
-				return false;
-			}
-		}
-	}
-
-	return true;
-}
diff --git a/Engine/Source/Editor/Experimental/EditorInteractiveToolsFramework/Private/EditorDragTools/FrustumSelectInteraction.cpp b/Engine/Source/Editor/Experimental/EditorInteractiveToolsFramework/Private/EditorDragTools/FrustumSelectInteraction.cpp
--- a/Engine/Source/Editor/Experimental/EditorInteractiveToolsFramework/Private/EditorDragTools/FrustumSelectInteraction.cpp
+++ b/Engine/Source/Editor/Experimental/EditorInteractiveToolsFramework/Private/EditorDragTools/FrustumSelectInteraction.cpp
@@ -155,9 +155,7 @@ UFrustumSelectInteraction::UFrustumSelectInteraction()
 
 FInputRayHit UFrustumSelectInteraction::CanBeginClickDragSequence(const FInputDeviceRay& InPressPos)
 {
-	return IsCurrentModeSupported()
-			 ? FInputRayHit(TNumericLimits<float>::Max()) // bHit is true. Depth is max to lose the standard tiebreaker.
-			 : FInputRayHit();
+	return FInputRayHit(TNumericLimits<float>::Max()); // bHit is true. Depth is max to lose the standard tiebreaker.
 }
 
 void UFrustumSelectInteraction::OnDragStart(const FInputDeviceRay& InPressPos)
@@ -215,6 +213,12 @@ void UFrustumSelectInteraction::OnDragEnd(const FInputDeviceRay& InReleasePos)
 		return;
 	}
 
+	FEditorModeTools* ModeTools = EditorViewportClient->GetModeTools();
+	if (!ModeTools)
+	{
+		return;
+	}
+
 	UBrushEditingSubsystem* BrushSubsystem = GEditor->GetEditorSubsystem<UBrushEditingSubsystem>();
 	const bool bGeometryMode = BrushSubsystem ? BrushSubsystem->IsGeometryEditorModeActive() : false;
 
@@ -410,7 +414,7 @@ void UFrustumSelectInteraction::OnTerminateDragSequence()
 	Start = End = FVector::ZeroVector;
 }
 
-TArray<FEditorModeID> UFrustumSelectInteraction::GetUnsupportedModes()
+TArray<FEditorModeID> UFrustumSelectInteraction::GetUnsupportedModes() const
 {
 	const TArray<FEditorModeID> UnsupportedModes = { FBuiltinEditorModes::EM_Landscape, FBuiltinEditorModes::EM_Foliage };
 	return UnsupportedModes;
diff --git a/Engine/Source/Editor/Experimental/EditorInteractiveToolsFramework/Private/EditorDragTools/MeasureToolInteraction.cpp b/Engine/Source/Editor/Experimental/EditorInteractiveToolsFramework/Private/EditorDragTools/MeasureToolInteraction.cpp
--- a/Engine/Source/Editor/Experimental/EditorInteractiveToolsFramework/Private/EditorDragTools/MeasureToolInteraction.cpp
+++ b/Engine/Source/Editor/Experimental/EditorInteractiveToolsFramework/Private/EditorDragTools/MeasureToolInteraction.cpp
@@ -107,13 +107,10 @@ FInputRayHit UMeasureToolInteraction::CanBeginClickDragSequence(const FInputDevi
 {
 	if (FEditorViewportClient* const EditorViewportClient = GetEditorViewportClient())
 	{
-		if (IsCurrentModeSupported())
+		// Measure Tool is for Orthographic viewports only
+		if (EditorViewportClient->IsOrtho())
 		{
-			// Measure Tool is for Orthographic viewports only
-			if (EditorViewportClient->IsOrtho())
-			{
-				return FInputRayHit(TNumericLimits<float>::Max());
-			}
+			return FInputRayHit(TNumericLimits<float>::Max());
 		}
 	}
 
diff --git a/Engine/Source/Editor/Experimental/EditorInteractiveToolsFramework/Private/ViewportInteractions/ViewportInteraction.cpp b/Engine/Source/Editor/Experimental/EditorInteractiveToolsFramework/Private/ViewportInteractions/ViewportInteraction.cpp
--- a/Engine/Source/Editor/Experimental/EditorInteractiveToolsFramework/Private/ViewportInteractions/ViewportInteraction.cpp
+++ b/Engine/Source/Editor/Experimental/EditorInteractiveToolsFramework/Private/ViewportInteractions/ViewportInteraction.cpp
@@ -1,6 +1,8 @@
 // Copyright Epic Games, Inc. All Rights Reserved.
 
 #include "ViewportInteractions/ViewportInteraction.h"
+
+#include "EditorModeManager.h"
 #include "Framework/Commands/InputBindingManager.h"
 #include "Framework/Commands/UICommandInfo.h"
 #include "ViewportInteractions/ViewportInteractionsBehaviorSource.h"
@@ -185,3 +187,19 @@ void UViewportInteraction::RegisterInputBehaviors(TArray<UInputBehavior*> InBeha
 		RegisterInputBehavior(InputBehavior);
 	}
 }
+
+bool UViewportInteraction::IsCurrentModeSupported(const FEditorModeTools* InModeTools) const
+{
+	if (InModeTools)
+	{
+		for (const FEditorModeID UnsupportedMode : GetUnsupportedModes())
+		{
+			if (InModeTools->IsModeActive(UnsupportedMode))
+			{
+				return false;
+			}
+		}
+	}
+
+	return true;
+}
diff --git a/Engine/Source/Editor/Experimental/EditorInteractiveToolsFramework/Private/ViewportInteractions/ViewportInteractionsBehaviorSource.cpp b/Engine/Source/Editor/Experimental/EditorInteractiveToolsFramework/Private/ViewportInteractions/ViewportInteractionsBehaviorSource.cpp
--- a/Engine/Source/Editor/Experimental/EditorInteractiveToolsFramework/Private/ViewportInteractions/ViewportInteractionsBehaviorSource.cpp
+++ b/Engine/Source/Editor/Experimental/EditorInteractiveToolsFramework/Private/ViewportInteractions/ViewportInteractionsBehaviorSource.cpp
@@ -306,13 +306,32 @@ UViewportInteraction* UViewportInteractionsBehaviorSource::AddInteraction(const
 		Initialize(EditorInteractiveToolsContextWeak.Get());
 	}
 
-	UViewportInteraction* Interaction = NewObject<UViewportInteraction>(GetTransientPackage(), InInteractionClass);
-
-	if (ViewportInteractions.Contains(Interaction->GetInteractionName()))
+	// Make sure the required interaction should be added
 	{
-		return nullptr;
+		const UViewportInteraction* DefaultInteraction = GetDefault<UViewportInteraction>(InInteractionClass);
+
+		if (ViewportInteractions.Contains(DefaultInteraction->GetInteractionName()))
+		{
+			UE_LOG(LogTemp, Warning, TEXT("Viewport Interactions: %s interaction has already been added."), *DefaultInteraction->GetInteractionName().ToString());
+			return nullptr;
+		}
+
+		FEditorViewportClient* EditorViewportClient = GetEditorViewportClient();
+		if (!EditorViewportClient)
+		{
+			UE_LOG(LogTemp, Warning, TEXT("Viewport Interactions Behavior Source has no Editor Viewport Client available!"));
+			return nullptr;
+		}
+
+		// Current mode does not support this interaction: skip it
+		if (!DefaultInteraction->IsCurrentModeSupported(EditorViewportClient->GetModeTools()))
+		{
+			return nullptr;
+		}
 	}
 
+	UViewportInteraction* Interaction = NewObject<UViewportInteraction>(GetTransientPackage(), InInteractionClass);
+
 	Interaction->Initialize(this);
 
 	UE::Editor::ViewportInteractions::LOG(*FString::Printf(TEXT("Initializing Viewport Interaction \"%s\""), *InInteractionClass->GetName()));
diff --git a/Engine/Source/Editor/Experimental/EditorInteractiveToolsFramework/Public/EditorDragTools/BoxSelectInteraction.h b/Engine/Source/Editor/Experimental/EditorInteractiveToolsFramework/Public/EditorDragTools/BoxSelectInteraction.h
--- a/Engine/Source/Editor/Experimental/EditorInteractiveToolsFramework/Public/EditorDragTools/BoxSelectInteraction.h
+++ b/Engine/Source/Editor/Experimental/EditorInteractiveToolsFramework/Public/EditorDragTools/BoxSelectInteraction.h
@@ -27,7 +27,7 @@ public:
 	//~ End IViewportClickDragBehaviorTarget
 
 protected:
-	virtual TArray<FEditorModeID> GetUnsupportedModes() override;
+	virtual TArray<FEditorModeID> GetUnsupportedModes() const override;
 
 private:
 	/**
diff --git a/Engine/Source/Editor/Experimental/EditorInteractiveToolsFramework/Public/EditorDragTools/DragToolInteraction.h b/Engine/Source/Editor/Experimental/EditorInteractiveToolsFramework/Public/EditorDragTools/DragToolInteraction.h
--- a/Engine/Source/Editor/Experimental/EditorInteractiveToolsFramework/Public/EditorDragTools/DragToolInteraction.h
+++ b/Engine/Source/Editor/Experimental/EditorInteractiveToolsFramework/Public/EditorDragTools/DragToolInteraction.h
@@ -107,12 +107,6 @@ public:
 	}
 
 protected:
-	virtual TArray<FEditorModeID> GetUnsupportedModes() { return {}; }
-
-	bool IsCurrentModeSupported();
-
-	FEditorModeTools* ModeTools;
-
 	/** The start/end location of the current drag. */
 	FVector Start, End;
 
diff --git a/Engine/Source/Editor/Experimental/EditorInteractiveToolsFramework/Public/EditorDragTools/FrustumSelectInteraction.h b/Engine/Source/Editor/Experimental/EditorInteractiveToolsFramework/Public/EditorDragTools/FrustumSelectInteraction.h
--- a/Engine/Source/Editor/Experimental/EditorInteractiveToolsFramework/Public/EditorDragTools/FrustumSelectInteraction.h
+++ b/Engine/Source/Editor/Experimental/EditorInteractiveToolsFramework/Public/EditorDragTools/FrustumSelectInteraction.h
@@ -37,7 +37,7 @@ public:
 	//~ End IViewportClickDragBehaviorTarget
 
 protected:
-	virtual TArray<FEditorModeID> GetUnsupportedModes() override;
+	virtual TArray<FEditorModeID> GetUnsupportedModes() const override;
 
 private:
 	/**
diff --git a/Engine/Source/Editor/Experimental/EditorInteractiveToolsFramework/Public/ViewportInteractions/ViewportInteraction.h b/Engine/Source/Editor/Experimental/EditorInteractiveToolsFramework/Public/ViewportInteractions/ViewportInteraction.h
--- a/Engine/Source/Editor/Experimental/EditorInteractiveToolsFramework/Public/ViewportInteractions/ViewportInteraction.h
+++ b/Engine/Source/Editor/Experimental/EditorInteractiveToolsFramework/Public/ViewportInteractions/ViewportInteraction.h
@@ -127,7 +127,14 @@ public:
 
 	//~ Begin UObject
 	UE_API virtual void BeginDestroy() override;
-	//~ Edn UObject
+	//~ End UObject
+
+	/**
+	 * Checks if the provided Mode is supported by this interaction.
+	 * Used to know whether this interaction should be added or not to the current Mode.
+	 * Should be called on the class CDO to avoid instantiating interactions which might not be used.
+	 */
+	bool IsCurrentModeSupported(const FEditorModeTools* InModeTools) const;
 
 protected:
 	UE_API FEditorViewportClient* GetEditorViewportClient() const;
@@ -151,6 +158,8 @@ protected:
 	UE_API void RegisterInputBehavior(UInputBehavior* InBehavior);
 	UE_API void RegisterInputBehaviors(TArray<UInputBehavior*> InBehaviors);
 
+	virtual TArray<FEditorModeID> GetUnsupportedModes() const { return {}; }
+
 	UPROPERTY()
 	TWeakObjectPtr<UViewportInteractionsBehaviorSource> ViewportInteractionsBehaviorSource;
 
