[ITF Viewport Interactions] Camera Issues with Landscape and Foliage Painting

Disable MoveYaw when in Landscape and Foliage modes.

diff --git a/Engine/Source/Editor/Experimental/EditorInteractiveToolsFramework/Private/ViewportInteractions/ViewportMoveYawInteraction.cpp b/Engine/Source/Editor/Experimental/EditorInteractiveToolsFramework/Private/ViewportInteractions/ViewportMoveYawInteraction.cpp
--- a/Engine/Source/Editor/Experimental/EditorInteractiveToolsFramework/Private/ViewportInteractions/ViewportMoveYawInteraction.cpp
+++ b/Engine/Source/Editor/Experimental/EditorInteractiveToolsFramework/Private/ViewportInteractions/ViewportMoveYawInteraction.cpp
@@ -1,6 +1,8 @@
 // Copyright Epic Games, Inc. All Rights Reserved.
 
 #include "ViewportInteractions/ViewportMoveYawInteraction.h"
+
+#include "EditorModes.h"
 #include "BaseBehaviors/ClickDragBehavior.h"
 #include "EditorViewportClient.h"
 #include "Settings/LevelEditorViewportSettings.h"
@@ -79,3 +81,9 @@ bool UViewportMoveYawInteraction::CanBeActivated(const FInputDeviceState& InInpu
 
 	return false;
 }
+
+TArray<FEditorModeID> UViewportMoveYawInteraction::GetUnsupportedModes() const
+{
+	const TArray<FEditorModeID> UnsupportedModes = { FBuiltinEditorModes::EM_Landscape, FBuiltinEditorModes::EM_Foliage };
+	return UnsupportedModes;
+}
diff --git a/Engine/Source/Editor/Experimental/EditorInteractiveToolsFramework/Public/ViewportInteractions/ViewportMoveYawInteraction.h b/Engine/Source/Editor/Experimental/EditorInteractiveToolsFramework/Public/ViewportInteractions/ViewportMoveYawInteraction.h
--- a/Engine/Source/Editor/Experimental/EditorInteractiveToolsFramework/Public/ViewportInteractions/ViewportMoveYawInteraction.h
+++ b/Engine/Source/Editor/Experimental/EditorInteractiveToolsFramework/Public/ViewportInteractions/ViewportMoveYawInteraction.h
@@ -32,6 +32,9 @@ public:
 	UE_API virtual void OnDragDelta(float InMouseDeltaX, float InMouseDeltaY) override;
 	UE_API virtual bool CanBeActivated(const FInputDeviceState& InInputDeviceState = FInputDeviceState()) const override;
 	//~ End UViewportDragInteraction
+
+protected:
+	virtual TArray<FEditorModeID> GetUnsupportedModes() const override;
 };
 
 #undef UE_API
