Gizmo & Explicit Mode: do not ensure when the interface has not been registered but log a unique message instead

#rnx

diff --git a/Engine/Source/Editor/Experimental/EditorInteractiveToolsFramework/Private/EditorGizmos/GizmoRotationUtil.cpp b/Engine/Source/Editor/Experimental/EditorInteractiveToolsFramework/Private/EditorGizmos/GizmoRotationUtil.cpp
--- a/Engine/Source/Editor/Experimental/EditorInteractiveToolsFramework/Private/EditorGizmos/GizmoRotationUtil.cpp
+++ b/Engine/Source/Editor/Experimental/EditorInteractiveToolsFramework/Private/EditorGizmos/GizmoRotationUtil.cpp
@@ -136,7 +136,10 @@ IRelativeTransformInterface* FRelativeTransformInterfaceRegistry::FindRelativeTr
 IRelativeTransformInterface* FRelativeTransformInterfaceRegistry::FindRelativeTransformInterface(const UClass* InClass) const
 {
 	const TUniquePtr<IRelativeTransformInterface>* Interface = WorldInterfaceToRelativeTransformInterface.Find(InClass);
-	ensureMsgf(Interface, TEXT("No relative transform interface found for class %s. Did you call RegisterRelativeTransformInterface<> for that class?"), *InClass->GetName());
+	if (!Interface)
+	{
+		LogOnce(InClass);		
+	}
 	return Interface ? Interface->Get() : nullptr;
 }
 
@@ -151,6 +154,17 @@ void FRelativeTransformInterfaceRegistry::RegisterDefaultInterfaces()
 	Registry.RegisterRelativeTransformInterface<UComponentElementEditorWorldInterface>(MakeUnique<FComponentRelativeTransformInterface>());
 }
 
+void FRelativeTransformInterfaceRegistry::LogOnce(const UClass* InClass) const
+{
+	if (InClass && !QueriedUnregisteredClass.Contains(InClass))
+	{
+		UE_LOG(LogTemp, Warning, TEXT("No relative transform interface found for class %s. Did you call RegisterRelativeTransformInterface<> for that class?"), *InClass->GetName());
+		QueriedUnregisteredClass.Add(InClass);
+	}
+}
+	
+	
+
 bool GetRelativeTransform(const TTypedElement<ITypedElementWorldInterface>& InElement, FTransform& OutTransform, FRotationContext& InOutRotationContext)
 {
 	InOutRotationContext.RotationOrder = EEulerRotationOrder::XYZ;
diff --git a/Engine/Source/Editor/Experimental/EditorInteractiveToolsFramework/Public/EditorGizmos/GizmoRotationUtil.h b/Engine/Source/Editor/Experimental/EditorInteractiveToolsFramework/Public/EditorGizmos/GizmoRotationUtil.h
--- a/Engine/Source/Editor/Experimental/EditorInteractiveToolsFramework/Public/EditorGizmos/GizmoRotationUtil.h
+++ b/Engine/Source/Editor/Experimental/EditorInteractiveToolsFramework/Public/EditorGizmos/GizmoRotationUtil.h
@@ -78,6 +78,9 @@ namespace UE::GizmoRotationUtil
 		FRelativeTransformInterfaceRegistry() = default;
 		
 		TMap<UClass*, TUniquePtr<IRelativeTransformInterface>> WorldInterfaceToRelativeTransformInterface;
+
+		mutable TSet<const UClass*> QueriedUnregisteredClass;
+		void LogOnce(const UClass* InClass) const;
 	};
 
 	/**
