[ITF Drag Tools] Duplicating Spline Point also duplicates the Spline Component

`UDuplicateDragSelectionInteraction` now skips duplication if a component visualizer is active.

diff --git a/Engine/Source/Editor/Experimental/EditorInteractiveToolsFramework/Private/EditorDragTools/DuplicateDragSelectionInteraction.cpp b/Engine/Source/Editor/Experimental/EditorInteractiveToolsFramework/Private/EditorDragTools/DuplicateDragSelectionInteraction.cpp
--- a/Engine/Source/Editor/Experimental/EditorInteractiveToolsFramework/Private/EditorDragTools/DuplicateDragSelectionInteraction.cpp
+++ b/Engine/Source/Editor/Experimental/EditorInteractiveToolsFramework/Private/EditorDragTools/DuplicateDragSelectionInteraction.cpp
@@ -2,6 +2,7 @@
 
 #include "EditorDragTools/DuplicateDragSelectionInteraction.h"
 #include "ContextObjectStore.h"
+#include "Editor/UnrealEdEngine.h"
 #include "EditorDragTools/DragToolInteraction.h"
 #include "EditorGizmos/EditorTransformGizmo.h"
 #include "EditorGizmos/EditorTransformGizmoUtil.h"
@@ -15,6 +16,7 @@
 #include "LevelEditorViewport.h"
 #include "ScopedTransaction.h"
 #include "Tools/AssetEditorContextInterface.h"
+#include "UnrealEdGlobals.h"
 #include "ViewportInteractions/ViewportInteractionsBehaviorSource.h"
 
 #define LOCTEXT_NAMESPACE "EditorDuplicateDragSelection"
@@ -160,6 +162,12 @@ void UDuplicateDragSelectionInteraction::OnGizmoMovementEnd(UTransformProxy* InT
 
 void UDuplicateDragSelectionInteraction::DuplicateSelection(bool bInSelectNewElements) const
 {
+	// "Active" Component visualizers process inputs, so they might want to handle Alt + Drag differently (e.g. spline points editing)
+	if (GUnrealEd->ComponentVisManager.IsActive())
+	{
+		return;
+	}
+
 	if (GCurrentLevelEditingViewportClient)
 	{
 		if (TSharedPtr<ILevelEditor> LevelEditor = GCurrentLevelEditingViewportClient->ParentLevelEditor.Pin())
