[Gizmos] Fix issues where cached viewport client and toolkit host pointers could persist through a layout restore, causing crashes. Fix issues where viewport interactions used viewport clients as the source for mode tools, when that is unecessary. This caused issues when loading a layout, as interactions were created before viewports existed.

diff --git a/Engine/Source/Editor/Experimental/EditorInteractiveToolsFramework/Private/EditorDragTools/DragToolInteraction.cpp b/Engine/Source/Editor/Experimental/EditorInteractiveToolsFramework/Private/EditorDragTools/DragToolInteraction.cpp
--- a/Engine/Source/Editor/Experimental/EditorInteractiveToolsFramework/Private/EditorDragTools/DragToolInteraction.cpp
+++ b/Engine/Source/Editor/Experimental/EditorInteractiveToolsFramework/Private/EditorDragTools/DragToolInteraction.cpp
@@ -116,12 +116,9 @@ void UDragToolInteraction::Initialize(UViewportInteractionsBehaviorSource* const
 {
 	UViewportInteraction::Initialize(InViewportInteractionsBehaviorSource);
 
-	if (const FEditorViewportClient* const EditorViewportClient = GetEditorViewportClient())
+	if (const FEditorModeTools* const ModeTools = GetModeTools())
 	{
-		if (const FEditorModeTools* const ModeTools = EditorViewportClient->GetModeTools())
-		{
-			EditorViewportClientProxy = IEditorViewportClientProxy::CreateViewportClientProxy(ModeTools->GetInteractiveToolsContext());
-		}
+		EditorViewportClientProxy = IEditorViewportClientProxy::CreateViewportClientProxy(ModeTools->GetInteractiveToolsContext());
 	}
 }
 
diff --git a/Engine/Source/Editor/Experimental/EditorInteractiveToolsFramework/Private/EditorDragTools/DuplicateDragSelectionInteraction.cpp b/Engine/Source/Editor/Experimental/EditorInteractiveToolsFramework/Private/EditorDragTools/DuplicateDragSelectionInteraction.cpp
--- a/Engine/Source/Editor/Experimental/EditorInteractiveToolsFramework/Private/EditorDragTools/DuplicateDragSelectionInteraction.cpp
+++ b/Engine/Source/Editor/Experimental/EditorInteractiveToolsFramework/Private/EditorDragTools/DuplicateDragSelectionInteraction.cpp
@@ -85,16 +85,19 @@ void UDuplicateDragSelectionInteraction::Initialize(UViewportInteractionsBehavio
 {
 	Super::Initialize(InViewportInteractionsBehaviorSource);
 
-	if (const FEditorViewportClient* const EditorViewportClient = GetEditorViewportClient())
+	if (const FEditorModeTools* const ModeTools = GetModeTools())
 	{
-		if (const FEditorModeTools* const ModeTools = EditorViewportClient->GetModeTools())
+		if (const UModeManagerInteractiveToolsContext* const InteractiveToolsContext = ModeTools->GetInteractiveToolsContext())
 		{
-			if (const UModeManagerInteractiveToolsContext* const InteractiveToolsContext = ModeTools->GetInteractiveToolsContext())
-			{
-				ToolManager = InteractiveToolsContext->ToolManager;
-			}
+			ToolManager = InteractiveToolsContext->ToolManager;
 		}
 	}
+	
+	if (!ToolManager)
+	{
+		UE_LOG(LogTemp, Error, TEXT("No ToolManager found for UDuplicateDragSelectionInteraction."));
+		return;
+	}
 
 	// Map to modifier key state change
 	GetViewportInteractionsBehaviorSource()->OnAltKeyStateChanged().AddUObject(this, &UDuplicateDragSelectionInteraction::OnModifierKeyChanged);
diff --git a/Engine/Source/Editor/Experimental/EditorInteractiveToolsFramework/Private/EditorDragTools/MoveCameraWithObjectInteraction.cpp b/Engine/Source/Editor/Experimental/EditorInteractiveToolsFramework/Private/EditorDragTools/MoveCameraWithObjectInteraction.cpp
--- a/Engine/Source/Editor/Experimental/EditorInteractiveToolsFramework/Private/EditorDragTools/MoveCameraWithObjectInteraction.cpp
+++ b/Engine/Source/Editor/Experimental/EditorInteractiveToolsFramework/Private/EditorDragTools/MoveCameraWithObjectInteraction.cpp
@@ -101,16 +101,19 @@ void UMoveCameraWithObjectInteraction::Initialize(UViewportInteractionsBehaviorS
 	
 	GetViewportInteractionsBehaviorSource()->OnShiftKeyStateChanged().AddUObject(this, &UMoveCameraWithObjectInteraction::OnShiftKeyChanged);
 
-	if (const FEditorViewportClient* const EditorViewportClient = GetEditorViewportClient())
+	if (const FEditorModeTools* const ModeTools = GetModeTools())
 	{
-		if (const FEditorModeTools* const ModeTools = EditorViewportClient->GetModeTools())
+		if (const UModeManagerInteractiveToolsContext* const InteractiveToolsContext = ModeTools->GetInteractiveToolsContext())
 		{
-			if (const UModeManagerInteractiveToolsContext* const InteractiveToolsContext = ModeTools->GetInteractiveToolsContext())
-			{
-				ToolManager = InteractiveToolsContext->ToolManager;
-			}
+			ToolManager = InteractiveToolsContext->ToolManager;
 		}
 	}
+	
+	if (!ToolManager)
+	{
+		UE_LOG(LogTemp, Error, TEXT("No ToolManager found for UMoveCameraWithObjectInteraction"));
+		return;
+	}
 
 	// New TRS Gizmos are already enabled - let's retrieve the Transform Gizmo, from which we register to drag Begin and End delegates
 	if (UEditorInteractiveGizmoManager::UsesNewTRSGizmos())
diff --git a/Engine/Source/Editor/Experimental/EditorInteractiveToolsFramework/Private/ViewportInteractions/ViewportInteraction.cpp b/Engine/Source/Editor/Experimental/EditorInteractiveToolsFramework/Private/ViewportInteractions/ViewportInteraction.cpp
--- a/Engine/Source/Editor/Experimental/EditorInteractiveToolsFramework/Private/ViewportInteractions/ViewportInteraction.cpp
+++ b/Engine/Source/Editor/Experimental/EditorInteractiveToolsFramework/Private/ViewportInteractions/ViewportInteraction.cpp
@@ -149,6 +149,16 @@ void UViewportInteraction::BeginDestroy()
 	}
 }
 
+FEditorModeTools* UViewportInteraction::GetModeTools() const
+{
+	if (UViewportInteractionsBehaviorSource* BehaviorSource = GetViewportInteractionsBehaviorSource())
+	{
+		return BehaviorSource->GetModeTools();
+	}
+
+	return nullptr;
+}
+
 FEditorViewportClient* UViewportInteraction::GetEditorViewportClient() const
 {
 	if (UViewportInteractionsBehaviorSource* BehaviorSource = GetViewportInteractionsBehaviorSource())
diff --git a/Engine/Source/Editor/Experimental/EditorInteractiveToolsFramework/Private/ViewportInteractions/ViewportInteractionsBehaviorSource.cpp b/Engine/Source/Editor/Experimental/EditorInteractiveToolsFramework/Private/ViewportInteractions/ViewportInteractionsBehaviorSource.cpp
--- a/Engine/Source/Editor/Experimental/EditorInteractiveToolsFramework/Private/ViewportInteractions/ViewportInteractionsBehaviorSource.cpp
+++ b/Engine/Source/Editor/Experimental/EditorInteractiveToolsFramework/Private/ViewportInteractions/ViewportInteractionsBehaviorSource.cpp
@@ -316,17 +316,13 @@ UViewportInteraction* UViewportInteractionsBehaviorSource::AddInteraction(const
 			return nullptr;
 		}
 
-		FEditorViewportClient* EditorViewportClient = GetEditorViewportClient();
-		if (!EditorViewportClient)
-		{
-			UE_LOG(LogTemp, Warning, TEXT("Viewport Interactions Behavior Source has no Editor Viewport Client available!"));
-			return nullptr;
-		}
-
-		// Current mode does not support this interaction: skip it
-		if (!DefaultInteraction->IsCurrentModeSupported(EditorViewportClient->GetModeTools()))
+		if (const FEditorModeTools* ModeTools = GetModeTools())
 		{
-			return nullptr;
+			if (!DefaultInteraction->IsCurrentModeSupported(ModeTools))
+			{
+				// Current mode does not support this interaction: skip it
+				return nullptr;
+			}
 		}
 	}
 
@@ -443,6 +439,15 @@ void UViewportInteractionsBehaviorSource::SetIsMouseLooking(bool bInIsLooking)
 	OnMouseLookingStateChanged().Broadcast(bInIsLooking);
 }
 
+FEditorModeTools* UViewportInteractionsBehaviorSource::GetModeTools() const
+{
+	if (const UEditorInteractiveToolsContext* const InteractiveToolsContext = EditorInteractiveToolsContextWeak.Get())
+	{
+		return InteractiveToolsContext->GetParentEditorModeManager();
+	}
+	return nullptr;
+}
+
 FEditorViewportClient* UViewportInteractionsBehaviorSource::GetEditorViewportClient() const
 {
 	FEditorViewportClient* EditorViewportClient = nullptr;
diff --git a/Engine/Source/Editor/Experimental/EditorInteractiveToolsFramework/Public/ViewportInteractions/ViewportInteraction.h b/Engine/Source/Editor/Experimental/EditorInteractiveToolsFramework/Public/ViewportInteractions/ViewportInteraction.h
--- a/Engine/Source/Editor/Experimental/EditorInteractiveToolsFramework/Public/ViewportInteractions/ViewportInteraction.h
+++ b/Engine/Source/Editor/Experimental/EditorInteractiveToolsFramework/Public/ViewportInteractions/ViewportInteraction.h
@@ -137,6 +137,7 @@ public:
 	bool IsCurrentModeSupported(const FEditorModeTools* InModeTools) const;
 
 protected:
+	UE_API FEditorModeTools* GetModeTools() const;
 	UE_API FEditorViewportClient* GetEditorViewportClient() const;
 
 	/**
diff --git a/Engine/Source/Editor/Experimental/EditorInteractiveToolsFramework/Public/ViewportInteractions/ViewportInteractionsBehaviorSource.h b/Engine/Source/Editor/Experimental/EditorInteractiveToolsFramework/Public/ViewportInteractions/ViewportInteractionsBehaviorSource.h
--- a/Engine/Source/Editor/Experimental/EditorInteractiveToolsFramework/Public/ViewportInteractions/ViewportInteractionsBehaviorSource.h
+++ b/Engine/Source/Editor/Experimental/EditorInteractiveToolsFramework/Public/ViewportInteractions/ViewportInteractionsBehaviorSource.h
@@ -229,6 +229,8 @@ public:
 
 	void SetIsMouseLooking(bool bInIsLooking);
 
+	FEditorModeTools* GetModeTools() const;
+
 	FEditorViewportClient* GetEditorViewportClient() const;
 
 	FEditorViewportClient* GetHoveredEditorViewportClient() const;
diff --git a/Engine/Source/Editor/UnrealEd/Private/EditorModeManager.cpp b/Engine/Source/Editor/UnrealEd/Private/EditorModeManager.cpp
--- a/Engine/Source/Editor/UnrealEd/Private/EditorModeManager.cpp
+++ b/Engine/Source/Editor/UnrealEd/Private/EditorModeManager.cpp
@@ -101,6 +101,7 @@ FEditorModeTools::FEditorModeTools()
 
 		// This binding ensures the mode is destroyed if the type is unregistered outside of normal shutdown process
 		GEditor->GetEditorSubsystem<UAssetEditorSubsystem>()->OnEditorModeUnregistered().AddRaw(this, &FEditorModeTools::OnModeUnregistered);
+		GEditor->OnViewportClientListChanged().AddRaw(this, &FEditorModeTools::OnViewportListChanged);
 	}
 
 	UE::EditorTransformGizmoUtil::RegisterTransformGizmoContextObject(this);
@@ -780,6 +781,38 @@ void FEditorModeTools::OnWorldCleanup(UWorld* InWorld, bool bSessionEnded, bool
 	}
 }
 
+void FEditorModeTools::OnViewportListChanged()
+{
+	bool bFoundHovered = false;
+	bool bFoundFocused = false;
+	
+	if (GEditor)
+	{
+		for (const FEditorViewportClient* Client : GEditor->GetAllViewportClients())
+		{
+			if (Client == HoveredViewportClient)
+			{
+				bFoundHovered = true;
+			}
+			if (Client == FocusedViewportClient)
+			{
+				bFoundFocused = true;
+			}
+		}
+	}
+	
+	// If the hovered or focused clients are not found in the list, they have been removed.
+	// Their pointers are now pointing at garbage and need to be reset.
+	if (!bFoundHovered)
+	{
+		HoveredViewportClient = nullptr;
+	}
+	if (!bFoundFocused)
+	{
+		FocusedViewportClient = nullptr;
+	}
+}
+
 void FEditorModeTools::RemoveAllDelegateHandlers()
 {
 	if (GEditor)
@@ -789,6 +822,8 @@ void FEditorModeTools::RemoveAllDelegateHandlers()
 		{
 			AssetEditorSubsystem->OnEditorModeUnregistered().RemoveAll(this);
 		}
+		
+		GEditor->OnViewportClientListChanged().RemoveAll(this);
 	}
 
 	FWorldDelegates::OnWorldCleanup.RemoveAll(this);
diff --git a/Engine/Source/Editor/UnrealEd/Public/EditorModeManager.h b/Engine/Source/Editor/UnrealEd/Public/EditorModeManager.h
--- a/Engine/Source/Editor/UnrealEd/Public/EditorModeManager.h
+++ b/Engine/Source/Editor/UnrealEd/Public/EditorModeManager.h
@@ -578,6 +578,8 @@ protected:
 
 	/** Handles the notification when a world is going through GC to clean up any modes pending deactivation. */
 	UNREALED_API void OnWorldCleanup(UWorld* InWorld, bool bSessionEnded, bool bCleanupResources);
+	
+	UNREALED_API void OnViewportListChanged();
 
 	UNREALED_API virtual void DrawBrackets(FEditorViewportClient* ViewportClient, FViewport* Viewport, const FSceneView* View, FCanvas* Canvas);
 
