[ITF Viewport Interactions] Fix an issue where WASD camera flying was not active while placing two fingers on a trackpad for mouselook.

#rnx

diff --git a/Engine/Source/Editor/Experimental/EditorInteractiveToolsFramework/Private/ViewportInteractions/ViewportCameraTranslateInteraction.cpp b/Engine/Source/Editor/Experimental/EditorInteractiveToolsFramework/Private/ViewportInteractions/ViewportCameraTranslateInteraction.cpp
--- a/Engine/Source/Editor/Experimental/EditorInteractiveToolsFramework/Private/ViewportInteractions/ViewportCameraTranslateInteraction.cpp
+++ b/Engine/Source/Editor/Experimental/EditorInteractiveToolsFramework/Private/ViewportInteractions/ViewportCameraTranslateInteraction.cpp
@@ -6,6 +6,7 @@
 #include "EditorViewportClient.h"
 #include "ViewportClientNavigationHelper.h"
 #include "ViewportInteractions/ViewportInteractionsBehaviorSource.h"
+#include "Framework/Application/SlateApplication.h"
 
 namespace UE::Editor::ViewportInteractions::Private
 {
@@ -68,7 +69,7 @@ UViewportCameraTranslateInteraction::UViewportCameraTranslateInteraction()
 			}
 		}
 
-		return bResult;
+		return bResult || FSlateApplication::Get().IsUsingTrackpad();
 	};
 	KeyInputBehaviorWeak = KeyInputBehavior;
 
diff --git a/Engine/Source/Editor/Experimental/EditorInteractiveToolsFramework/Private/ViewportInteractions/ViewportClickInteraction.cpp b/Engine/Source/Editor/Experimental/EditorInteractiveToolsFramework/Private/ViewportInteractions/ViewportClickInteraction.cpp
--- a/Engine/Source/Editor/Experimental/EditorInteractiveToolsFramework/Private/ViewportInteractions/ViewportClickInteraction.cpp
+++ b/Engine/Source/Editor/Experimental/EditorInteractiveToolsFramework/Private/ViewportInteractions/ViewportClickInteraction.cpp
@@ -37,15 +37,12 @@ void UViewportClickInteraction::OnClicked(const FInputDeviceRay& InClickPos)
 		ViewportInteractionBehaviorSource->ClearMouseCursorOverride();
 	}
 
-	// Make sure this release is not happening while user is moving the viewport camera - in that case, this is not considered a click.
 	if (IsMouseLooking())
 	{
 		if (UViewportInteractionsBehaviorSource* const ViewportInteractionBehaviorSource = GetViewportInteractionsBehaviorSource())
 		{
 			ViewportInteractionBehaviorSource->SetIsMouseLooking(false);
 		}
-
-		return;
 	}
 
 	const int32 CurrentClickPosX = static_cast<int32>(InClickPos.ScreenPosition.X);
diff --git a/Engine/Source/Editor/Experimental/EditorInteractiveToolsFramework/Private/ViewportInteractions/ViewportInteractionsBehaviorSource.cpp b/Engine/Source/Editor/Experimental/EditorInteractiveToolsFramework/Private/ViewportInteractions/ViewportInteractionsBehaviorSource.cpp
--- a/Engine/Source/Editor/Experimental/EditorInteractiveToolsFramework/Private/ViewportInteractions/ViewportInteractionsBehaviorSource.cpp
+++ b/Engine/Source/Editor/Experimental/EditorInteractiveToolsFramework/Private/ViewportInteractions/ViewportInteractionsBehaviorSource.cpp
@@ -11,6 +11,7 @@
 #include "ViewportClientNavigationHelper.h"
 #include "EditorDragTools/DragToolInteraction.h"
 #include "ViewportInteractions/ViewportInteraction.h"
+#include "Framework/Application/SlateApplication.h"
 
 namespace UE::Editor::ViewportInteractions
 {
@@ -412,6 +413,11 @@ void UViewportInteractionsBehaviorSource::ClearMouseCursorOverride()
 
 bool UViewportInteractionsBehaviorSource::IsMouseLooking() const
 {
+	if (FSlateApplication::Get().IsUsingTrackpad())
+	{
+		return true;
+	}
+
 	if (FEditorViewportClient* const EditorViewportClient = GetEditorViewportClient())
 	{
 		return EditorViewportClient->GetViewportNavigationHelper()->bIsMouseLooking;
