Fix 2D viewport behavior target not properly dealing with Top and Bottom view since they changed their orientation ( Y is now horizontal and X is vertical )
Note : The ortho related code is only used by the unified Dataflow editor ( the UV and cloth editors are actually emulating ortho views using a perspective view and as a result are not affected by this change )
This fixes the Dataflow Editor middle click drag feature in 2D view

diff --git a/Engine/Source/Editor/Experimental/EditorInteractiveToolsFramework/Private/Behaviors/2DViewportBehaviorTargets.cpp b/Engine/Source/Editor/Experimental/EditorInteractiveToolsFramework/Private/Behaviors/2DViewportBehaviorTargets.cpp
--- a/Engine/Source/Editor/Experimental/EditorInteractiveToolsFramework/Private/Behaviors/2DViewportBehaviorTargets.cpp
+++ b/Engine/Source/Editor/Experimental/EditorInteractiveToolsFramework/Private/Behaviors/2DViewportBehaviorTargets.cpp
@@ -80,7 +80,7 @@ void FEditor2DScrollBehaviorTarget::OnClickDrag(const FInputDeviceRay& DragPos)
 		switch (ViewportClient->GetViewportType())
 		{
 		case LVT_OrthoXY:
-			CameraDisplacement = FVector(-DX * UnitsPerPixel, -DY * UnitsPerPixel, 0.0);
+			CameraDisplacement = FVector(-DY * UnitsPerPixel, DX * UnitsPerPixel, 0.0);
 			break;
 		case LVT_OrthoXZ:
 			CameraDisplacement = FVector(-DX * UnitsPerPixel, 0.0, DY * UnitsPerPixel);
@@ -89,7 +89,7 @@ void FEditor2DScrollBehaviorTarget::OnClickDrag(const FInputDeviceRay& DragPos)
 			CameraDisplacement = FVector(0.0, -DX * UnitsPerPixel, DY * UnitsPerPixel);
 			break;
 		case LVT_OrthoNegativeXY:
-			CameraDisplacement = FVector(DX * UnitsPerPixel, -DY * UnitsPerPixel, 0.0);
+			CameraDisplacement = FVector(-DY * UnitsPerPixel, -DX * UnitsPerPixel, 0.0);
 			break;
 		case LVT_OrthoNegativeXZ:
 			CameraDisplacement = FVector(DX * UnitsPerPixel, 0.0, DY * UnitsPerPixel);
