[ITF Viewport Interactions] Fix issues where the two-finger-click-and-drag gesture on Mac did not pan the camera.

#rnx

diff --git a/Engine/Source/Editor/Experimental/EditorInteractiveToolsFramework/Private/ViewportInteractions/ViewportInteractionsUtilities.cpp b/Engine/Source/Editor/Experimental/EditorInteractiveToolsFramework/Private/ViewportInteractions/ViewportInteractionsUtilities.cpp
--- a/Engine/Source/Editor/Experimental/EditorInteractiveToolsFramework/Private/ViewportInteractions/ViewportInteractionsUtilities.cpp
+++ b/Engine/Source/Editor/Experimental/EditorInteractiveToolsFramework/Private/ViewportInteractions/ViewportInteractionsUtilities.cpp
@@ -31,6 +31,7 @@ void UE::Editor::ViewportInteractions::AddDefaultCameraMovementInteractions(
 				UViewportOrbitInteraction::StaticClass(),
 				UViewportMoveYawInteraction::StaticClass(),
 				UViewportMovePanInteraction::StaticClass(),
+				UViewportTrackpadPanInteraction::StaticClass(),
 				UViewportViewAngleInteraction::StaticClass(),
 				UViewportOrthoPanInteraction::StaticClass(),
 				UViewportPanInteraction::StaticClass(),
diff --git a/Engine/Source/Editor/Experimental/EditorInteractiveToolsFramework/Private/ViewportInteractions/ViewportPanInteraction.cpp b/Engine/Source/Editor/Experimental/EditorInteractiveToolsFramework/Private/ViewportInteractions/ViewportPanInteraction.cpp
--- a/Engine/Source/Editor/Experimental/EditorInteractiveToolsFramework/Private/ViewportInteractions/ViewportPanInteraction.cpp
+++ b/Engine/Source/Editor/Experimental/EditorInteractiveToolsFramework/Private/ViewportInteractions/ViewportPanInteraction.cpp
@@ -7,6 +7,7 @@
 #include "Settings/LevelEditorViewportSettings.h"
 #include "ViewportClientNavigationHelper.h"
 #include "ViewportInteractions/ViewportInteractionsBehaviorSource.h"
+#include "Framework/Application/SlateApplication.h"
 
 UViewportPanInteraction::UViewportPanInteraction()
 {
@@ -111,3 +112,18 @@ void UViewportMovePanInteraction::OnDragDelta(float InMouseDeltaX, float InMouse
 	}
 }
 
+UViewportTrackpadPanInteraction::UViewportTrackpadPanInteraction()
+{
+	InteractionName = TEXT("Perpsective Trackpad Pan");
+	
+	ViewportClickDragBehavior->DisableButton(EKeys::LeftMouseButton);
+	
+	// TODO: Make behaviors with more buttons automatically have higher priority
+	ViewportClickDragBehavior->SetDefaultPriority(UE::Editor::ViewportInteractions::DRAG_PRIORITY);
+	ViewportClickDragBehavior->SetDragConfirmedPriority(UE::Editor::ViewportInteractions::DRAG_CONFIRMED_PRIORITY);
+}
+
+bool UViewportTrackpadPanInteraction::CanBeActivated(const FInputDeviceState& InInputDeviceState) const
+{
+	return IsEnabled() && FSlateApplication::Get().IsUsingTrackpad();	
+}
diff --git a/Engine/Source/Editor/Experimental/EditorInteractiveToolsFramework/Private/ViewportInteractions/ViewportViewAngleInteraction.cpp b/Engine/Source/Editor/Experimental/EditorInteractiveToolsFramework/Private/ViewportInteractions/ViewportViewAngleInteraction.cpp
--- a/Engine/Source/Editor/Experimental/EditorInteractiveToolsFramework/Private/ViewportInteractions/ViewportViewAngleInteraction.cpp
+++ b/Engine/Source/Editor/Experimental/EditorInteractiveToolsFramework/Private/ViewportInteractions/ViewportViewAngleInteraction.cpp
@@ -5,6 +5,7 @@
 #include "EditorViewportClient.h"
 #include "Settings/LevelEditorViewportSettings.h"
 #include "ViewportClientNavigationHelper.h"
+#include "Framework/Application/SlateApplication.h"
 #include "ViewportInteractions/ViewportInteractionsBehaviorSource.h"
 
 UViewportViewAngleInteraction::UViewportViewAngleInteraction()
@@ -16,6 +17,13 @@ UViewportViewAngleInteraction::UViewportViewAngleInteraction()
 	EnableButton(EKeys::RightMouseButton);
 }
 
+bool UViewportViewAngleInteraction::CanBeActivated(const FInputDeviceState& InInputDeviceState) const
+{
+	// TODO: Trackpad view angle adjustment is still driven by `FEditorViewportClient::InputGesture()`
+	// Trackpad right click & drag is intended to be a pan and is handled by `UViewportTrackpadPanInteraction`
+	return Super::CanBeActivated(InInputDeviceState) && !FSlateApplication::Get().IsUsingTrackpad();
+}
+
 FInputRayHit UViewportViewAngleInteraction::CanBeginClickDragSequence(const FInputDeviceRay& InPressPos)
 {
 	bool bSuccess = CanBeActivated();
diff --git a/Engine/Source/Editor/Experimental/EditorInteractiveToolsFramework/Public/ViewportInteractions/ViewportPanInteraction.h b/Engine/Source/Editor/Experimental/EditorInteractiveToolsFramework/Public/ViewportInteractions/ViewportPanInteraction.h
--- a/Engine/Source/Editor/Experimental/EditorInteractiveToolsFramework/Public/ViewportInteractions/ViewportPanInteraction.h
+++ b/Engine/Source/Editor/Experimental/EditorInteractiveToolsFramework/Public/ViewportInteractions/ViewportPanInteraction.h
@@ -50,4 +50,18 @@ protected:
 	virtual void OnDragDelta(float InMouseDeltaX, float InMouseDeltaY) override;
 };
 
+/**
+ * Handles RMB + Trackpad movement for world-space-up Panning
+ */
+UCLASS(MinimalAPI, Transient)
+class UViewportTrackpadPanInteraction : public UViewportMovePanInteraction
+{
+	GENERATED_BODY()
+	
+public:
+	UE_API UViewportTrackpadPanInteraction();
+	
+	UE_API virtual bool CanBeActivated(const FInputDeviceState& InInputDeviceState = FInputDeviceState()) const override;
+};
+
 #undef UE_API
diff --git a/Engine/Source/Editor/Experimental/EditorInteractiveToolsFramework/Public/ViewportInteractions/ViewportViewAngleInteraction.h b/Engine/Source/Editor/Experimental/EditorInteractiveToolsFramework/Public/ViewportInteractions/ViewportViewAngleInteraction.h
--- a/Engine/Source/Editor/Experimental/EditorInteractiveToolsFramework/Public/ViewportInteractions/ViewportViewAngleInteraction.h
+++ b/Engine/Source/Editor/Experimental/EditorInteractiveToolsFramework/Public/ViewportInteractions/ViewportViewAngleInteraction.h
@@ -21,6 +21,7 @@ public:
 	UE_API UViewportViewAngleInteraction();
 
 	//~ Begin IViewportClickDragBehaviorTarget
+	UE_API virtual bool CanBeActivated(const FInputDeviceState& InInputDeviceState = FInputDeviceState()) const override;
 	UE_API virtual FInputRayHit CanBeginClickDragSequence(const FInputDeviceRay& InPressPos) override;
 	UE_API virtual void OnDragStart(const FInputDeviceRay& InDragStartPos) override;
 	UE_API virtual void OnDragEnd(const FInputDeviceRay& InDragEndPos) override;
