[ITF Viewport Interactions] Add Interaction cleanup override when removing from Behavior Source

diff --git a/Engine/Source/Editor/Experimental/EditorInteractiveToolsFramework/Private/ViewportInteractions/ViewportInteractionsBehaviorSource.cpp b/Engine/Source/Editor/Experimental/EditorInteractiveToolsFramework/Private/ViewportInteractions/ViewportInteractionsBehaviorSource.cpp
--- a/Engine/Source/Editor/Experimental/EditorInteractiveToolsFramework/Private/ViewportInteractions/ViewportInteractionsBehaviorSource.cpp
+++ b/Engine/Source/Editor/Experimental/EditorInteractiveToolsFramework/Private/ViewportInteractions/ViewportInteractionsBehaviorSource.cpp
@@ -356,6 +356,7 @@ void UViewportInteractionsBehaviorSource::RemoveInteraction(const FName InIntera
 		UE::Editor::ViewportInteractions::LOG(*FString::Printf(TEXT("Removing Viewport Interaction \"%s\""), *Interaction->GetInteractionName().ToString()));
 
 		ViewportInteractions.Remove(Interaction->GetInteractionName());
+		Interaction->Shutdown();
 
 		BehaviorSet->RemoveBySource(Interaction);
 
@@ -372,6 +373,17 @@ void UViewportInteractionsBehaviorSource::Reset()
 {
 	DeregisterBehaviorSources();
 	BehaviorSet->RemoveAll();
+
+	// We are about to reset interactions list, but they might be garbage collected later.
+	// We need to make sure they are properly cleaned up, e.g. in case they are listening to some delegates
+	for (const TTuple<FName, TObjectPtr<UViewportInteraction>>& Pair : ViewportInteractions)
+	{
+		if (UViewportInteraction* const Interaction = Pair.Value)
+		{
+			Interaction->Shutdown();
+		}
+	}
+
 	ViewportInteractions.Reset();
 	ViewportInteractionsStatusMap.Reset();
 	
diff --git a/Engine/Source/Editor/Experimental/EditorInteractiveToolsFramework/Public/ViewportInteractions/ViewportInteraction.h b/Engine/Source/Editor/Experimental/EditorInteractiveToolsFramework/Public/ViewportInteractions/ViewportInteraction.h
--- a/Engine/Source/Editor/Experimental/EditorInteractiveToolsFramework/Public/ViewportInteractions/ViewportInteraction.h
+++ b/Engine/Source/Editor/Experimental/EditorInteractiveToolsFramework/Public/ViewportInteractions/ViewportInteraction.h
@@ -100,6 +100,9 @@ public:
 	UE_API void SetEnabled(bool bInEnabled);
 	UE_API bool IsEnabled() const;
 
+	/** Called when the interaction is removed from its behavior source. Override to make sure the interaction is properly cleaned up */
+	virtual void Shutdown() {}
+
 	/** Called every frame. */
 	virtual void Tick(float InDeltaTime) const {}
 	
