[ITF Drag Tools] Prevent Alt Duplicate interaction from performing multiple duplicates

diff --git a/Engine/Source/Editor/Experimental/EditorInteractiveToolsFramework/Private/EditorDragTools/DuplicateDragSelectionInteraction.cpp b/Engine/Source/Editor/Experimental/EditorInteractiveToolsFramework/Private/EditorDragTools/DuplicateDragSelectionInteraction.cpp
--- a/Engine/Source/Editor/Experimental/EditorInteractiveToolsFramework/Private/EditorDragTools/DuplicateDragSelectionInteraction.cpp
+++ b/Engine/Source/Editor/Experimental/EditorInteractiveToolsFramework/Private/EditorDragTools/DuplicateDragSelectionInteraction.cpp
@@ -133,7 +133,14 @@ void UDuplicateDragSelectionInteraction::Initialize(UViewportInteractionsBehavio
 	);
 }
 
-UDuplicateDragSelectionInteraction::~UDuplicateDragSelectionInteraction()
+void UDuplicateDragSelectionInteraction::BeginDestroy()
+{
+	Super::BeginDestroy();
+
+	Shutdown();
+}
+
+void UDuplicateDragSelectionInteraction::Shutdown()
 {
 	UEditorInteractiveGizmoManager::OnUsesNewTRSGizmosChangedDelegate().RemoveAll(this);
 	DuplicateTransaction.Reset();
diff --git a/Engine/Source/Editor/Experimental/EditorInteractiveToolsFramework/Public/EditorDragTools/DuplicateDragSelectionInteraction.h b/Engine/Source/Editor/Experimental/EditorInteractiveToolsFramework/Public/EditorDragTools/DuplicateDragSelectionInteraction.h
--- a/Engine/Source/Editor/Experimental/EditorInteractiveToolsFramework/Public/EditorDragTools/DuplicateDragSelectionInteraction.h
+++ b/Engine/Source/Editor/Experimental/EditorInteractiveToolsFramework/Public/EditorDragTools/DuplicateDragSelectionInteraction.h
@@ -28,8 +28,10 @@ class UDuplicateDragSelectionInteraction
 public:
 	UDuplicateDragSelectionInteraction();
 
-	virtual ~UDuplicateDragSelectionInteraction() override;
+	virtual void BeginDestroy() override;
 
+	virtual void Shutdown() override;
+	
 protected:
 	/**
 	 * Duplicate the current selection
