[ITF Viewport Interactions] Cleanup updates

Override Cleanup where needed
Move Cleanup logic from UObject destructor to BeginDestroy

diff --git a/Engine/Source/Editor/Experimental/EditorInteractiveToolsFramework/Private/EditorDragTools/MoveCameraWithObjectInteraction.cpp b/Engine/Source/Editor/Experimental/EditorInteractiveToolsFramework/Private/EditorDragTools/MoveCameraWithObjectInteraction.cpp
--- a/Engine/Source/Editor/Experimental/EditorInteractiveToolsFramework/Private/EditorDragTools/MoveCameraWithObjectInteraction.cpp
+++ b/Engine/Source/Editor/Experimental/EditorInteractiveToolsFramework/Private/EditorDragTools/MoveCameraWithObjectInteraction.cpp
@@ -23,27 +23,11 @@ UMoveCameraWithObjectInteraction::UMoveCameraWithObjectInteraction()
 	GroupName = UE::Editor::ViewportInteractions::DragToolsGroupName;
 }
 
-UMoveCameraWithObjectInteraction::~UMoveCameraWithObjectInteraction()
+void UMoveCameraWithObjectInteraction::BeginDestroy()
 {
-	UEditorInteractiveGizmoManager::OnUsesNewTRSGizmosChangedDelegate().RemoveAll(this);
-
-	if (ToolManager)
-	{
-		if (UContextObjectStore* ContextStore = ToolManager->GetContextObjectStore())
-		{
-			if (UEditorTransformGizmoContextObject* ContextObject = ContextStore->FindContext<UEditorTransformGizmoContextObject>())
-			{
-				ContextObject->OnGizmoCreatedDelegate().RemoveAll(this);
-			}
-		}
-	}
-	
-	if (UViewportInteractionsBehaviorSource* Source = GetViewportInteractionsBehaviorSource())
-	{
-		Source->OnShiftKeyStateChanged().RemoveAll(this);
-	}
+	Super::BeginDestroy();
 
-	UnregisterProxyDelegates();
+	Shutdown();
 }
 
 void UMoveCameraWithObjectInteraction::RegisterProxyDelegates()
@@ -148,6 +132,29 @@ void UMoveCameraWithObjectInteraction::Initialize(UViewportInteractionsBehaviorS
 	);
 }
 
+void UMoveCameraWithObjectInteraction::Shutdown()
+{
+	UEditorInteractiveGizmoManager::OnUsesNewTRSGizmosChangedDelegate().RemoveAll(this);
+
+	if (ToolManager)
+	{
+		if (UContextObjectStore* ContextStore = ToolManager->GetContextObjectStore())
+		{
+			if (UEditorTransformGizmoContextObject* ContextObject = ContextStore->FindContext<UEditorTransformGizmoContextObject>())
+			{
+				ContextObject->OnGizmoCreatedDelegate().RemoveAll(this);
+			}
+		}
+	}
+	
+	if (UViewportInteractionsBehaviorSource* Source = GetViewportInteractionsBehaviorSource())
+	{
+		Source->OnShiftKeyStateChanged().RemoveAll(this);
+	}
+
+	UnregisterProxyDelegates();
+}
+
 void UMoveCameraWithObjectInteraction::OnShiftKeyChanged(bool bDown)
 {
 	bIsModifierActive = bDown;
diff --git a/Engine/Source/Editor/Experimental/EditorInteractiveToolsFramework/Private/ViewportInteractions/ViewportFOVInteraction.cpp b/Engine/Source/Editor/Experimental/EditorInteractiveToolsFramework/Private/ViewportInteractions/ViewportFOVInteraction.cpp
--- a/Engine/Source/Editor/Experimental/EditorInteractiveToolsFramework/Private/ViewportInteractions/ViewportFOVInteraction.cpp
+++ b/Engine/Source/Editor/Experimental/EditorInteractiveToolsFramework/Private/ViewportInteractions/ViewportFOVInteraction.cpp
@@ -55,8 +55,10 @@ UViewportFOVInteraction::UViewportFOVInteraction()
 	RegisterInputBehavior(KeyInputBehavior);
 }
 
-UViewportFOVInteraction::~UViewportFOVInteraction()
+void UViewportFOVInteraction::BeginDestroy()
 {
+	Super::BeginDestroy();
+
 	if (UViewportInteractionsBehaviorSource* BehaviorSource = GetViewportInteractionsBehaviorSource())
 	{
 		BehaviorSource->OnMouseLookingStateChanged().RemoveAll(this);
diff --git a/Engine/Source/Editor/Experimental/EditorInteractiveToolsFramework/Public/EditorDragTools/MoveCameraWithObjectInteraction.h b/Engine/Source/Editor/Experimental/EditorInteractiveToolsFramework/Public/EditorDragTools/MoveCameraWithObjectInteraction.h
--- a/Engine/Source/Editor/Experimental/EditorInteractiveToolsFramework/Public/EditorDragTools/MoveCameraWithObjectInteraction.h
+++ b/Engine/Source/Editor/Experimental/EditorInteractiveToolsFramework/Public/EditorDragTools/MoveCameraWithObjectInteraction.h
@@ -27,10 +27,12 @@ class UMoveCameraWithObjectInteraction
 
 public:
 	UE_API UMoveCameraWithObjectInteraction();
-	UE_API virtual ~UMoveCameraWithObjectInteraction() override;
+	UE_API virtual void BeginDestroy() override;
 
 	virtual void Initialize(UViewportInteractionsBehaviorSource* const InViewportInteractionsBehaviorSource) override;
 
+	virtual void Shutdown() override;
+
 protected:
 
 	UE_API void OnShiftKeyChanged(bool bDown);
diff --git a/Engine/Source/Editor/Experimental/EditorInteractiveToolsFramework/Public/ViewportInteractions/ViewportFOVInteraction.h b/Engine/Source/Editor/Experimental/EditorInteractiveToolsFramework/Public/ViewportInteractions/ViewportFOVInteraction.h
--- a/Engine/Source/Editor/Experimental/EditorInteractiveToolsFramework/Public/ViewportInteractions/ViewportFOVInteraction.h
+++ b/Engine/Source/Editor/Experimental/EditorInteractiveToolsFramework/Public/ViewportInteractions/ViewportFOVInteraction.h
@@ -23,7 +23,7 @@ class UViewportFOVInteraction
 
 public:
 	UE_API UViewportFOVInteraction();
-	UE_API virtual ~UViewportFOVInteraction() override;
+	UE_API virtual void BeginDestroy() override;
 
 	//~ Begin IKeyInputBehaviorTarget
 	UE_API virtual void OnKeyPressed(const FKey& InKeyID) override;
