[Editor ITF] Snapping visualizers toggled off in ITF (in place of legacy MouseDeltaTracker end)

diff --git a/Engine/Source/Editor/Experimental/EditorInteractiveToolsFramework/Private/EditorGizmos/EditorGizmoStateTarget.cpp b/Engine/Source/Editor/Experimental/EditorInteractiveToolsFramework/Private/EditorGizmos/EditorGizmoStateTarget.cpp
--- a/Engine/Source/Editor/Experimental/EditorInteractiveToolsFramework/Private/EditorGizmos/EditorGizmoStateTarget.cpp
+++ b/Engine/Source/Editor/Experimental/EditorInteractiveToolsFramework/Private/EditorGizmos/EditorGizmoStateTarget.cpp
@@ -4,6 +4,7 @@
 #include "EditorGizmos/TransformGizmo.h"
 #include "EditorModeManager.h"
 #include "GizmoEdModeInterface.h"
+#include "SnappingUtils.h"
 
 #include UE_INLINE_GENERATED_CPP_BY_NAME(EditorGizmoStateTarget)
 
@@ -25,6 +26,12 @@ void UEditorGizmoStateTarget::BeginUpdate()
 
 void UEditorGizmoStateTarget::EndUpdate()
 {
+	// For compatibility with users of FSnappingUtils - previously executed in MouseDeltaTracker for legacy gizmos
+	{
+		constexpr bool bClearImmediately = true;
+		FSnappingUtils::ClearSnappingHelpers(bClearImmediately);	
+	}
+	
 	const TSharedPtr<FEditorModeTools> ModeTools = WeakModeTools.IsValid() ? WeakModeTools.Pin() : nullptr;
 	if (ensure(ModeTools))
 	{
