[ITF Gizmos] WASD Dragging actor conflict

- Centralize gizmo transform delegates so they live in Behavior Source
- Make sure camera transform interactions driven by keyboard inputs do not work while gizmo is being interacted with

#Synced-CL 48027920

diff --git a/Engine/Source/Editor/Experimental/EditorInteractiveToolsFramework/Private/EditorDragTools/DuplicateDragSelectionInteraction.cpp b/Engine/Source/Editor/Experimental/EditorInteractiveToolsFramework/Private/EditorDragTools/DuplicateDragSelectionInteraction.cpp
--- a/Engine/Source/Editor/Experimental/EditorInteractiveToolsFramework/Private/EditorDragTools/DuplicateDragSelectionInteraction.cpp
+++ b/Engine/Source/Editor/Experimental/EditorInteractiveToolsFramework/Private/EditorDragTools/DuplicateDragSelectionInteraction.cpp
@@ -1,21 +1,12 @@
 // Copyright Epic Games, Inc. All Rights Reserved.
 
 #include "EditorDragTools/DuplicateDragSelectionInteraction.h"
-#include "ContextObjectStore.h"
 #include "Editor/UnrealEdEngine.h"
 #include "EditorDragTools/DragToolInteraction.h"
 #include "EditorGizmos/EditorTransformGizmo.h"
-#include "EditorGizmos/EditorTransformGizmoUtil.h"
 #include "EditorInteractiveGizmoManager.h"
-#include "EditorModeManager.h"
-#include "Elements/Framework/TypedElementCommonActions.h"
-#include "ILevelEditor.h"
-#include "InteractiveGizmo.h"
-#include "InteractiveGizmoManager.h"
-#include "InteractiveToolManager.h"
 #include "LevelEditorViewport.h"
 #include "ScopedTransaction.h"
-#include "Tools/AssetEditorContextInterface.h"
 #include "UnrealEdGlobals.h"
 #include "ViewportInteractions/ViewportInteractionsBehaviorSource.h"
 
@@ -28,109 +19,12 @@ UDuplicateDragSelectionInteraction::UDuplicateDragSelectionInteraction()
 	GroupName = UE::Editor::ViewportInteractions::DragToolsGroupName;
 }
 
-void UDuplicateDragSelectionInteraction::RegisterProxyDelegates()
-{
-	if (TransformGizmo)
-	{
-		// Register Proxy Delegates
-		if (TObjectPtr<UTransformProxy> TransformProxy = TransformGizmo->ActiveTarget)
-		{
-			TransformProxy->OnBeginTransformEdit.AddUObject(
-				this, &UDuplicateDragSelectionInteraction::OnGizmoMovementBegin
-			);
-
-			TransformProxy->OnEndTransformEdit.AddUObject(
-				this, &UDuplicateDragSelectionInteraction::OnGizmoMovementEnd
-			);
-		}
-	}
-}
-
-void UDuplicateDragSelectionInteraction::Reset()
-{
-	if (TransformGizmo)
-	{
-		// Unregister Proxy delegates
-		if (TObjectPtr<UTransformProxy> TransformProxy = TransformGizmo->ActiveTarget)
-		{
-			TransformProxy->OnBeginTransformEdit.RemoveAll(this);
-			TransformProxy->OnEndTransformEdit.RemoveAll(this);
-			TransformProxy->OnTransformChanged.RemoveAll(this);
-		}
-
-		TransformGizmo = nullptr;
-	}
-}
-
-void UDuplicateDragSelectionInteraction::OnGizmoCreatedDelegate(UTransformGizmo* InTransformGizmo)
-{
-	if (!TransformGizmo)
-	{
-		TransformGizmo = Cast<UEditorTransformGizmo>(InTransformGizmo);
-
-		RegisterProxyDelegates();
-	}
-}
-
-void UDuplicateDragSelectionInteraction::OnUsesNewTRSGizmosChanged(bool bInUseNewTRSGizmos)
-{
-	if (!bInUseNewTRSGizmos)
-	{
-		DuplicateTransaction.Reset();
-		Reset();
-	}
-}
-
 void UDuplicateDragSelectionInteraction::Initialize(UViewportInteractionsBehaviorSource* const InViewportInteractionsBehaviorSource)
 {
 	Super::Initialize(InViewportInteractionsBehaviorSource);
 
-	if (const FEditorModeTools* const ModeTools = GetModeTools())
-	{
-		if (const UModeManagerInteractiveToolsContext* const InteractiveToolsContext = ModeTools->GetInteractiveToolsContext())
-		{
-			ToolManager = InteractiveToolsContext->ToolManager;
-		}
-	}
-	
-	if (!ToolManager)
-	{
-		UE_LOG(LogTemp, Error, TEXT("No ToolManager found for UDuplicateDragSelectionInteraction."));
-		return;
-	}
-
 	// Map to modifier key state change
 	GetViewportInteractionsBehaviorSource()->OnAltKeyStateChanged().AddUObject(this, &UDuplicateDragSelectionInteraction::OnModifierKeyChanged);
-
-	// New TRS Gizmos are already enabled - let's retrieve the Transform Gizmo, from which we register to drag Begin and End delegates
-	if (UEditorInteractiveGizmoManager::UsesNewTRSGizmos())
-	{
-		const TArray<UInteractiveGizmo*> Gizmos =
-		ToolManager->GetPairedGizmoManager()->FindAllGizmosOfType("EditorTransformGizmoBuilder");
-
-		if (!Gizmos.IsEmpty())
-		{
-			TransformGizmo = Cast<UEditorTransformGizmo>(Gizmos[0]);
-		}
-
-		RegisterProxyDelegates();
-	}
-
-	// If New TRS Gizmos are not enabled, we need to be able to know as soon as they are enabled
-	// This allows to retrieve the Transform Gizmo, from which we register to drag Begin and End delegates
-	if (UContextObjectStore* ContextStore = ToolManager->GetContextObjectStore())
-	{
-		if (UEditorTransformGizmoContextObject* ContextObject =
-				ContextStore->FindContext<UEditorTransformGizmoContextObject>())
-		{
-			ContextObject->OnGizmoCreatedDelegate().AddUObject(this, &UDuplicateDragSelectionInteraction::OnGizmoCreatedDelegate);
-		}
-	}
-
-	// In case new TRS Gizmos are disabled, we want to know that, so we can stop listening to Drag delegates
-	UEditorInteractiveGizmoManager::OnUsesNewTRSGizmosChangedDelegate().AddUObject(
-		this, &UDuplicateDragSelectionInteraction::OnUsesNewTRSGizmosChanged
-	);
 }
 
 void UDuplicateDragSelectionInteraction::BeginDestroy()
@@ -142,31 +36,33 @@ void UDuplicateDragSelectionInteraction::BeginDestroy()
 
 void UDuplicateDragSelectionInteraction::Shutdown()
 {
-	UEditorInteractiveGizmoManager::OnUsesNewTRSGizmosChangedDelegate().RemoveAll(this);
+	if (UViewportInteractionsBehaviorSource* BehaviorSource = GetViewportInteractionsBehaviorSource())
+	{
+		BehaviorSource->OnAltKeyStateChanged().RemoveAll(this);
+	}
 	DuplicateTransaction.Reset();
-	Reset();
 }
 
 void UDuplicateDragSelectionInteraction::OnGizmoMovementBegin(UTransformProxy* InTransformProxy)
 {
-	if (!bGizmoIsDragged)
+	const UEditorTransformGizmo* const TransformGizmo = GetViewportInteractionsBehaviorSource()->GetEditorTransformGizmo();
+	if (!TransformGizmo)
 	{
-		bGizmoIsDragged = true;
+		return;
+	}
 
-		if (bModifierKeyIsPressed && !TransformGizmo->IsIndirectManipulation())
-		{
-			DuplicateTransaction = MakeShared<FScopedTransaction>(LOCTEXT("DuplicateAndTransformTransactionName", "Duplicate and Transform"));
+	if (bModifierKeyIsPressed && !TransformGizmo->IsIndirectManipulation())
+	{
+		DuplicateTransaction = MakeShared<FScopedTransaction>(LOCTEXT("DuplicateAndTransformTransactionName", "Duplicate and Transform"));
 
-			// Duplicate selection and select the newly created objects
-			constexpr bool bSelectNewElements = true;
-			DuplicateSelection(bSelectNewElements);
-		}
+		// Duplicate selection and select the newly created objects
+		constexpr bool bSelectNewElements = true;
+		DuplicateSelection(bSelectNewElements);
 	}
 }
 
 void UDuplicateDragSelectionInteraction::OnGizmoMovementEnd(UTransformProxy* InTransformProxy)
 {
-	bGizmoIsDragged = false;
 	DuplicateTransaction.Reset();
 }
 
@@ -201,7 +97,7 @@ void UDuplicateDragSelectionInteraction::OnModifierKeyChanged(bool bInIsDown)
 
 	// If already dragging, let's duplicate without selecting new objects
 	// This will "drop" a copy of selection matching current drag transform
-	if (bGizmoIsDragged && bModifierKeyIsPressed)
+	if (GetViewportInteractionsBehaviorSource()->IsGizmoDragging() && bModifierKeyIsPressed)
 	{
 		constexpr bool bSelectNewElements = false;
 		DuplicateSelection(bSelectNewElements);
diff --git a/Engine/Source/Editor/Experimental/EditorInteractiveToolsFramework/Private/EditorDragTools/MoveCameraWithObjectInteraction.cpp b/Engine/Source/Editor/Experimental/EditorInteractiveToolsFramework/Private/EditorDragTools/MoveCameraWithObjectInteraction.cpp
--- a/Engine/Source/Editor/Experimental/EditorInteractiveToolsFramework/Private/EditorDragTools/MoveCameraWithObjectInteraction.cpp
+++ b/Engine/Source/Editor/Experimental/EditorInteractiveToolsFramework/Private/EditorDragTools/MoveCameraWithObjectInteraction.cpp
@@ -2,18 +2,10 @@
 
 #include "EditorDragTools/MoveCameraWithObjectInteraction.h"
 
-#include "ContextObjectStore.h"
 #include "EditorDragTools/DragToolInteraction.h"
 #include "EditorGizmos/EditorTransformGizmo.h"
-#include "EditorGizmos/EditorTransformGizmoUtil.h"
-#include "EditorInteractiveGizmoManager.h"
-#include "EditorModeManager.h"
-#include "InteractiveGizmo.h"
-#include "InteractiveGizmoManager.h"
-#include "InteractiveToolManager.h"
-#include "LevelEditorViewport.h"
+#include "EditorViewportClient.h"
 #include "ViewportClientNavigationHelper.h"
-#include "Tools/EdModeInteractiveToolsContext.h"
 #include "ViewportInteractions/ViewportInteractionsBehaviorSource.h"
 
 UMoveCameraWithObjectInteraction::UMoveCameraWithObjectInteraction()
@@ -30,129 +22,21 @@ void UMoveCameraWithObjectInteraction::BeginDestroy()
 	Shutdown();
 }
 
-void UMoveCameraWithObjectInteraction::RegisterProxyDelegates()
-{
-	if (TransformGizmo)
-	{
-		// Register Proxy Delegates
-		if (TObjectPtr<UTransformProxy> TransformProxy = TransformGizmo->ActiveTarget)
-		{
-			OnBeginPivotEditDelegate =
-				TransformProxy->OnBeginTransformEdit.AddUObject(this, &UMoveCameraWithObjectInteraction::OnGizmoMovementBegin);
-
-			OnEndPivotEditDelegate =
-				TransformProxy->OnEndTransformEdit.AddUObject(this, &UMoveCameraWithObjectInteraction::OnGizmoMovementEnd);
-		}
-	}
-}
-
-void UMoveCameraWithObjectInteraction::UnregisterProxyDelegates()
-{
-	if (TransformGizmo)
-	{
-		// Unregister Proxy delegates
-		if (TObjectPtr<UTransformProxy> TransformProxy = TransformGizmo->ActiveTarget)
-		{
-			TransformProxy->OnBeginTransformEdit.Remove(OnBeginPivotEditDelegate);
-			TransformProxy->OnEndTransformEdit.Remove(OnEndPivotEditDelegate);
-		}
-
-		TransformGizmo = nullptr;
-	}
-}
-
-void UMoveCameraWithObjectInteraction::OnGizmoCreatedDelegate(UTransformGizmo* InTransformGizmo)
-{
-	if (!TransformGizmo)
-	{
-		TransformGizmo = Cast<UEditorTransformGizmo>(InTransformGizmo);
-
-		RegisterProxyDelegates();
-	}
-}
-
-void UMoveCameraWithObjectInteraction::OnUsesNewTRSGizmosChanged(bool bInUseNewTRSGizmos)
-{
-	if (!bInUseNewTRSGizmos)
-	{
-		UnregisterProxyDelegates();
-	}
-}
 
 void UMoveCameraWithObjectInteraction::Initialize(UViewportInteractionsBehaviorSource* const InViewportInteractionsBehaviorSource)
 {
 	Super::Initialize(InViewportInteractionsBehaviorSource);
-	
-	GetViewportInteractionsBehaviorSource()->OnShiftKeyStateChanged().AddUObject(this, &UMoveCameraWithObjectInteraction::OnShiftKeyChanged);
-
-	if (const FEditorModeTools* const ModeTools = GetModeTools())
-	{
-		if (const UModeManagerInteractiveToolsContext* const InteractiveToolsContext = ModeTools->GetInteractiveToolsContext())
-		{
-			ToolManager = InteractiveToolsContext->ToolManager;
-		}
-	}
-	
-	if (!ToolManager)
-	{
-		UE_LOG(LogTemp, Error, TEXT("No ToolManager found for UMoveCameraWithObjectInteraction"));
-		return;
-	}
-
-	// New TRS Gizmos are already enabled - let's retrieve the Transform Gizmo, from which we register to drag Begin and End delegates
-	if (UEditorInteractiveGizmoManager::UsesNewTRSGizmos())
-	{
-		if (UInteractiveGizmoManager* PairedGizmoManager = ToolManager->GetPairedGizmoManager())
-		{
-			const TArray<UInteractiveGizmo*> Gizmos = PairedGizmoManager->FindAllGizmosOfType("EditorTransformGizmoBuilder");
-
-			if (!Gizmos.IsEmpty())
-			{
-				TransformGizmo = Cast<UEditorTransformGizmo>(Gizmos[0]);
-			}
-		}
 
-		RegisterProxyDelegates();
-	}
-
-	// If New TRS Gizmos are not enabled, we need to be able to know as soon as they are enabled
-	// This allows to retrieve the Transform Gizmo, from which we register to drag Begin and End delegates
-	if (UContextObjectStore* ContextStore = ToolManager->GetContextObjectStore())
-	{
-		if (UEditorTransformGizmoContextObject* ContextObject =
-				ContextStore->FindContext<UEditorTransformGizmoContextObject>())
-		{
-			ContextObject->OnGizmoCreatedDelegate().AddUObject(this, &UMoveCameraWithObjectInteraction::OnGizmoCreatedDelegate);
-		}
-	}
-
-	// In case new TRS Gizmos are disabled, we want to know that, so we can stop listening to Drag delegates
-	UEditorInteractiveGizmoManager::OnUsesNewTRSGizmosChangedDelegate().AddUObject(
-		this, &UMoveCameraWithObjectInteraction::OnUsesNewTRSGizmosChanged
-	);
+	GetViewportInteractionsBehaviorSource()->OnShiftKeyStateChanged().AddUObject(this, &UMoveCameraWithObjectInteraction::OnShiftKeyChanged);
 }
 
 void UMoveCameraWithObjectInteraction::Shutdown()
 {
-	UEditorInteractiveGizmoManager::OnUsesNewTRSGizmosChangedDelegate().RemoveAll(this);
-
-	if (ToolManager)
-	{
-		if (UContextObjectStore* ContextStore = ToolManager->GetContextObjectStore())
-		{
-			if (UEditorTransformGizmoContextObject* ContextObject = ContextStore->FindContext<UEditorTransformGizmoContextObject>())
-			{
-				ContextObject->OnGizmoCreatedDelegate().RemoveAll(this);
-			}
-		}
-	}
-	
 	if (UViewportInteractionsBehaviorSource* Source = GetViewportInteractionsBehaviorSource())
 	{
 		Source->OnShiftKeyStateChanged().RemoveAll(this);
 	}
 
-	UnregisterProxyDelegates();
 }
 
 void UMoveCameraWithObjectInteraction::OnShiftKeyChanged(bool bDown)
@@ -163,13 +47,11 @@ void UMoveCameraWithObjectInteraction::OnShiftKeyChanged(bool bDown)
 
 void UMoveCameraWithObjectInteraction::OnGizmoMovementBegin(UTransformProxy* InTransformProxy)
 {
-	bGizmoIsDragged = true;
 	UpdateCameraFollows();
 }
 
 void UMoveCameraWithObjectInteraction::OnGizmoMovementEnd(UTransformProxy* InTransformProxy)
 {
-	bGizmoIsDragged = false;
 	UpdateCameraFollows();
 }
 
@@ -179,7 +61,7 @@ void UMoveCameraWithObjectInteraction::UpdateCameraFollows() const
 	{
 		if (FViewportClientNavigationHelper* NavigationHelper = ViewportClient->GetViewportNavigationHelper())
 		{
-			NavigationHelper->bCameraFollowsGizmo = bGizmoIsDragged && bIsModifierActive;
+			NavigationHelper->bCameraFollowsGizmo = GetViewportInteractionsBehaviorSource()->IsGizmoDragging() && bIsModifierActive;
 		}
 	}
 }
diff --git a/Engine/Source/Editor/Experimental/EditorInteractiveToolsFramework/Private/ViewportInteractions/ViewportCameraRotateInteraction.cpp b/Engine/Source/Editor/Experimental/EditorInteractiveToolsFramework/Private/ViewportInteractions/ViewportCameraRotateInteraction.cpp
--- a/Engine/Source/Editor/Experimental/EditorInteractiveToolsFramework/Private/ViewportInteractions/ViewportCameraRotateInteraction.cpp
+++ b/Engine/Source/Editor/Experimental/EditorInteractiveToolsFramework/Private/ViewportInteractions/ViewportCameraRotateInteraction.cpp
@@ -76,6 +76,17 @@ void UViewportCameraRotateInteraction::OnKeyReleased(const FKey& InKeyID)
 
 void UViewportCameraRotateInteraction::Tick(float InDeltaTime) const
 {
+	UViewportInteractionsBehaviorSource* BehaviorSource = GetViewportInteractionsBehaviorSource();
+	if (!BehaviorSource)
+	{
+		return;
+	}
+
+	if (BehaviorSource->IsGizmoDragging())
+	{
+		return;
+	}
+
 	if (RotateYawImpulse != 0 || RotatePitchImpulse != 0)
 	{
 		if (FEditorViewportClient* EditorViewportClient = GetEditorViewportClient())
diff --git a/Engine/Source/Editor/Experimental/EditorInteractiveToolsFramework/Private/ViewportInteractions/ViewportCameraTranslateInteraction.cpp b/Engine/Source/Editor/Experimental/EditorInteractiveToolsFramework/Private/ViewportInteractions/ViewportCameraTranslateInteraction.cpp
--- a/Engine/Source/Editor/Experimental/EditorInteractiveToolsFramework/Private/ViewportInteractions/ViewportCameraTranslateInteraction.cpp
+++ b/Engine/Source/Editor/Experimental/EditorInteractiveToolsFramework/Private/ViewportInteractions/ViewportCameraTranslateInteraction.cpp
@@ -82,6 +82,17 @@ void UViewportCameraTranslateInteraction::OnForceEndCapture()
 
 void UViewportCameraTranslateInteraction::Tick(float InDeltaTime) const
 {
+	UViewportInteractionsBehaviorSource* BehaviorSource = GetViewportInteractionsBehaviorSource();
+	if (!BehaviorSource)
+	{
+		return;
+	}
+
+	if (BehaviorSource->IsGizmoDragging())
+	{
+		return;
+	}
+
 	using namespace UE::Editor::ViewportInteractions::Private;
 
 	constexpr float ImpulseValue = 1.0f;
diff --git a/Engine/Source/Editor/Experimental/EditorInteractiveToolsFramework/Private/ViewportInteractions/ViewportInteractionsBehaviorSource.cpp b/Engine/Source/Editor/Experimental/EditorInteractiveToolsFramework/Private/ViewportInteractions/ViewportInteractionsBehaviorSource.cpp
--- a/Engine/Source/Editor/Experimental/EditorInteractiveToolsFramework/Private/ViewportInteractions/ViewportInteractionsBehaviorSource.cpp
+++ b/Engine/Source/Editor/Experimental/EditorInteractiveToolsFramework/Private/ViewportInteractions/ViewportInteractionsBehaviorSource.cpp
@@ -2,16 +2,18 @@
 
 #include "ViewportInteractions/ViewportInteractionsBehaviorSource.h"
 
-#include "ContextObjectStore.h"
 #include "BaseBehaviors/KeyAsModifierInputBehavior.h"
+#include "ContextObjectStore.h"
+#include "EditorGizmos/EditorTransformGizmo.h"
+#include "EditorGizmos/EditorTransformGizmoUtil.h"
+#include "EditorInteractiveGizmoManager.h"
 #include "EditorModeManager.h"
 #include "EditorViewportClient.h"
+#include "Framework/Application/SlateApplication.h"
 #include "Tools/EdModeInteractiveToolsContext.h"
 #include "UnrealClient.h"
 #include "ViewportClientNavigationHelper.h"
-#include "EditorDragTools/DragToolInteraction.h"
 #include "ViewportInteractions/ViewportInteraction.h"
-#include "Framework/Application/SlateApplication.h"
 
 namespace UE::Editor::ViewportInteractions
 {
@@ -222,7 +224,7 @@ void UViewportInteractionsBehaviorSource::Initialize(const UEditorInteractiveToo
 		MouseButtonAsModifierInputBehavior->SetDefaultPriority(UE::Editor::ViewportInteractions::MODIFIERS_PRIORITY);
 		CoreBehaviors.Add(MouseButtonAsModifierInputBehavior);
 	}
-	
+
 	Reset();
 }
 
@@ -234,6 +236,42 @@ void UViewportInteractionsBehaviorSource::RegisterBehaviorSources()
 		{
 			InputRouter->RegisterSource(this);
 		}
+
+		// Gizmos delegates
+		{
+			if (UInteractiveToolManager* ToolManager = InteractiveToolsContext->ToolManager)
+			{
+				// Trying to gather Transform gizmo
+				if (UInteractiveGizmoManager* PairedGizmoManager = ToolManager->GetPairedGizmoManager())
+				{
+					const TArray<UInteractiveGizmo*> Gizmos =
+						PairedGizmoManager->FindAllGizmosOfType("EditorTransformGizmoBuilder");
+
+					if (!Gizmos.IsEmpty())
+					{
+						TransformGizmo = Cast<UEditorTransformGizmo>(Gizmos[0]);
+						RegisterProxyDelegates();
+					}
+				}
+
+				// If not available yet, register to creation delegate so we can setup as soon as it exists
+				if (!TransformGizmo)
+				{
+					// If New TRS Gizmos are not enabled, we need to be able to know as soon as they are enabled
+					// This allows to retrieve the Transform Gizmo, from which we register to drag Begin and End delegates
+					if (UContextObjectStore* ContextStore = ToolManager->GetContextObjectStore())
+					{
+						if (UEditorTransformGizmoContextObject* ContextObject =
+								ContextStore->FindContext<UEditorTransformGizmoContextObject>())
+						{
+							ContextObject->OnGizmoCreatedDelegate().AddUObject(
+								this, &UViewportInteractionsBehaviorSource::OnGizmoCreatedDelegate
+							);
+						}
+					}
+				}
+			}
+		}
 	}
 }
 
@@ -245,6 +283,8 @@ void UViewportInteractionsBehaviorSource::DeregisterBehaviorSources()
 		{
 			InputRouter->DeregisterSource(this);
 		}
+
+		DeregisterProxyDelegates();
 	}
 }
 
@@ -383,6 +423,7 @@ void UViewportInteractionsBehaviorSource::RemoveInteraction(const FName InIntera
 
 void UViewportInteractionsBehaviorSource::Reset()
 {
+	DeregisterProxyDelegates();
 	DeregisterBehaviorSources();
 	BehaviorSet->RemoveAll();
 
@@ -568,3 +609,96 @@ void UViewportInteractionsBehaviorSource::SetShiftKeyState(bool bInDown)
 	bIsShiftDown = bInDown;
 	OnShiftKeyStateChangedDelegate.Broadcast(bIsShiftDown);
 }
+
+void UViewportInteractionsBehaviorSource::RegisterProxyDelegates()
+{
+	if (TransformGizmo)
+	{
+		// Register Proxy Delegates
+		if (TObjectPtr<UTransformProxy> TransformProxy = TransformGizmo->ActiveTarget)
+		{
+			TransformProxy->OnBeginTransformEdit.AddUObject(this, &UViewportInteractionsBehaviorSource::OnGizmoMovementBegin);
+			TransformProxy->OnTransformChanged.AddUObject(this, &UViewportInteractionsBehaviorSource::OnGizmoTransformChanged);
+			TransformProxy->OnEndTransformEdit.AddUObject(this, &UViewportInteractionsBehaviorSource::OnGizmoMovementEnd);
+		}
+	}
+}
+
+void UViewportInteractionsBehaviorSource::DeregisterProxyDelegates()
+{
+	if (TransformGizmo)
+	{
+		// Unregister Proxy delegates
+		if (TObjectPtr<UTransformProxy> TransformProxy = TransformGizmo->ActiveTarget)
+		{
+			TransformProxy->OnBeginTransformEdit.RemoveAll(this);
+			TransformProxy->OnTransformChanged.RemoveAll(this);
+			TransformProxy->OnEndTransformEdit.RemoveAll(this);
+		}
+
+		TransformGizmo = nullptr;
+	}
+
+	if (TStrongObjectPtr<const UEditorInteractiveToolsContext> InteractiveToolsContext = EditorInteractiveToolsContextWeak.Pin())
+	{
+		if (UInteractiveToolManager* ToolManager = InteractiveToolsContext->ToolManager)
+		{
+			if (UContextObjectStore* ContextStore = ToolManager->GetContextObjectStore())
+			{
+				if (UEditorTransformGizmoContextObject* ContextObject = ContextStore->FindContext<UEditorTransformGizmoContextObject>())
+				{
+					ContextObject->OnGizmoCreatedDelegate().RemoveAll(this);
+				}
+			}
+		}
+	}
+}
+
+void UViewportInteractionsBehaviorSource::OnGizmoCreatedDelegate(UTransformGizmo* InTransformGizmo)
+{
+	if (!TransformGizmo)
+	{
+		TransformGizmo = Cast<UEditorTransformGizmo>(InTransformGizmo);
+
+		RegisterProxyDelegates();
+	}
+}
+
+void UViewportInteractionsBehaviorSource::OnGizmoMovementBegin(UTransformProxy* InTransformProxy)
+{
+	bGizmoDragging = true;
+
+	for (const TTuple<FName, TObjectPtr<UViewportInteraction>>& Pair : ViewportInteractions)
+	{
+		if (UViewportInteraction* const Interaction = Pair.Value)
+		{
+			Interaction->OnGizmoMovementBegin(InTransformProxy);
+		}
+	}
+}
+
+void UViewportInteractionsBehaviorSource::OnGizmoTransformChanged(
+	UTransformProxy* InTransformProxy, UE::Math::TTransform<double> InTransform
+)
+{
+	for (const TTuple<FName, TObjectPtr<UViewportInteraction>>& Pair : ViewportInteractions)
+	{
+		if (UViewportInteraction* const Interaction = Pair.Value)
+		{
+			Interaction->OnGizmoTransformChanged(InTransformProxy, InTransform);
+		}
+	}
+}
+
+void UViewportInteractionsBehaviorSource::OnGizmoMovementEnd(UTransformProxy* InTransformProxy)
+{
+	bGizmoDragging = false;
+
+	for (const TTuple<FName, TObjectPtr<UViewportInteraction>>& Pair : ViewportInteractions)
+	{
+		if (UViewportInteraction* const Interaction = Pair.Value)
+		{
+			Interaction->OnGizmoMovementEnd(InTransformProxy);
+		}
+	}
+}
diff --git a/Engine/Source/Editor/Experimental/EditorInteractiveToolsFramework/Public/EditorDragTools/DuplicateDragSelectionInteraction.h b/Engine/Source/Editor/Experimental/EditorInteractiveToolsFramework/Public/EditorDragTools/DuplicateDragSelectionInteraction.h
--- a/Engine/Source/Editor/Experimental/EditorInteractiveToolsFramework/Public/EditorDragTools/DuplicateDragSelectionInteraction.h
+++ b/Engine/Source/Editor/Experimental/EditorInteractiveToolsFramework/Public/EditorDragTools/DuplicateDragSelectionInteraction.h
@@ -3,7 +3,6 @@
 #pragma once
 
 #include "BaseGizmos/TransformProxy.h"
-#include "EditorGizmos/TransformGizmo.h"
 #include "ViewportInteractions/ViewportInteraction.h"
 
 #include "DuplicateDragSelectionInteraction.generated.h"
@@ -39,27 +38,14 @@ protected:
 	 */
 	static void DuplicateSelection(bool bInSelectNewElements);
 
-	void OnGizmoMovementBegin(UTransformProxy* InTransformProxy);
-	void OnGizmoMovementEnd(UTransformProxy* InTransformProxy);
-
-	void RegisterProxyDelegates();
-	void Reset();
-
-	void OnGizmoCreatedDelegate(UTransformGizmo* InTransformGizmo);
-	void OnUsesNewTRSGizmosChanged(bool bInUseNewTRSGizmos);
+	virtual void OnGizmoMovementBegin(UTransformProxy* InTransformProxy) override;
+	virtual void OnGizmoMovementEnd(UTransformProxy* InTransformProxy) override;
 
 	virtual void Initialize(UViewportInteractionsBehaviorSource* const InViewportInteractionsBehaviorSource) override;
 
 private:
 	void OnModifierKeyChanged(bool bInIsDown);
 
-	UPROPERTY(Transient)
-	TObjectPtr<UInteractiveToolManager> ToolManager;
-
-	UPROPERTY(Transient)
-	TObjectPtr<UEditorTransformGizmo> TransformGizmo;
-
-	bool bGizmoIsDragged = false;
 	bool bModifierKeyIsPressed = false;
 
 	/** Keeps the duplicate + transform transaction alive between gizmo movement begin and end */
diff --git a/Engine/Source/Editor/Experimental/EditorInteractiveToolsFramework/Public/EditorDragTools/MoveCameraWithObjectInteraction.h b/Engine/Source/Editor/Experimental/EditorInteractiveToolsFramework/Public/EditorDragTools/MoveCameraWithObjectInteraction.h
--- a/Engine/Source/Editor/Experimental/EditorInteractiveToolsFramework/Public/EditorDragTools/MoveCameraWithObjectInteraction.h
+++ b/Engine/Source/Editor/Experimental/EditorInteractiveToolsFramework/Public/EditorDragTools/MoveCameraWithObjectInteraction.h
@@ -36,29 +36,13 @@ public:
 protected:
 
 	UE_API void OnShiftKeyChanged(bool bDown);
-	UE_API void OnGizmoMovementBegin(UTransformProxy* InTransformProxy);
-	UE_API void OnGizmoMovementEnd(UTransformProxy* InTransformProxy);
-
-	UE_API void RegisterProxyDelegates();
-	UE_API void UnregisterProxyDelegates();
-
-	UE_API void OnGizmoCreatedDelegate(UTransformGizmo* InTransformGizmo);
-	UE_API void OnUsesNewTRSGizmosChanged(bool bInUseNewTRSGizmos);
+	UE_API virtual void OnGizmoMovementBegin(UTransformProxy* InTransformProxy) override;
+	UE_API virtual void OnGizmoMovementEnd(UTransformProxy* InTransformProxy) override;
 
 private:
 	void UpdateCameraFollows() const;
 
-	UPROPERTY(Transient)
-	TObjectPtr<UInteractiveToolManager> ToolManager;
-	UPROPERTY(Transient)
-	TObjectPtr<UEditorTransformGizmo> TransformGizmo;
-
-	bool bGizmoIsDragged = false;
 	bool bIsModifierActive = false;
-
-	FDelegateHandle OnBeginPivotEditDelegate;
-	FDelegateHandle OnEndPivotEditDelegate;
-	FDelegateHandle OnTransformChangedDelegate;
 };
 
 #undef UE_API
diff --git a/Engine/Source/Editor/Experimental/EditorInteractiveToolsFramework/Public/ViewportInteractions/ViewportInteraction.h b/Engine/Source/Editor/Experimental/EditorInteractiveToolsFramework/Public/ViewportInteractions/ViewportInteraction.h
--- a/Engine/Source/Editor/Experimental/EditorInteractiveToolsFramework/Public/ViewportInteractions/ViewportInteraction.h
+++ b/Engine/Source/Editor/Experimental/EditorInteractiveToolsFramework/Public/ViewportInteractions/ViewportInteraction.h
@@ -15,6 +15,7 @@ class FViewportCameraMover;
 class FViewportClickHandler;
 class FViewportCommandsHandler;
 class UInputBehavior;
+class UTransformProxy;
 class UViewportInteractionsBehaviorSource;
 
 namespace UE::Editor::ViewportInteractions
@@ -99,6 +100,10 @@ public:
 	/** Called when the interaction is removed from its behavior source. Override to make sure the interaction is properly cleaned up */
 	virtual void Shutdown() {}
 
+	virtual void OnGizmoMovementBegin(UTransformProxy* InTransformProxy) {}
+	virtual void OnGizmoTransformChanged(UTransformProxy* InTransformProxy, UE::Math::TTransform<double> InTransform) {}
+	virtual void OnGizmoMovementEnd(UTransformProxy* InTransformProxy) {}
+
 	/** Called every frame. */
 	virtual void Tick(float InDeltaTime) const {}
 	
diff --git a/Engine/Source/Editor/Experimental/EditorInteractiveToolsFramework/Public/ViewportInteractions/ViewportInteractionsBehaviorSource.h b/Engine/Source/Editor/Experimental/EditorInteractiveToolsFramework/Public/ViewportInteractions/ViewportInteractionsBehaviorSource.h
--- a/Engine/Source/Editor/Experimental/EditorInteractiveToolsFramework/Public/ViewportInteractions/ViewportInteractionsBehaviorSource.h
+++ b/Engine/Source/Editor/Experimental/EditorInteractiveToolsFramework/Public/ViewportInteractions/ViewportInteractionsBehaviorSource.h
@@ -16,7 +16,10 @@ class FUICommandInfo;
 class FViewportCameraMover;
 class FViewportClickHandler;
 class UEditorInteractiveToolsContext;
+class UEditorTransformGizmo;
 class UInputBehaviorSet;
+class UTransformGizmo;
+class UTransformProxy;
 class UViewportInteraction;
 
 namespace UE::Editor::ViewportInteractions
@@ -229,6 +232,18 @@ public:
 
 	void SetIsMouseLooking(bool bInIsLooking);
 
+	/** Whether the Editor Transform Gizmo is being interacted with */
+	bool IsGizmoDragging() const
+	{
+		return bGizmoDragging;
+	}
+
+	/** If available, return the current Editor Transform Gizmo */
+	UEditorTransformGizmo* GetEditorTransformGizmo() const
+	{
+		return TransformGizmo;
+	}
+
 	FEditorModeTools* GetModeTools() const;
 
 	FEditorViewportClient* GetEditorViewportClient() const;
@@ -288,6 +303,14 @@ private:
 	void SetCtrlKeyState(bool bInDown);
 	void SetShiftKeyState(bool bInDown);
 
+	void RegisterProxyDelegates();
+	void DeregisterProxyDelegates();
+
+	void OnGizmoCreatedDelegate(UTransformGizmo* InTransformGizmo);
+	void OnGizmoMovementBegin(UTransformProxy* InTransformProxy);
+	void OnGizmoTransformChanged(UTransformProxy* InTransformProxy, UE::Math::TTransform<double> InTransform);
+	void OnGizmoMovementEnd(UTransformProxy* InTransformProxy);
+
 	bool bIsShiftDown;
 	bool bIsCtrlDown;
 	bool bIsAltDown;
@@ -312,4 +335,9 @@ private:
 	FOnModifierKeyChanged OnAltKeyStateChangedDelegate;
 	FOnModifierKeyChanged OnCtrlKeyStateChangedDelegate;
 	FOnModifierKeyChanged OnShiftKeyStateChangedDelegate;
+
+	UPROPERTY(Transient)
+	TObjectPtr<UEditorTransformGizmo> TransformGizmo;
+
+	bool bGizmoDragging = false;
 };
