Gizmos: fixed inverted Z rotation (regression from 48278321)

#rnx
#Synced-CL 48349740

diff --git a/Engine/Source/Editor/Experimental/EditorInteractiveToolsFramework/Private/EditorGizmos/TransformGizmo.cpp b/Engine/Source/Editor/Experimental/EditorInteractiveToolsFramework/Private/EditorGizmos/TransformGizmo.cpp
--- a/Engine/Source/Editor/Experimental/EditorInteractiveToolsFramework/Private/EditorGizmos/TransformGizmo.cpp
+++ b/Engine/Source/Editor/Experimental/EditorInteractiveToolsFramework/Private/EditorGizmos/TransformGizmo.cpp
@@ -3029,7 +3029,7 @@ FVector2D UTransformGizmo::GetScreenRotateAxisDir(const FInputDeviceRay& InPress
 	}
 
 	// store world origin and axis
-	static const TArray RotateAxis({FVector::XAxisVector, FVector::YAxisVector, FVector::ZAxisVector});
+	static const TArray RotateAxis({FVector::XAxisVector, FVector::YAxisVector, -FVector::ZAxisVector});
 	const FVector WorldAxis = RotateAxis.IsValidIndex(RotateID) ? GetWorldAxis(RotateAxis[RotateID]) : GizmoViewContext->GetViewDirection();
 
 	return GetWorldToScreenRotateAxisDir(InPressPos, WorldAxis);
