[ITF Viewport Interactions] Right Click invoked after WASD Navigation

#Synced-CL 48027920

diff --git a/Engine/Source/Editor/Experimental/EditorInteractiveToolsFramework/Private/ViewportInteractions/ViewportCameraRotateInteraction.cpp b/Engine/Source/Editor/Experimental/EditorInteractiveToolsFramework/Private/ViewportInteractions/ViewportCameraRotateInteraction.cpp
--- a/Engine/Source/Editor/Experimental/EditorInteractiveToolsFramework/Private/ViewportInteractions/ViewportCameraRotateInteraction.cpp
+++ b/Engine/Source/Editor/Experimental/EditorInteractiveToolsFramework/Private/ViewportInteractions/ViewportCameraRotateInteraction.cpp
@@ -89,6 +89,8 @@ void UViewportCameraRotateInteraction::Tick(float InDeltaTime) const
 
 	if (RotateYawImpulse != 0 || RotatePitchImpulse != 0)
 	{
+		BehaviorSource->SetCameraHasMoved(true);
+
 		if (FEditorViewportClient* EditorViewportClient = GetEditorViewportClient())
 		{
 			if (FViewportClientNavigationHelper* NavigationHelper = EditorViewportClient->GetViewportNavigationHelper())
diff --git a/Engine/Source/Editor/Experimental/EditorInteractiveToolsFramework/Private/ViewportInteractions/ViewportCameraTranslateInteraction.cpp b/Engine/Source/Editor/Experimental/EditorInteractiveToolsFramework/Private/ViewportInteractions/ViewportCameraTranslateInteraction.cpp
--- a/Engine/Source/Editor/Experimental/EditorInteractiveToolsFramework/Private/ViewportInteractions/ViewportCameraTranslateInteraction.cpp
+++ b/Engine/Source/Editor/Experimental/EditorInteractiveToolsFramework/Private/ViewportInteractions/ViewportCameraTranslateInteraction.cpp
@@ -154,6 +154,8 @@ void UViewportCameraTranslateInteraction::Tick(float InDeltaTime) const
 
 	if (ForwardBackwardImpulse != 0 || RightLeftImpulse != 0 || WorldUpDownImpulse != 0 || LocalUpDownImpulse != 0)
 	{
+		BehaviorSource->SetCameraHasMoved(true);
+
 		if (FEditorViewportClient* EditorViewportClient = GetEditorViewportClient())
 		{
 			if (FViewportClientNavigationHelper* NavigationHelper = EditorViewportClient->GetViewportNavigationHelper())
diff --git a/Engine/Source/Editor/Experimental/EditorInteractiveToolsFramework/Private/ViewportInteractions/ViewportClickInteraction.cpp b/Engine/Source/Editor/Experimental/EditorInteractiveToolsFramework/Private/ViewportInteractions/ViewportClickInteraction.cpp
--- a/Engine/Source/Editor/Experimental/EditorInteractiveToolsFramework/Private/ViewportInteractions/ViewportClickInteraction.cpp
+++ b/Engine/Source/Editor/Experimental/EditorInteractiveToolsFramework/Private/ViewportInteractions/ViewportClickInteraction.cpp
@@ -41,7 +41,15 @@ void UViewportClickInteraction::OnClicked(const FInputDeviceRay& InClickPos)
 	{
 		if (UViewportInteractionsBehaviorSource* const ViewportInteractionBehaviorSource = GetViewportInteractionsBehaviorSource())
 		{
+			// Store info about camera movement locally, since SetIsMouseLooking call will reset it
+			const bool bCameraHasMoved = ViewportInteractionBehaviorSource->HasCameraMoved();
 			ViewportInteractionBehaviorSource->SetIsMouseLooking(false);
+
+			// If camera has moved, this was definitely not a click
+			if (bCameraHasMoved)
+			{
+				return;
+			}
 		}
 	}
 
diff --git a/Engine/Source/Editor/Experimental/EditorInteractiveToolsFramework/Private/ViewportInteractions/ViewportFOVInteraction.cpp b/Engine/Source/Editor/Experimental/EditorInteractiveToolsFramework/Private/ViewportInteractions/ViewportFOVInteraction.cpp
--- a/Engine/Source/Editor/Experimental/EditorInteractiveToolsFramework/Private/ViewportInteractions/ViewportFOVInteraction.cpp
+++ b/Engine/Source/Editor/Experimental/EditorInteractiveToolsFramework/Private/ViewportInteractions/ViewportFOVInteraction.cpp
@@ -96,6 +96,11 @@ void UViewportFOVInteraction::Tick(float InDeltaTime) const
 		{
 			if (ZoomOutInImpulse != 0)
 			{
+				if (UViewportInteractionsBehaviorSource* BehaviorSource = GetViewportInteractionsBehaviorSource())
+				{
+					BehaviorSource->SetCameraHasMoved(true);
+				}
+
 				NavigationHelper->ImpulseDataDelta.ZoomOutInImpulse += ZoomOutInImpulse;
 			}
 		}
diff --git a/Engine/Source/Editor/Experimental/EditorInteractiveToolsFramework/Private/ViewportInteractions/ViewportInteractionsBehaviorSource.cpp b/Engine/Source/Editor/Experimental/EditorInteractiveToolsFramework/Private/ViewportInteractions/ViewportInteractionsBehaviorSource.cpp
--- a/Engine/Source/Editor/Experimental/EditorInteractiveToolsFramework/Private/ViewportInteractions/ViewportInteractionsBehaviorSource.cpp
+++ b/Engine/Source/Editor/Experimental/EditorInteractiveToolsFramework/Private/ViewportInteractions/ViewportInteractionsBehaviorSource.cpp
@@ -501,6 +501,7 @@ void UViewportInteractionsBehaviorSource::SetIsMouseLooking(bool bInIsLooking)
 		}
 		else
 		{
+			bCameraHasMoved = false;
 			ClearMouseCursorOverride();
 		}
 
@@ -510,6 +511,11 @@ void UViewportInteractionsBehaviorSource::SetIsMouseLooking(bool bInIsLooking)
 	OnMouseLookingStateChanged().Broadcast(bInIsLooking);
 }
 
+void UViewportInteractionsBehaviorSource::SetCameraHasMoved(bool bInHasMoved)
+{
+	bCameraHasMoved = bInHasMoved;
+}
+
 FEditorModeTools* UViewportInteractionsBehaviorSource::GetModeTools() const
 {
 	if (const UEditorInteractiveToolsContext* const InteractiveToolsContext = EditorInteractiveToolsContextWeak.Get())
diff --git a/Engine/Source/Editor/Experimental/EditorInteractiveToolsFramework/Public/ViewportInteractions/ViewportInteractionsBehaviorSource.h b/Engine/Source/Editor/Experimental/EditorInteractiveToolsFramework/Public/ViewportInteractions/ViewportInteractionsBehaviorSource.h
--- a/Engine/Source/Editor/Experimental/EditorInteractiveToolsFramework/Public/ViewportInteractions/ViewportInteractionsBehaviorSource.h
+++ b/Engine/Source/Editor/Experimental/EditorInteractiveToolsFramework/Public/ViewportInteractions/ViewportInteractionsBehaviorSource.h
@@ -232,6 +232,13 @@ public:
 
 	void SetIsMouseLooking(bool bInIsLooking);
 
+	bool HasCameraMoved() const
+	{
+		return bCameraHasMoved;
+	}
+
+	void SetCameraHasMoved(bool bInHasMoved);
+
 	/** Whether the Editor Transform Gizmo is being interacted with */
 	bool IsGizmoDragging() const
 	{
@@ -340,4 +347,5 @@ private:
 	TObjectPtr<UEditorTransformGizmo> TransformGizmo;
 
 	bool bGizmoDragging = false;
+	bool bCameraHasMoved = false;
 };
