[ITF Viewport Interactions] Selection - Ctrl click not deselecting

- "distance" priority values
- remove unused priority

#Synced-CL 48364501

diff --git a/Engine/Source/Editor/Experimental/EditorInteractiveToolsFramework/Public/ViewportInteractions/ViewportInteraction.h b/Engine/Source/Editor/Experimental/EditorInteractiveToolsFramework/Public/ViewportInteractions/ViewportInteraction.h
--- a/Engine/Source/Editor/Experimental/EditorInteractiveToolsFramework/Public/ViewportInteractions/ViewportInteraction.h
+++ b/Engine/Source/Editor/Experimental/EditorInteractiveToolsFramework/Public/ViewportInteractions/ViewportInteraction.h
@@ -28,11 +28,10 @@ constexpr int LeftMouseButtonMod = 4;
 constexpr int RightMouseButtonMod = 5;
 constexpr int MiddleMouseButtonMod = 6;
 
-static constexpr int VIEWPORT_INTERACTIONS_DEFAULT_PRIORITY = FInputCapturePriority::DEFAULT_TOOL_PRIORITY - 1;
-static constexpr int DRAG_TOOLS_DEFAULT_PRIORITY = FInputCapturePriority::DEFAULT_GIZMO_PRIORITY - 1;
+static constexpr int VIEWPORT_INTERACTIONS_DEFAULT_PRIORITY = FInputCapturePriority::DEFAULT_TOOL_PRIORITY - 5;
 static constexpr int DRAG_CONFIRMED_PRIORITY = FInputCapturePriority::DEFAULT_TOOL_PRIORITY + 5;
-static constexpr int CLICK_PRIORITY = FInputCapturePriority::DEFAULT_TOOL_PRIORITY + 6;
-static constexpr int DRAG_PRIORITY = FInputCapturePriority::DEFAULT_TOOL_PRIORITY + 7;
+static constexpr int CLICK_PRIORITY = FInputCapturePriority::DEFAULT_TOOL_PRIORITY + 10;
+static constexpr int DRAG_PRIORITY = FInputCapturePriority::DEFAULT_TOOL_PRIORITY + 15;
 
 /**
  * Possible types of Viewport Interactions.
