[ITF Drag Tools] Add unbounded mouse cursor support

- Add to `FMouseInputDeviceState` the unbounded version of Mouse position, delta and world ray

- `UModeManagerInteractiveToolsContext` is taking care of setting the values of the new `FMouseInputDeviceState` members.

- Added `FViewport::GetUnboundedMousePos(FIntPoint&)` to retrieve a cached unbounded position, valid as long as at least one mouse button is pressed

- Added support for unbounded mouse position to `UViewportClickDragBehavior`

#Synced-CL 48901306

diff --git a/Engine/Source/Editor/Experimental/EditorInteractiveToolsFramework/Private/Behaviors/ViewportClickDragBehavior.cpp b/Engine/Source/Editor/Experimental/EditorInteractiveToolsFramework/Private/Behaviors/ViewportClickDragBehavior.cpp
--- a/Engine/Source/Editor/Experimental/EditorInteractiveToolsFramework/Private/Behaviors/ViewportClickDragBehavior.cpp
+++ b/Engine/Source/Editor/Experimental/EditorInteractiveToolsFramework/Private/Behaviors/ViewportClickDragBehavior.cpp
@@ -35,6 +35,8 @@ FInputCapturePriority UViewportClickDragBehavior::GetPriority()
 
 FInputCaptureRequest UViewportClickDragBehavior::WantsCapture(const FInputDeviceState& InInputState)
 {
+	const FVector2D& MousePosition = GetMousePosition(InInputState);
+
 	if (!MouseTraveledDistance.IsSet())
 	{
 		ActiveBindingComplexity = GetBindingComplexity(EInputStage::Start, InInputState);
@@ -47,7 +49,7 @@ FInputCaptureRequest UViewportClickDragBehavior::WantsCapture(const FInputDevice
 			MouseTraveledDistance = 0.0f;
 
 			// Cache button down location
-			MouseDownPosition = InInputState.Mouse.Position2D;
+			MouseDownPosition = MousePosition;
 		}
 		else
 		{
@@ -66,7 +68,7 @@ FInputCaptureRequest UViewportClickDragBehavior::WantsCapture(const FInputDevice
 			}
 		}
 		
-		MouseTraveledDistance = static_cast<float>(FVector2D::Distance(InInputState.Mouse.Position2D, MouseDownPosition));
+		MouseTraveledDistance = static_cast<float>(FVector2D::Distance(MousePosition, MouseDownPosition));
 	}
 
 	const FInputRayHit HitResult = Target->CanBeginClickDragSequence(GetDeviceRay(InInputState));
@@ -92,6 +94,8 @@ FInputCaptureUpdate UViewportClickDragBehavior::BeginCapture(const FInputDeviceS
 
 FInputCaptureUpdate UViewportClickDragBehavior::UpdateCapture(const FInputDeviceState& InputState, const FInputCaptureData& CaptureData)
 {
+	const FVector2D& MousePosition = GetMousePosition(InputState);
+
 	// This happens when a keyboard event is routed to the active mouse capture for updating modifier states
 	// The behavior should always continue here so that subsequent mouse updates can end capture more correctly	
 	if (!InputState.IsFromDevice(EInputDevices::Mouse))
@@ -112,7 +116,7 @@ FInputCaptureUpdate UViewportClickDragBehavior::UpdateCapture(const FInputDevice
 		return FInputCaptureUpdate::End();
 	}
 	
-	MouseTraveledDistance = static_cast<float>(FVector2D::Distance(InputState.Mouse.Position2D, MouseDownPosition));
+	MouseTraveledDistance = static_cast<float>(FVector2D::Distance(MousePosition, MouseDownPosition));
 
 	if (!bDragging)
 	{
@@ -153,6 +157,34 @@ void UViewportClickDragBehavior::SetBindings(const TArray<UE::Editor::ViewportIn
 	bRequiresTrigger = Algo::FindByPredicate(Bindings, [](const UE::Editor::ViewportInteractions::FButtonBinding& InBinding) { return InBinding.bTriggersStart; }) != nullptr;
 }
 
+FInputDeviceRay UViewportClickDragBehavior::GetDeviceRay(const FInputDeviceState& InInputDeviceState)
+{
+	if (InInputDeviceState.IsFromDevice(EInputDevices::Mouse))
+	{
+		const FVector2D& MousePosition = GetMousePosition(InInputDeviceState);
+		const FRay& WorldRay = UsesUnboundedCursor() ? InInputDeviceState.Mouse.UnboundedWorldRay : InInputDeviceState.Mouse.WorldRay;
+
+		return FInputDeviceRay(WorldRay, MousePosition);
+	}
+
+	return FInputDeviceRay(FRay(FVector::ZeroVector, FVector(0, 0, 1), true));
+}
+
+const FVector2D& UViewportClickDragBehavior::GetMousePosition(const FInputDeviceState& InInputDeviceState) const
+{
+	return UsesUnboundedCursor() ? InInputDeviceState.Mouse.UnboundedPosition2D : InInputDeviceState.Mouse.Position2D;
+}
+
+bool UViewportClickDragBehavior::UsesUnboundedCursor() const
+{
+	return UnboundedCursor.IsSet() ? UnboundedCursor.Get() : false;
+}
+
+void UViewportClickDragBehavior::SetUsesUnboundedCursor(TAttribute<bool> InUsesUnboundedCursor)
+{
+	UnboundedCursor = InUsesUnboundedCursor;
+}
+
 namespace UE::Editor::ViewportInteractions::Private
 {
 	FDeviceButtonState DownState(const FKey& InKey, bool bInDown)
@@ -187,7 +219,7 @@ void UViewportClickDragBehavior::ProcessBinding(EInputStage Stage, const FInputD
 				{
 					OutComplexity += ChordComplexity;
 					bOutTrigger = bOutTrigger || bHasTrigger; 
-					return;	
+					return;
 				}
 			}
 		}
@@ -298,16 +330,6 @@ void UViewportClickDragBehavior::OnDragEndInternal(const FInputDeviceState& Inpu
 	Target->OnDragEnd(GetDeviceRay(InputState));
 }
 
-FInputDeviceRay UViewportClickDragBehavior::GetDeviceRay(const FInputDeviceState& input)
-{
-	if (input.IsFromDevice(EInputDevices::Mouse))
-	{
-		return FInputDeviceRay(input.Mouse.WorldRay, input.Mouse.Position2D);
-	}
-
-	return FInputDeviceRay(FRay(FVector::ZeroVector, FVector(0, 0, 1), true));
-}
-
 void UViewportClickDragBehavior::OnStateUpdatedInternal(const FInputDeviceState& InInputDeviceState)
 {
 	Target->OnStateUpdated(InInputDeviceState);
diff --git a/Engine/Source/Editor/Experimental/EditorInteractiveToolsFramework/Public/Behaviors/ViewportClickDragBehavior.h b/Engine/Source/Editor/Experimental/EditorInteractiveToolsFramework/Public/Behaviors/ViewportClickDragBehavior.h
--- a/Engine/Source/Editor/Experimental/EditorInteractiveToolsFramework/Public/Behaviors/ViewportClickDragBehavior.h
+++ b/Engine/Source/Editor/Experimental/EditorInteractiveToolsFramework/Public/Behaviors/ViewportClickDragBehavior.h
@@ -158,7 +158,16 @@ public:
 	EDITORINTERACTIVETOOLSFRAMEWORK_API void SetBindings(const TArray<UE::Editor::ViewportInteractions::FButtonBinding>& InBindings);
 
 	/** @return current 3D world ray and optional 2D position for Target Device */
-	EDITORINTERACTIVETOOLSFRAMEWORK_API virtual FInputDeviceRay GetDeviceRay(const FInputDeviceState& input);
+	EDITORINTERACTIVETOOLSFRAMEWORK_API virtual FInputDeviceRay GetDeviceRay(const FInputDeviceState& InInputDeviceState);
+
+	EDITORINTERACTIVETOOLSFRAMEWORK_API virtual const FVector2D& GetMousePosition(const FInputDeviceState& InInputDeviceState) const;
+
+	EDITORINTERACTIVETOOLSFRAMEWORK_API virtual bool UsesUnboundedCursor() const;
+
+	/**
+	 * Call to set the default unbounded cursor attribute.
+	 */
+	EDITORINTERACTIVETOOLSFRAMEWORK_API void SetUsesUnboundedCursor(TAttribute<bool> InUsesUnboundedCursor);
 
 	bool IsDragging() const { return bDragging; }
 
@@ -209,6 +218,8 @@ protected:
 
 	/** True when capturing and dragging */
 	bool bDragging = false;
+
+	TAttribute<bool> UnboundedCursor;
 };
 
 /**
diff --git a/Engine/Source/Editor/UnrealEd/Private/Tools/EdModeInteractiveToolsContext.cpp b/Engine/Source/Editor/UnrealEd/Private/Tools/EdModeInteractiveToolsContext.cpp
--- a/Engine/Source/Editor/UnrealEd/Private/Tools/EdModeInteractiveToolsContext.cpp
+++ b/Engine/Source/Editor/UnrealEd/Private/Tools/EdModeInteractiveToolsContext.cpp
@@ -1119,8 +1119,14 @@ bool UModeManagerInteractiveToolsContext::MouseEnter(FEditorViewportClient* View
 	UE_LOG(LogInteractiveToolsContext, Warning, TEXT("MOUSE ENTER"));
 #endif
 
-	CurrentMouseState.Mouse.Position2D = FVector2D(x, y);
-	CurrentMouseState.Mouse.WorldRay = GetRayFromMousePos(ViewportClient, Viewport, x, y);
+	if (!bIsTrackingMouse)
+	{
+		CurrentMouseState.Mouse.Position2D = FVector2D(x, y);
+		CurrentMouseState.Mouse.WorldRay = GetRayFromMousePos(ViewportClient, Viewport, x, y);
+
+		CurrentMouseState.Mouse.UnboundedPosition2D = CurrentMouseState.Mouse.Position2D;
+		CurrentMouseState.Mouse.UnboundedWorldRay = CurrentMouseState.Mouse.WorldRay;
+	}
 
 	return false;
 }
@@ -1137,6 +1143,12 @@ bool UModeManagerInteractiveToolsContext::MouseMove(FEditorViewportClient* Viewp
 	FInputDeviceState InputState = CurrentMouseState;
 	InputState.InputDevice = EInputDevices::Mouse;
 
+	if (!bIsTrackingMouse)
+	{
+		CurrentMouseState.Mouse.UnboundedPosition2D = CurrentMouseState.Mouse.Position2D;
+		CurrentMouseState.Mouse.UnboundedWorldRay = CurrentMouseState.Mouse.WorldRay;
+	}
+
 	FModifierKeysState ModifierKeys = FSlateApplication::Get().GetModifierKeys();
 	InputState.SetModifierKeyStates(
 		ModifierKeys.IsShiftDown(), ModifierKeys.IsAltDown(),
@@ -1171,6 +1183,12 @@ bool UModeManagerInteractiveToolsContext::MouseLeave(FEditorViewportClient* View
 	UE_LOG(LogInteractiveToolsContext, Warning, TEXT("MOUSE LEAVE"));
 #endif
 
+	if (!bIsTrackingMouse)
+	{
+		CurrentMouseState.Mouse.UnboundedPosition2D = CurrentMouseState.Mouse.Position2D;
+		CurrentMouseState.Mouse.UnboundedWorldRay = CurrentMouseState.Mouse.WorldRay;
+	}
+
 	return false;
 }
 
@@ -1191,6 +1209,17 @@ bool UModeManagerInteractiveToolsContext::CapturedMouseMove(FEditorViewportClien
 	CurrentMouseState.Mouse.Position2D = FVector2D(InMouseX, InMouseY);
 	CurrentMouseState.Mouse.WorldRay = GetRayFromMousePos(InViewportClient, InViewport, InMouseX, InMouseY, bHasActiveMouseCapture && bClampMouseToBoundsWhenCaptured);
 
+	// Unbound cursor
+	{
+		FIntPoint UnboundedMousePos;
+		InViewport->GetUnboundedMousePos(UnboundedMousePos);
+		const FVector2D UnboundedPosition = UnboundedMousePos;
+
+		CurrentMouseState.Mouse.UnboundedMouseDelta2D = UnboundedPosition - CurrentMouseState.Mouse.UnboundedPosition2D;
+		CurrentMouseState.Mouse.UnboundedPosition2D = UnboundedPosition;
+		CurrentMouseState.Mouse.UnboundedWorldRay = GetRayFromMousePos(InViewportClient, InViewport, UnboundedPosition.X, UnboundedPosition.Y, false);
+	}
+
 	if (bHasActiveMouseCapture)
 	{
 #ifdef ENABLE_DEBUG_PRINTING
diff --git a/Engine/Source/Runtime/Engine/Private/Slate/SceneViewport.cpp b/Engine/Source/Runtime/Engine/Private/Slate/SceneViewport.cpp
--- a/Engine/Source/Runtime/Engine/Private/Slate/SceneViewport.cpp
+++ b/Engine/Source/Runtime/Engine/Private/Slate/SceneViewport.cpp
@@ -212,6 +212,20 @@ void FSceneViewport::GetMousePos( FIntPoint& MousePosition, const bool bLocalPos
 	}
 }
 
+void FSceneViewport::GetUnboundedMousePos(FIntPoint& OutUnboundedMousePosition, const bool bLocalPosition)
+{
+	if (bLocalPosition)
+	{
+		OutUnboundedMousePosition = CachedUnboundedCursorPos;
+	}
+	else
+	{
+		const FVector2D AbsoluteMousePos = CachedGeometry.LocalToAbsolute(FVector2D(CachedUnboundedCursorPos.X / CachedGeometry.Scale, CachedUnboundedCursorPos.Y / CachedGeometry.Scale));
+		OutUnboundedMousePosition.X = AbsoluteMousePos.X;
+		OutUnboundedMousePosition.Y = AbsoluteMousePos.Y;
+	}
+}
+
 void FSceneViewport::SetMouse( int32 X, int32 Y )
 {
 	const FVector2D NormalizedLocalMousePosition = FVector2D(X, Y) / GetSizeXY();
@@ -236,6 +250,9 @@ void FSceneViewport::UpdateCachedCursorPos( const FGeometry& InGeometry, const F
 
 		CachedCursorPos = LocalPixelMousePos.IntPoint();
 	}
+
+	CachedUnboundedCursorPos.X += InMouseEvent.GetCursorDelta().X;
+	CachedUnboundedCursorPos.Y += InMouseEvent.GetCursorDelta().Y;
 }
 
 void FSceneViewport::UpdateCachedGeometry( const FGeometry& InGeometry )
@@ -602,6 +619,16 @@ FReply FSceneViewport::OnMouseButtonDown( const FGeometry& InGeometry, const FPo
 	// Re-set prevent throttling here as it can get reset when inside of InputKey()
 	CurrentReplyState.PreventThrottling();
 
+	const uint32 DownButtonsCount = InMouseEvent.IsMouseButtonDown(EKeys::LeftMouseButton)
+								  + InMouseEvent.IsMouseButtonDown(EKeys::MiddleMouseButton)
+								  + InMouseEvent.IsMouseButtonDown(EKeys::RightMouseButton);
+
+	// This is the first pressed button: initialize unbounded cursor position to match current cursor position
+	if (DownButtonsCount == 1)
+	{
+		CachedUnboundedCursorPos = CachedCursorPos;
+	}
+
 	return CurrentReplyState;
 }
 
diff --git a/Engine/Source/Runtime/Engine/Public/Slate/SceneViewport.h b/Engine/Source/Runtime/Engine/Public/Slate/SceneViewport.h
--- a/Engine/Source/Runtime/Engine/Public/Slate/SceneViewport.h
+++ b/Engine/Source/Runtime/Engine/Public/Slate/SceneViewport.h
@@ -96,6 +96,12 @@ public:
 	 */
 	ENGINE_API virtual void GetMousePos( FIntPoint& MousePosition, const bool bLocalPosition = true) override;
 
+	/**
+	 * Retrieve current unbounded mouse position. Valid while a mouse button is pressed.
+	 * Useful e.g. for drag operations hiding the mouse cursor.
+	 */
+	ENGINE_API virtual void GetUnboundedMousePos(FIntPoint& OutUnboundedMousePosition, const bool bLocalPosition = true) override;
+
 	/**
 	 * Not implemented                   
 	 */
@@ -386,6 +392,11 @@ private:
 	TMap<FKey,bool> KeyStateMap;
 	/** The last known mouse position in local space, -1, -1 if unknown */
 	FIntPoint CachedCursorPos;
+	/** 
+	 * Mouse position, unbounded: it might be outside the physical display area. 
+	 * Valid while a mouse button is pressed - useful e.g. for drag operations hiding the mouse cursor
+	 */
+	FIntPoint CachedUnboundedCursorPos;
 	/** The last known geometry info */
 	FGeometry CachedGeometry;
 	/** Mouse position before the latest capture */
diff --git a/Engine/Source/Runtime/Engine/Public/UnrealClient.h b/Engine/Source/Runtime/Engine/Public/UnrealClient.h
--- a/Engine/Source/Runtime/Engine/Public/UnrealClient.h
+++ b/Engine/Source/Runtime/Engine/Public/UnrealClient.h
@@ -457,6 +457,7 @@ public:
 	virtual int32 GetMouseX() const = 0;
 	virtual int32 GetMouseY() const = 0;
 	virtual void GetMousePos(FIntPoint& MousePosition, const bool bLocalPosition = true) = 0;
+	virtual void GetUnboundedMousePos(FIntPoint& OutUnboundedMousePosition, const bool bLocalPosition = true) {}
 	virtual float GetTabletPressure() { return 0.f; }
 	virtual bool IsPenActive() { return false; }
 	virtual void SetMouse(int32 x, int32 y) = 0;
diff --git a/Engine/Source/Runtime/InteractiveToolsFramework/Public/InputState.h b/Engine/Source/Runtime/InteractiveToolsFramework/Public/InputState.h
--- a/Engine/Source/Runtime/InteractiveToolsFramework/Public/InputState.h
+++ b/Engine/Source/Runtime/InteractiveToolsFramework/Public/InputState.h
@@ -245,6 +245,18 @@ struct FMouseInputDeviceState
 	UPROPERTY(transient, BlueprintReadWrite, Category = MouseInputDeviceState)
 	FRay WorldRay;
 
+	/** The 2D position of the mouse, unconstrained by screen real-estate. Computed by summing unbound deltas. */
+	UPROPERTY(transient, BlueprintReadWrite, Category = MouseInputDeviceState)
+	FVector2D UnboundedPosition2D;
+
+	/** Change in 2D mouse position, valid also for movements outside the display bounds */
+	UPROPERTY(transient, BlueprintReadWrite, Category = MouseInputDeviceState)
+	FVector2D UnboundedMouseDelta2D;
+
+	/** Ray into current 3D scene at current unbounded 2D mouse position */
+	UPROPERTY(transient, BlueprintReadWrite, Category = MouseInputDeviceState)
+	FRay UnboundedWorldRay;
+
 	FMouseInputDeviceState()
 	{
 		Left = FDeviceButtonState(EKeys::LeftMouseButton);
@@ -254,6 +266,9 @@ struct FMouseInputDeviceState
 		Position2D = FVector2D::ZeroVector;
 		Delta2D = FVector2D::ZeroVector;
 		WorldRay = FRay();
+		UnboundedPosition2D = FVector2D::ZeroVector;
+		UnboundedMouseDelta2D = FVector2D::ZeroVector;
+		UnboundedWorldRay = FRay();
 	}
 };
 
