[ITF Viewport Drag Interactions] Use unbounded cursor position

Adds unbounded cursor usage to drag interactions, e.g. Orbit, Camera angle, etc

#Synced-CL 49032025

diff --git a/Engine/Source/Editor/Experimental/EditorInteractiveToolsFramework/Private/ViewportInteractions/ViewportDragInteraction.cpp b/Engine/Source/Editor/Experimental/EditorInteractiveToolsFramework/Private/ViewportInteractions/ViewportDragInteraction.cpp
--- a/Engine/Source/Editor/Experimental/EditorInteractiveToolsFramework/Private/ViewportInteractions/ViewportDragInteraction.cpp
+++ b/Engine/Source/Editor/Experimental/EditorInteractiveToolsFramework/Private/ViewportInteractions/ViewportDragInteraction.cpp
@@ -18,6 +18,7 @@ UViewportDragInteraction::UViewportDragInteraction()
 
 	// Other higher priority drag interactions might steal capture (e.g. Gizmo Indirect Manipulation)
 	ClickDragBehavior->SetAllowsCaptureStealing(true);
+	ClickDragBehavior->SetUsesUnboundedCursor(true);
 
 	ViewportClickDragBehavior = ClickDragBehavior;
 
diff --git a/Engine/Source/Editor/Experimental/EditorInteractiveToolsFramework/Private/ViewportInteractions/ViewportInteractionsBehaviorSource.cpp b/Engine/Source/Editor/Experimental/EditorInteractiveToolsFramework/Private/ViewportInteractions/ViewportInteractionsBehaviorSource.cpp
--- a/Engine/Source/Editor/Experimental/EditorInteractiveToolsFramework/Private/ViewportInteractions/ViewportInteractionsBehaviorSource.cpp
+++ b/Engine/Source/Editor/Experimental/EditorInteractiveToolsFramework/Private/ViewportInteractions/ViewportInteractionsBehaviorSource.cpp
@@ -508,11 +508,19 @@ void UViewportInteractionsBehaviorSource::SetMouseCursorOverride(EMouseCursor::T
 		}
 
 		EditorViewportClient->SetRequiredCursorOverride(true, CursorOverride);
+
+		// In case we hide the cursor, we show it back where it was when clearing the override
+		if (InMouseCursor == EMouseCursor::None)
+		{
+			HiddenCursorPosition.X = EditorViewportClient->Viewport->GetMouseX();
+			HiddenCursorPosition.Y = EditorViewportClient->Viewport->GetMouseY();
+		}
 	}
 }
 
 void UViewportInteractionsBehaviorSource::ClearMouseCursorOverride()
 {
+	const bool bResetMousePos = CursorOverride == EMouseCursor::None;
 	CursorOverride = EMouseCursor::Default;
 
 	if (FEditorViewportClient* const EditorViewportClient = GetEditorViewportClient())
@@ -523,6 +531,11 @@ void UViewportInteractionsBehaviorSource::ClearMouseCursorOverride()
 		}
 
 		EditorViewportClient->SetRequiredCursorOverride(false);
+
+		if (bResetMousePos)
+		{
+			EditorViewportClient->Viewport->SetMouse(HiddenCursorPosition.X, HiddenCursorPosition.Y);
+		}
 	}
 }
 
diff --git a/Engine/Source/Editor/Experimental/EditorInteractiveToolsFramework/Public/ViewportInteractions/ViewportInteractionsBehaviorSource.h b/Engine/Source/Editor/Experimental/EditorInteractiveToolsFramework/Public/ViewportInteractions/ViewportInteractionsBehaviorSource.h
--- a/Engine/Source/Editor/Experimental/EditorInteractiveToolsFramework/Public/ViewportInteractions/ViewportInteractionsBehaviorSource.h
+++ b/Engine/Source/Editor/Experimental/EditorInteractiveToolsFramework/Public/ViewportInteractions/ViewportInteractionsBehaviorSource.h
@@ -347,6 +347,11 @@ private:
 
 	EMouseCursor::Type CursorOverride = EMouseCursor::Default;
 
+	/**
+	 * Caches cursor position whenever we hide it, so we can place it back where it was as soon as we show it again
+	 */
+	FIntPoint HiddenCursorPosition;
+
 	UPROPERTY(Transient)
 	TWeakObjectPtr<const UEditorInteractiveToolsContext> EditorInteractiveToolsContextWeak;
 
