[ITF Viewport Interactions] Do not require that alt to always be held down once orbit & dolly actions are started.

#rnx

#Synced-CL 49559369

diff --git a/Engine/Source/Editor/Experimental/EditorInteractiveToolsFramework/Private/ViewportInteractions/ViewportDollyInteraction.cpp b/Engine/Source/Editor/Experimental/EditorInteractiveToolsFramework/Private/ViewportInteractions/ViewportDollyInteraction.cpp
--- a/Engine/Source/Editor/Experimental/EditorInteractiveToolsFramework/Private/ViewportInteractions/ViewportDollyInteraction.cpp
+++ b/Engine/Source/Editor/Experimental/EditorInteractiveToolsFramework/Private/ViewportInteractions/ViewportDollyInteraction.cpp
@@ -9,11 +9,12 @@
 
 UViewportDollyInteraction::UViewportDollyInteraction()
 {
-	InteractionName = UE::Editor::ViewportInteractions::PerspectiveDolly;
+	using namespace UE::Editor::ViewportInteractions;
+	InteractionName = PerspectiveDolly;
 	
 	ViewportClickDragBehavior->SetBindings({
 		EKeys::RightMouseButton,
-		EKeys::LeftAlt
+		FButtonBinding(EKeys::LeftAlt).RequiredToContinue(false)
 	});
 }
 
diff --git a/Engine/Source/Editor/Experimental/EditorInteractiveToolsFramework/Private/ViewportInteractions/ViewportOrbitInteraction.cpp b/Engine/Source/Editor/Experimental/EditorInteractiveToolsFramework/Private/ViewportInteractions/ViewportOrbitInteraction.cpp
--- a/Engine/Source/Editor/Experimental/EditorInteractiveToolsFramework/Private/ViewportInteractions/ViewportOrbitInteraction.cpp
+++ b/Engine/Source/Editor/Experimental/EditorInteractiveToolsFramework/Private/ViewportInteractions/ViewportOrbitInteraction.cpp
@@ -16,7 +16,7 @@ UViewportOrbitInteraction::UViewportOrbitInteraction()
 	using namespace UE::Editor::ViewportInteractions;
 	ViewportClickDragBehavior->SetBindings({
 		FButtonBinding(EKeys::LeftMouseButton).TriggersStart(),
-		FButtonBinding(EKeys::LeftAlt)		
+		FButtonBinding(EKeys::LeftAlt).RequiredToContinue(false)
 	});
 }
 
