[Editor TRS] Transform mode getting stuck when toggling mid-drag. Fixes by making conflicting behavior contingent on gizmo drag.

#Synced-CL 49644553

diff --git a/Engine/Source/Editor/Experimental/EditorInteractiveToolsFramework/Private/ViewportInteractions/ViewportCameraTranslateInteraction.cpp b/Engine/Source/Editor/Experimental/EditorInteractiveToolsFramework/Private/ViewportInteractions/ViewportCameraTranslateInteraction.cpp
--- a/Engine/Source/Editor/Experimental/EditorInteractiveToolsFramework/Private/ViewportInteractions/ViewportCameraTranslateInteraction.cpp
+++ b/Engine/Source/Editor/Experimental/EditorInteractiveToolsFramework/Private/ViewportInteractions/ViewportCameraTranslateInteraction.cpp
@@ -62,7 +62,7 @@ UViewportCameraTranslateInteraction::UViewportCameraTranslateInteraction()
 				{
 					if (UViewportInteractionsBehaviorSource* BehaviorSource = GetViewportInteractionsBehaviorSource())
 					{
-						BehaviorSource->SetIsMouseLooking(true);
+						BehaviorSource->SetIsMouseLooking(!BehaviorSource->IsGizmoDragging());
 					}
 				}
 			}
