[ITF Gizmos] Alt Duplicate duplicates twice with sequencer

#Synced-CL 49648897

diff --git a/Engine/Source/Editor/Experimental/EditorInteractiveToolsFramework/Private/ViewportInteractions/ViewportInteractionsBehaviorSource.cpp b/Engine/Source/Editor/Experimental/EditorInteractiveToolsFramework/Private/ViewportInteractions/ViewportInteractionsBehaviorSource.cpp
--- a/Engine/Source/Editor/Experimental/EditorInteractiveToolsFramework/Private/ViewportInteractions/ViewportInteractionsBehaviorSource.cpp
+++ b/Engine/Source/Editor/Experimental/EditorInteractiveToolsFramework/Private/ViewportInteractions/ViewportInteractionsBehaviorSource.cpp
@@ -230,6 +230,8 @@ void UViewportInteractionsBehaviorSource::Initialize(const UEditorInteractiveToo
 
 void UViewportInteractionsBehaviorSource::RegisterBehaviorSources()
 {
+	DeregisterBehaviorSources();
+
 	if (const UEditorInteractiveToolsContext* InteractiveToolsContext = EditorInteractiveToolsContextWeak.Get())
 	{
 		if (UInputRouter* InputRouter = InteractiveToolsContext->InputRouter)
@@ -335,7 +337,6 @@ void UViewportInteractionsBehaviorSource::AddInteractions(const TArray<const UCl
 	if (bInReregister)
 	{
 		// Refreshing registered behaviors
-		DeregisterBehaviorSources();
 		RegisterBehaviorSources();
 	}
 }
@@ -388,7 +389,6 @@ UViewportInteraction* UViewportInteractionsBehaviorSource::AddInteraction(const
 	if (bInReregister)
 	{
 		// Refreshing registered behaviors
-		DeregisterBehaviorSources();
 		RegisterBehaviorSources();
 	}
 
@@ -465,7 +465,6 @@ void UViewportInteractionsBehaviorSource::RemoveInteraction(const FName InIntera
 		if (bInReregister)
 		{
 			// Refreshing registered behaviors
-			DeregisterBehaviorSources();
 			RegisterBehaviorSources();
 		}
 	}
