[ITF Gizmo] Fix for cursor incorrectly changing icon when manipulating Rotation Gizmo

Viewport Click Interaction was overriding cursor even when gizmo interaction was occurring

#Synced-CL 49725045

diff --git a/Engine/Source/Editor/Experimental/EditorInteractiveToolsFramework/Private/ViewportInteractions/ViewportClickInteraction.cpp b/Engine/Source/Editor/Experimental/EditorInteractiveToolsFramework/Private/ViewportInteractions/ViewportClickInteraction.cpp
--- a/Engine/Source/Editor/Experimental/EditorInteractiveToolsFramework/Private/ViewportInteractions/ViewportClickInteraction.cpp
+++ b/Engine/Source/Editor/Experimental/EditorInteractiveToolsFramework/Private/ViewportInteractions/ViewportClickInteraction.cpp
@@ -24,7 +24,10 @@ FInputRayHit UViewportClickInteraction::IsHitByClick(const FInputDeviceRay& InCl
 {
 	if (UViewportInteractionsBehaviorSource* const ViewportInteractionBehaviorSource = GetViewportInteractionsBehaviorSource())
 	{
-		ViewportInteractionBehaviorSource->SetMouseCursorOverride(EMouseCursor::Crosshairs);
+		if (!ViewportInteractionBehaviorSource->IsGizmoDragging())
+		{
+			ViewportInteractionBehaviorSource->SetMouseCursorOverride(EMouseCursor::Crosshairs);
+		}
 	}
 
 	return FInputRayHit(0.0f);
