[Editor TRS] Added some colinearity checks to prevent invalid planes in rotation mode + ortho

#Synced-CL 49985552

diff --git a/Engine/Source/Editor/Experimental/EditorInteractiveToolsFramework/Private/EditorGizmos/EditorGizmoMath.cpp b/Engine/Source/Editor/Experimental/EditorInteractiveToolsFramework/Private/EditorGizmos/EditorGizmoMath.cpp
--- a/Engine/Source/Editor/Experimental/EditorInteractiveToolsFramework/Private/EditorGizmos/EditorGizmoMath.cpp
+++ b/Engine/Source/Editor/Experimental/EditorInteractiveToolsFramework/Private/EditorGizmos/EditorGizmoMath.cpp
@@ -155,6 +155,14 @@ namespace UE::Editor::GizmoMath
 			const FVector2d BottomRightToBottomLeft2D = (PlanePointBottomRight2D - PlanePointBottomLeft2D) + PerspectiveWarp2D.X * PlanePointBottomRight2D;
 			const FVector2d TopLeftToBottomLeft2D = (PlanePointTopLeft2D - PlanePointBottomLeft2D) + PerspectiveWarp2D.Y * PlanePointTopLeft2D;
 
+			// Accounts for colinear points
+			const double SignedArea = BottomRightToBottomLeft2D.X * TopLeftToBottomLeft2D.Y - BottomRightToBottomLeft2D.Y * TopLeftToBottomLeft2D.X;
+
+			if (FMath::IsNearlyZero(SignedArea) || BottomRightToBottomLeft2D.IsNearlyZero() || TopLeftToBottomLeft2D.IsNearlyZero())
+			{
+				return false; // Degenerate projection, plane is edge-on
+			}
+
 			OutPlaneToScreen = FMatrix3d(
 				BottomRightToBottomLeft2D.X, TopLeftToBottomLeft2D.X, PlanePointBottomLeft2D.X,
 				BottomRightToBottomLeft2D.Y, TopLeftToBottomLeft2D.Y, PlanePointBottomLeft2D.Y,
