[ITF Gizmos] Alt + Duplicate should be disabled on Rotation and Scale

Adds a config property allowing to choose wheter Alt Duplicate should work or not on Rotation operations. Default value is false.

Scale will never duplicate.

#Synced-CL 50039254

diff --git a/Engine/Source/Editor/Experimental/EditorInteractiveToolsFramework/Private/EditorDragTools/DuplicateDragSelectionInteraction.cpp b/Engine/Source/Editor/Experimental/EditorInteractiveToolsFramework/Private/EditorDragTools/DuplicateDragSelectionInteraction.cpp
--- a/Engine/Source/Editor/Experimental/EditorInteractiveToolsFramework/Private/EditorDragTools/DuplicateDragSelectionInteraction.cpp
+++ b/Engine/Source/Editor/Experimental/EditorInteractiveToolsFramework/Private/EditorDragTools/DuplicateDragSelectionInteraction.cpp
@@ -7,7 +7,9 @@
 #include "EditorInteractiveGizmoManager.h"
 #include "LevelEditorViewport.h"
 #include "ScopedTransaction.h"
+#include "TransformGizmoEditorSettings.h"
 #include "UnrealEdGlobals.h"
+#include "EditorGizmos/TransformGizmoAccessor.h"
 #include "ViewportInteractions/ViewportInteractionsBehaviorSource.h"
 
 #define LOCTEXT_NAMESPACE "EditorDuplicateDragSelection"
@@ -49,6 +51,11 @@ void UDuplicateDragSelectionInteraction::OnGizmoMovementBegin(UTransformProxy* I
 		return;
 	}
 
+	if (!ShouldDuplicate())
+	{
+		return;
+	}
+
 	if (bModifierKeyIsPressed && !TransformGizmo->IsIndirectManipulation())
 	{
 		DuplicateTransaction = MakeShared<FScopedTransaction>(LOCTEXT("DuplicateAndTransformTransactionName", "Duplicate and Transform"));
@@ -80,7 +87,7 @@ void UDuplicateDragSelectionInteraction::DuplicateSelection(bool bInSelectNewEle
 
 void UDuplicateDragSelectionInteraction::OnModifierKeyChanged(bool bInIsDown)
 {
-	if (!IsEnabled())
+	if (!ShouldDuplicate())
 	{
 		return;
 	}
@@ -102,4 +109,28 @@ void UDuplicateDragSelectionInteraction::OnModifierKeyChanged(bool bInIsDown)
 	}
 }
 
+bool UDuplicateDragSelectionInteraction::ShouldDuplicate() const
+{
+	if (!IsEnabled())
+	{
+		return false;
+	}
+
+	const UEditorTransformGizmo* const TransformGizmo = GetViewportInteractionsBehaviorSource()->GetEditorTransformGizmo();
+	if (!TransformGizmo)
+	{
+		return false;
+	}
+
+	// Always duplicate on Translate. Duplicate on Rotate only if appropriate flag is enabled. Never duplicate on Scale.
+	EGizmoTransformMode TransformMode = UE::Editor::InteractiveToolsFramework::FTransformGizmoAccessor().GetCurrentMode(*TransformGizmo);
+	if ((TransformMode == EGizmoTransformMode::Rotate && !GetDefault<UTransformGizmoEditorSettings>()->GizmosParameters.bDragDuplicateOnRotation) ||
+		TransformMode == EGizmoTransformMode::Scale)
+	{
+		return false;
+	}
+
+	return true;
+}
+
 #undef LOCTEXT_NAMESPACE
diff --git a/Engine/Source/Editor/Experimental/EditorInteractiveToolsFramework/Public/EditorDragTools/DuplicateDragSelectionInteraction.h b/Engine/Source/Editor/Experimental/EditorInteractiveToolsFramework/Public/EditorDragTools/DuplicateDragSelectionInteraction.h
--- a/Engine/Source/Editor/Experimental/EditorInteractiveToolsFramework/Public/EditorDragTools/DuplicateDragSelectionInteraction.h
+++ b/Engine/Source/Editor/Experimental/EditorInteractiveToolsFramework/Public/EditorDragTools/DuplicateDragSelectionInteraction.h
@@ -46,6 +46,8 @@ protected:
 private:
 	void OnModifierKeyChanged(bool bInIsDown);
 
+	bool ShouldDuplicate() const;
+
 	bool bModifierKeyIsPressed = false;
 
 	/** Keeps the duplicate + transform transaction alive between gizmo movement begin and end */
diff --git a/Engine/Source/Editor/Experimental/EditorInteractiveToolsFramework/Public/EditorGizmos/TransformGizmo.h b/Engine/Source/Editor/Experimental/EditorInteractiveToolsFramework/Public/EditorGizmos/TransformGizmo.h
--- a/Engine/Source/Editor/Experimental/EditorInteractiveToolsFramework/Public/EditorGizmos/TransformGizmo.h
+++ b/Engine/Source/Editor/Experimental/EditorInteractiveToolsFramework/Public/EditorGizmos/TransformGizmo.h
@@ -193,6 +193,10 @@ struct FGizmosParameters
 	UPROPERTY(EditAnywhere, Category = "Experimental")
 	bool bPersistHandleSelection = true;
 
+	/** When enabled, drag duplication will work on both translation and rotation (currently Level Editor only). */
+	UPROPERTY(EditAnywhere, Category=GizmoInteraction)
+	bool bDragDuplicateOnRotation = false;
+
 	UPROPERTY(EditAnywhere, Category = "Transform Gizmo|Development", meta=(ShowOnlyInnerProperties))
 	FGizmoDevParameters DevParameters;
 };
