[ViewportInteractions] Fix an issue where UViewportClickDragBehavior could call OnDragEnd() when a drag had not yet started.
Adjust default drag threshold to 2.
Add a debug CVar "Editor.ViewportInteractions.DragThresholdOverride". Set to >=0 to override the drag threshold for all viewport interactions.

#rnx

#Synced-CL 50003544

diff --git a/Engine/Source/Editor/Experimental/EditorInteractiveToolsFramework/Private/Behaviors/ViewportClickDragBehavior.cpp b/Engine/Source/Editor/Experimental/EditorInteractiveToolsFramework/Private/Behaviors/ViewportClickDragBehavior.cpp
--- a/Engine/Source/Editor/Experimental/EditorInteractiveToolsFramework/Private/Behaviors/ViewportClickDragBehavior.cpp
+++ b/Engine/Source/Editor/Experimental/EditorInteractiveToolsFramework/Private/Behaviors/ViewportClickDragBehavior.cpp
@@ -4,6 +4,15 @@
 
 #include "ViewportInteractions/ViewportInteractionsBehaviorSource.h"
 
+namespace UE::Editor::ViewportInteractions
+{
+	TAutoConsoleVariable<float> DragThresholdOverride(
+		TEXT("Editor.ViewportInteractions.DragThresholdOverride"),
+		-1.0f,
+		TEXT("Set to a >= 0 value to override all drag threshold values.")
+	);
+}
+
 void UViewportClickDragBehavior::Initialize(IViewportClickDragBehaviorTarget* InTarget)
 {
 	check(InTarget != nullptr);
@@ -22,12 +31,9 @@ FInputCapturePriority UViewportClickDragBehavior::GetPriority()
 {
 	FInputCapturePriority BasePriority = DefaultPriority;
 	
-	if (const float* TraveledDistance = MouseTraveledDistance.GetPtrOrNull())
+	if (HasMovedEnoughForDrag())
 	{
-		if (*TraveledDistance > ClickDistanceThreshold)
-		{
-			BasePriority = DragConfirmedPriority;
-		}
+		BasePriority = DragConfirmedPriority;
 	}
 	
 	return BasePriority.MakeHigher(ActiveBindingComplexity);
@@ -75,6 +81,7 @@ FInputCaptureRequest UViewportClickDragBehavior::WantsCapture(const FInputDevice
 	if (!HitResult.bHit)
 	{
 		MouseTraveledDistance = TOptional<float>();
+		ActiveBindingComplexity = 0;
 		return FInputCaptureRequest::Ignore();
 	}
 	
@@ -110,26 +117,21 @@ FInputCaptureUpdate UViewportClickDragBehavior::UpdateCapture(const FInputDevice
 	
 	if (GetBindingComplexity(EInputStage::Continue, InputState) <= 0)
 	{
-		OnDragEndInternal(InputState, CaptureData);
+		if (bDragging)
+		{
+			OnDragEndInternal(InputState, CaptureData);
+			bDragging = false;
+		}
 		MouseTraveledDistance = TOptional<float>();
-		bDragging = false;
 		return FInputCaptureUpdate::End();
 	}
 	
 	MouseTraveledDistance = static_cast<float>(FVector2D::Distance(MousePosition, MouseDownPosition));
 
-	if (!bDragging)
+	if (!bDragging && HasMovedEnoughForDrag())
 	{
-		if (const float* TraveledDistance = MouseTraveledDistance.GetPtrOrNull())
-		{
-			if (*TraveledDistance > ClickDistanceThreshold)
-			{
-				OnDragStartInternal(InputState);
-
-				// Do not start drag until ClickDistanceThreshold has been reached
-				bDragging = true;
-			}
-		}
+		bDragging = true;
+		OnDragStartInternal(InputState);
 	}
 
 	if (bDragging)
@@ -143,6 +145,7 @@ FInputCaptureUpdate UViewportClickDragBehavior::UpdateCapture(const FInputDevice
 void UViewportClickDragBehavior::ForceEndCapture(const FInputCaptureData& CaptureData)
 {
 	MouseTraveledDistance = TOptional<float>();
+	ActiveBindingComplexity = 0;
 
 	if (bDragging)
 	{
@@ -334,3 +337,19 @@ void UViewportClickDragBehavior::OnStateUpdatedInternal(const FInputDeviceState&
 {
 	Target->OnStateUpdated(InInputDeviceState);
 }
+
+bool UViewportClickDragBehavior::HasMovedEnoughForDrag() const
+{
+	if (const float* TraveledDistance = MouseTraveledDistance.GetPtrOrNull())
+	{
+		const float Override = UE::Editor::ViewportInteractions::DragThresholdOverride.GetValueOnAnyThread();
+		const float EffectiveThreshold = Override >= 0.0f ? Override : ClickDistanceThreshold;
+	
+		if (*TraveledDistance > EffectiveThreshold)
+		{
+			return true;
+		}
+	}
+
+	return false;
+}
diff --git a/Engine/Source/Editor/Experimental/EditorInteractiveToolsFramework/Public/Behaviors/ViewportClickDragBehavior.h b/Engine/Source/Editor/Experimental/EditorInteractiveToolsFramework/Public/Behaviors/ViewportClickDragBehavior.h
--- a/Engine/Source/Editor/Experimental/EditorInteractiveToolsFramework/Public/Behaviors/ViewportClickDragBehavior.h
+++ b/Engine/Source/Editor/Experimental/EditorInteractiveToolsFramework/Public/Behaviors/ViewportClickDragBehavior.h
@@ -190,7 +190,9 @@ protected:
 	virtual void OnStateUpdatedInternal(const FInputDeviceState& InInputDeviceState);
 	
 	/** Distance after which a drag is considered such */
-	float ClickDistanceThreshold = 1.0f;
+	float ClickDistanceThreshold = 2.0f;
+	
+	bool HasMovedEnoughForDrag() const;
 	
 	/** Cached cursor position */
 	FVector2D MouseDownPosition;
