[ITF Viewport Interactions] Use the Mouse Down Position as the start of Frusum and Box select. This makes frustum interactions feel more responsive.
Adjust ViewportClickDragBehavior.h to use `UE_API` instead of the longer package name.

#Synced-CL 50003544

#rnx
#Synced-CL 50003544

diff --git a/Engine/Source/Editor/Experimental/EditorInteractiveToolsFramework/Private/EditorDragTools/BoxSelectInteraction.cpp b/Engine/Source/Editor/Experimental/EditorInteractiveToolsFramework/Private/EditorDragTools/BoxSelectInteraction.cpp
--- a/Engine/Source/Editor/Experimental/EditorInteractiveToolsFramework/Private/EditorDragTools/BoxSelectInteraction.cpp
+++ b/Engine/Source/Editor/Experimental/EditorInteractiveToolsFramework/Private/EditorDragTools/BoxSelectInteraction.cpp
@@ -110,10 +110,8 @@ void UBoxSelectInteraction::OnDragStart(const FInputDeviceRay& InPressPos)
 	OnActivateTool().Broadcast();
 	bIsDragging = true;
 
-	FIntPoint MousePos;
-	Viewport->GetMousePos(MousePos);
-
-	Start = FVector(MousePos);
+	const FVector2D MouseDownPosition = ClickDragInputBehavior->GetMouseDownPosition();
+	Start = FVector(MouseDownPosition.X, MouseDownPosition.Y, 0.0);
 	End = Start;
 
 	FLevelEditorViewportClient::ClearHoverFromObjects();
diff --git a/Engine/Source/Editor/Experimental/EditorInteractiveToolsFramework/Private/EditorDragTools/FrustumSelectInteraction.cpp b/Engine/Source/Editor/Experimental/EditorInteractiveToolsFramework/Private/EditorDragTools/FrustumSelectInteraction.cpp
--- a/Engine/Source/Editor/Experimental/EditorInteractiveToolsFramework/Private/EditorDragTools/FrustumSelectInteraction.cpp
+++ b/Engine/Source/Editor/Experimental/EditorInteractiveToolsFramework/Private/EditorDragTools/FrustumSelectInteraction.cpp
@@ -163,25 +163,16 @@ void UFrustumSelectInteraction::OnDragStart(const FInputDeviceRay& InPressPos)
 		BehaviorSource->SetMouseCursorOverride(EMouseCursor::Default);
 	}
 
-	// Signal that this tool is now active
-	OnActivateTool().Broadcast();
-
-	Start = InPressPos.WorldRay.Origin;
 	bIsDragging = true;
 
-	// Snap to constraints.
-	if (bUseSnapping)
-	{
-		const float GridSize = GEditor->GetGridSize();
-		const FVector GridBase(GridSize, GridSize, GridSize);
-		FSnappingUtils::SnapPointToGrid(Start, GridBase);
-	}
-	End = Start;
+	// Signal that this tool is now active
+	OnActivateTool().Broadcast();
 
 	// Remove any active hover objects
 	FLevelEditorViewportClient::ClearHoverFromObjects();
 
-	Start = FVector(InPressPos.ScreenPosition.X, InPressPos.ScreenPosition.Y, 0);
+	const FVector2D MouseDownPosition = ClickDragInputBehavior->GetMouseDownPosition();
+	Start = FVector(MouseDownPosition.X, MouseDownPosition.Y, 0);
 	End = Start;
 }
 
@@ -225,6 +216,17 @@ void UFrustumSelectInteraction::OnDragEnd(const FInputDeviceRay& InReleasePos)
 	);
 	FSceneView* SceneView = EditorViewportClient->CalcSceneView(&ViewFamily);
 
+	// Normalize the start and end
+	if (Start.X > End.X)
+	{
+		Swap(Start.X, End.X);
+	}
+
+	if (Start.Y > End.Y)
+	{
+		Swap(Start.Y, End.Y);
+	}
+
 	// Generate a frustum out of the dragged box
 	FConvexVolume Frustum;
 	CalculateFrustum(SceneView, Frustum, true);
@@ -244,11 +246,21 @@ void UFrustumSelectInteraction::OnDragEnd(const FInputDeviceRay& InReleasePos)
 	// Let the editor mode try to handle the selection.
 	const bool bEditorModeHandledSelection = ModeTools->FrustumSelect(Frustum, EditorViewportClient, bShouldSelect);
 
+	if (bEditorModeHandledSelection)
+	{
+		UE_LOG(LogTemp, Log, TEXT("Mode tools handled frustum select"));
+	}
+
 	// Let the component visualizers try to handle the selection.
 	const bool bComponentVisHandledSelection =
 		!bEditorModeHandledSelection
 		&& GUnrealEd->ComponentVisManager.HandleFrustumSelect(Frustum, EditorViewportClient, Viewport);
 
+	if (bComponentVisHandledSelection)
+	{
+		UE_LOG(LogTemp, Log, TEXT("Component visualizer handled frustum select"));
+	}
+
 	if (!bEditorModeHandledSelection && !bComponentVisHandledSelection)
 	{
 		UTypedElementSelectionSet* SelectionSet = ModeTools->GetEditorSelectionSet();
@@ -264,16 +276,6 @@ void UFrustumSelectInteraction::OnDragEnd(const FInputDeviceRay& InReleasePos)
 		const int32 ViewportSizeX = Viewport->GetSizeXY().X;
 		const int32 ViewportSizeY = Viewport->GetSizeXY().Y;
 
-		if (Start.X > End.X)
-		{
-			Swap(Start.X, End.X);
-		}
-
-		if (Start.Y > End.Y)
-		{
-			Swap(Start.Y, End.Y);
-		}
-
 		TArray<FTypedElementHandle> ElementsToSelect;
 		const bool bTransparentBoxSelection = GetDefault<ULevelEditorViewportSettings>()->bTransparentBoxSelection;
 		if (bTransparentBoxSelection)
@@ -388,6 +390,11 @@ void UFrustumSelectInteraction::OnDragEnd(const FInputDeviceRay& InReleasePos)
 
 		FTypedElementSelectionOptions ElementSelectionOption;
 
+		if (ElementsToSelect.Num() > 0)
+		{
+			UE_LOG(LogTemp, Log, TEXT("Hit %i things"), ElementsToSelect.Num());
+		}
+
 		if (!FInputDeviceState::IsShiftKeyDown(InputState))
 		{
 			// If the user is selecting, but isn't hold down SHIFT, remove the previous selections.
diff --git a/Engine/Source/Editor/Experimental/EditorInteractiveToolsFramework/Public/Behaviors/ViewportClickDragBehavior.h b/Engine/Source/Editor/Experimental/EditorInteractiveToolsFramework/Public/Behaviors/ViewportClickDragBehavior.h
--- a/Engine/Source/Editor/Experimental/EditorInteractiveToolsFramework/Public/Behaviors/ViewportClickDragBehavior.h
+++ b/Engine/Source/Editor/Experimental/EditorInteractiveToolsFramework/Public/Behaviors/ViewportClickDragBehavior.h
@@ -7,6 +7,8 @@
 #include "Framework/Commands/UICommandInfo.h"
 #include "ViewportClickDragBehavior.generated.h"
 
+#define UE_API EDITORINTERACTIVETOOLSFRAMEWORK_API
+
 class IViewportClickDragBehaviorTarget
 {
 public:
@@ -136,40 +138,38 @@ public:
 		Continue
 	};
 
-	EDITORINTERACTIVETOOLSFRAMEWORK_API void Initialize(IViewportClickDragBehaviorTarget* InTarget);
+	UE_API void Initialize(IViewportClickDragBehaviorTarget* InTarget);
 
 	/**
 	 * Set the priority this Drag Behavior will return from GetPriority once the current interaction is confirmed to be a drag.
 	 * Should be higher than its DefaultPriority, and higher than the priority of any other behavior it wants to steal capture from
 	 */
-	EDITORINTERACTIVETOOLSFRAMEWORK_API void SetDragConfirmedPriority(FInputCapturePriority InDragConfirmedPriority);
-
-	EDITORINTERACTIVETOOLSFRAMEWORK_API virtual FInputCapturePriority GetPriority() override;
+	UE_API void SetDragConfirmedPriority(FInputCapturePriority InDragConfirmedPriority);
 
-	EDITORINTERACTIVETOOLSFRAMEWORK_API virtual FInputCaptureRequest WantsCapture(const FInputDeviceState& InputState) override;
-	EDITORINTERACTIVETOOLSFRAMEWORK_API virtual FInputCaptureUpdate BeginCapture(
-		const FInputDeviceState& InputState, EInputCaptureSide eSide
-	) override;
-	EDITORINTERACTIVETOOLSFRAMEWORK_API virtual FInputCaptureUpdate UpdateCapture(
-		const FInputDeviceState& InputState, const FInputCaptureData& CaptureData
-	) override;
-	EDITORINTERACTIVETOOLSFRAMEWORK_API virtual void ForceEndCapture(const FInputCaptureData& CaptureData) override;
+	UE_API virtual FInputCapturePriority GetPriority() override;
 
-	EDITORINTERACTIVETOOLSFRAMEWORK_API void SetBindings(const TArray<UE::Editor::ViewportInteractions::FButtonBinding>& InBindings);
+	UE_API virtual FInputCaptureRequest WantsCapture(const FInputDeviceState& InputState) override;
+	UE_API virtual FInputCaptureUpdate BeginCapture(const FInputDeviceState& InputState, EInputCaptureSide eSide) override;
+	UE_API virtual FInputCaptureUpdate UpdateCapture(const FInputDeviceState& InputState, const FInputCaptureData& CaptureData) override;
+	UE_API virtual void ForceEndCapture(const FInputCaptureData& CaptureData) override;
+	
+	UE_API void SetBindings(const TArray<UE::Editor::ViewportInteractions::FButtonBinding>& InBindings);
 
 	/** @return current 3D world ray and optional 2D position for Target Device */
-	EDITORINTERACTIVETOOLSFRAMEWORK_API virtual FInputDeviceRay GetDeviceRay(const FInputDeviceState& InInputDeviceState);
+	UE_API virtual FInputDeviceRay GetDeviceRay(const FInputDeviceState& InInputDeviceState);
 
-	EDITORINTERACTIVETOOLSFRAMEWORK_API virtual const FVector2D& GetMousePosition(const FInputDeviceState& InInputDeviceState) const;
+	/** Extract the current mouse position from the Input Device State */
+	UE_API virtual const FVector2D& GetMousePosition(const FInputDeviceState& InInputDeviceState) const;
 
-	EDITORINTERACTIVETOOLSFRAMEWORK_API virtual bool UsesUnboundedCursor() const;
+	UE_API virtual bool UsesUnboundedCursor() const;
 
 	/**
 	 * Call to set the default unbounded cursor attribute.
 	 */
-	EDITORINTERACTIVETOOLSFRAMEWORK_API void SetUsesUnboundedCursor(TAttribute<bool> InUsesUnboundedCursor);
+	UE_API void SetUsesUnboundedCursor(TAttribute<bool> InUsesUnboundedCursor);
 
 	bool IsDragging() const { return bDragging; }
+	FVector2D GetMouseDownPosition() const { return MouseDownPosition; }
 
 protected:
 
@@ -179,11 +179,11 @@ protected:
 	int32 GetBindingComplexity(EInputStage Stage, const FInputDeviceState& InputDeviceState);
 	bool DidAnyBindingStateChange(const FInputDeviceState& InputDeviceState) const;
 
-	EDITORINTERACTIVETOOLSFRAMEWORK_API virtual void OnDragStartInternal(const FInputDeviceState& InputState);
-	EDITORINTERACTIVETOOLSFRAMEWORK_API virtual void OnClickDragInternal(
+	UE_API virtual void OnDragStartInternal(const FInputDeviceState& InputState);
+	UE_API virtual void OnClickDragInternal(
 		const FInputDeviceState& InputState, const FInputCaptureData& Data
 	);
-	EDITORINTERACTIVETOOLSFRAMEWORK_API virtual void OnDragEndInternal(
+	UE_API virtual void OnDragEndInternal(
 		const FInputDeviceState& InputState, const FInputCaptureData& Data
 	);
 
@@ -305,3 +305,5 @@ public:
 		OnTerminateFunc();
 	}
 };
+
+#undef UE_API
