[Viewport Interactions] Turn off the verbose logging CVar by default.

#rnx
#Synced-CL 50547624

diff --git a/Engine/Source/Editor/Experimental/EditorInteractiveToolsFramework/Private/ViewportInteractions/ViewportInteractionsBehaviorSource.cpp b/Engine/Source/Editor/Experimental/EditorInteractiveToolsFramework/Private/ViewportInteractions/ViewportInteractionsBehaviorSource.cpp
--- a/Engine/Source/Editor/Experimental/EditorInteractiveToolsFramework/Private/ViewportInteractions/ViewportInteractionsBehaviorSource.cpp
+++ b/Engine/Source/Editor/Experimental/EditorInteractiveToolsFramework/Private/ViewportInteractions/ViewportInteractionsBehaviorSource.cpp
@@ -32,7 +32,7 @@ static FAutoConsoleVariableRef CVarEnableITFViewportInteractions(
 );
 
 // CVar initializer
-static int32 VerboseITFViewportInteractions = 1;
+static int32 VerboseITFViewportInteractions = 0;
 static FAutoConsoleVariableRef CVarEnableVerboseITFViewportInteractions(
 	TEXT("ViewportInteractions.Verbose"),
 	VerboseITFViewportInteractions,
