[ITF Viewport Interactions] Add the "ViewportInteractions.MouseRestoration" CVar that adds the option of, when exiting a EMouseCursor::None override, only reset the cursor position when the cursor has moved outside of the viewport bounds.

#rnx

#Synced-CL 50708145

#Synced-CL 50708145

diff --git a/Engine/Source/Editor/Experimental/EditorInteractiveToolsFramework/Private/ViewportInteractions/ViewportInteractionsBehaviorSource.cpp b/Engine/Source/Editor/Experimental/EditorInteractiveToolsFramework/Private/ViewportInteractions/ViewportInteractionsBehaviorSource.cpp
--- a/Engine/Source/Editor/Experimental/EditorInteractiveToolsFramework/Private/ViewportInteractions/ViewportInteractionsBehaviorSource.cpp
+++ b/Engine/Source/Editor/Experimental/EditorInteractiveToolsFramework/Private/ViewportInteractions/ViewportInteractionsBehaviorSource.cpp
@@ -43,6 +43,15 @@ static FAutoConsoleVariableRef CVarEnableVerboseITFViewportInteractions(
 		}
 	)
 );
+	
+TAutoConsoleVariable<int32> CVarMouseRestoration(
+	TEXT("ViewportInteractions.MouseRestoration"),
+	1,
+	TEXT("Set how the viewport handles the mouse when the cursor stops being invisible\n")
+	TEXT("	0 - Do nothing\n")
+	TEXT("	1 - Reset to initial position (default)\n")
+	TEXT("	2 - Reset to initial position if outside the viewport\n")	
+);
 
 bool UseEditorViewportInteractions()
 {
@@ -514,8 +523,7 @@ void UViewportInteractionsBehaviorSource::SetMouseCursorOverride(EMouseCursor::T
 		// In case we hide the cursor, we show it back where it was when clearing the override
 		if (InMouseCursor == EMouseCursor::None)
 		{
-			HiddenCursorPosition.X = EditorViewportClient->Viewport->GetMouseX();
-			HiddenCursorPosition.Y = EditorViewportClient->Viewport->GetMouseY();
+			EditorViewportClient->Viewport->GetMousePos(HiddenCursorPosition);
 		}
 	}
 }
@@ -534,7 +542,29 @@ void UViewportInteractionsBehaviorSource::ClearMouseCursorOverride()
 
 		if (bResetMousePos)
 		{
-			EditorViewportClient->Viewport->SetMouse(HiddenCursorPosition.X, HiddenCursorPosition.Y);
+			const int32 Mode = UE::Editor::ViewportInteractions::CVarMouseRestoration.GetValueOnAnyThread();
+			switch (Mode)
+			{
+			case 1: // Reset the mouse back to its original position (parity with legacy)
+				EditorViewportClient->Viewport->SetMouse(HiddenCursorPosition.X, HiddenCursorPosition.Y);
+				break;
+			case 2: // Reset the mouse to original position only when the mouse is no longer within the viewport
+			{
+				FIntPoint CurrentMousePosition;
+				EditorViewportClient->Viewport->GetMousePos(CurrentMousePosition);
+		
+				FIntPoint Size = EditorViewportClient->Viewport->GetSizeXY();
+		
+				if (CurrentMousePosition.X < 0 || CurrentMousePosition.X > Size.X
+					|| CurrentMousePosition.Y < 0 || CurrentMousePosition.Y > Size.Y)
+				{
+					EditorViewportClient->Viewport->SetMouse(HiddenCursorPosition.X, HiddenCursorPosition.Y);				
+				}
+				break;
+			}
+			default:
+				break;
+			}
 		}
 	}
 }
