[ITF Gizmo] Alt Tab issue with Alt duplicate

Quick fix to make sure this behavior doesn't store a wrong bModifierKeyIsPressed value.
This behavior will disappear with an upcoming CL, but in the meantime this change might help prevent some issues now and then when Alt + Tabbing

#Synced-CL 50648887

diff --git a/Engine/Source/Editor/Experimental/EditorInteractiveToolsFramework/Private/EditorDragTools/DuplicateDragSelectionInteraction.cpp b/Engine/Source/Editor/Experimental/EditorInteractiveToolsFramework/Private/EditorDragTools/DuplicateDragSelectionInteraction.cpp
--- a/Engine/Source/Editor/Experimental/EditorInteractiveToolsFramework/Private/EditorDragTools/DuplicateDragSelectionInteraction.cpp
+++ b/Engine/Source/Editor/Experimental/EditorInteractiveToolsFramework/Private/EditorDragTools/DuplicateDragSelectionInteraction.cpp
@@ -87,18 +87,19 @@ void UDuplicateDragSelectionInteraction::DuplicateSelection(bool bInSelectNewEle
 
 void UDuplicateDragSelectionInteraction::OnModifierKeyChanged(bool bInIsDown)
 {
-	if (!ShouldDuplicate())
+	// Nothing to do
+	if (bModifierKeyIsPressed == bInIsDown)
 	{
 		return;
 	}
 
-	// Nothing to do
-	if (bModifierKeyIsPressed == bInIsDown)
+	bModifierKeyIsPressed = bInIsDown;
+
+	if (!ShouldDuplicate())
 	{
 		return;
 	}
 
-	bModifierKeyIsPressed = bInIsDown;
 
 	// If already dragging, let's duplicate without selecting new objects
 	// This will "drop" a copy of selection matching current drag transform
