[ITF Viewport Interactions] Key Input Behaviors like Camera Movement potentially steal modifiers (e.g. Ctrl, Alt, Shift, Cmd)

Prevent capturing from within `UKeyInputBehavior::ModifierCheckFunc` when modifiers are pressed.

#Synced-CL 50767940

#Synced-CL 50767940

diff --git a/Engine/Source/Editor/Experimental/EditorInteractiveToolsFramework/Private/ViewportInteractions/ViewportCameraRotateInteraction.cpp b/Engine/Source/Editor/Experimental/EditorInteractiveToolsFramework/Private/ViewportInteractions/ViewportCameraRotateInteraction.cpp
--- a/Engine/Source/Editor/Experimental/EditorInteractiveToolsFramework/Private/ViewportInteractions/ViewportCameraRotateInteraction.cpp
+++ b/Engine/Source/Editor/Experimental/EditorInteractiveToolsFramework/Private/ViewportInteractions/ViewportCameraRotateInteraction.cpp
@@ -19,6 +19,12 @@ UViewportCameraRotateInteraction::UViewportCameraRotateInteraction()
 
 	KeyInputBehavior->ModifierCheckFunc = [this](const FInputDeviceState& InputDeviceState)
 	{
+		if (FInputDeviceState::IsAltKeyDown(InputDeviceState) || FInputDeviceState::IsCtrlKeyDown(InputDeviceState)
+			|| FInputDeviceState::IsShiftKeyDown(InputDeviceState) || FInputDeviceState::IsCmdKeyDown(InputDeviceState))
+		{
+			return false;
+		}
+
 		bool bResult = true;
 
 		if (InputDeviceState.InputDevice == EInputDevices::Keyboard)
diff --git a/Engine/Source/Editor/Experimental/EditorInteractiveToolsFramework/Private/ViewportInteractions/ViewportCameraTranslateInteraction.cpp b/Engine/Source/Editor/Experimental/EditorInteractiveToolsFramework/Private/ViewportInteractions/ViewportCameraTranslateInteraction.cpp
--- a/Engine/Source/Editor/Experimental/EditorInteractiveToolsFramework/Private/ViewportInteractions/ViewportCameraTranslateInteraction.cpp
+++ b/Engine/Source/Editor/Experimental/EditorInteractiveToolsFramework/Private/ViewportInteractions/ViewportCameraTranslateInteraction.cpp
@@ -35,7 +35,14 @@ UViewportCameraTranslateInteraction::UViewportCameraTranslateInteraction()
 	KeyInputBehavior->Initialize(this, GetKeys());
 	KeyInputBehavior->SetDefaultPriority(UE::Editor::ViewportInteractions::VIEWPORT_INTERACTIONS_DEFAULT_PRIORITY);
 	KeyInputBehavior->bRequireAllKeys = false;
-	KeyInputBehavior->ModifierCheckFunc = [this](const FInputDeviceState& InputDeviceState) {
+	KeyInputBehavior->ModifierCheckFunc = [this](const FInputDeviceState& InputDeviceState) 
+	{
+		if (FInputDeviceState::IsAltKeyDown(InputDeviceState) || FInputDeviceState::IsCtrlKeyDown(InputDeviceState)
+			|| FInputDeviceState::IsShiftKeyDown(InputDeviceState) || FInputDeviceState::IsCmdKeyDown(InputDeviceState))
+		{
+			return false;
+		}
+
 		bool bResult = true;
 
 		if (InputDeviceState.InputDevice == EInputDevices::Keyboard)
diff --git a/Engine/Source/Editor/Experimental/EditorInteractiveToolsFramework/Private/ViewportInteractions/ViewportFOVInteraction.cpp b/Engine/Source/Editor/Experimental/EditorInteractiveToolsFramework/Private/ViewportInteractions/ViewportFOVInteraction.cpp
--- a/Engine/Source/Editor/Experimental/EditorInteractiveToolsFramework/Private/ViewportInteractions/ViewportFOVInteraction.cpp
+++ b/Engine/Source/Editor/Experimental/EditorInteractiveToolsFramework/Private/ViewportInteractions/ViewportFOVInteraction.cpp
@@ -19,7 +19,14 @@ UViewportFOVInteraction::UViewportFOVInteraction()
 	KeyInputBehavior->Initialize(this, GetKeys());
 	KeyInputBehavior->SetDefaultPriority(UE::Editor::ViewportInteractions::VIEWPORT_INTERACTIONS_DEFAULT_PRIORITY);
 	KeyInputBehavior->bRequireAllKeys = false;
-	KeyInputBehavior->ModifierCheckFunc = [this](const FInputDeviceState& InputDeviceState) {
+	KeyInputBehavior->ModifierCheckFunc = [this](const FInputDeviceState& InputDeviceState) 
+	{
+		if (FInputDeviceState::IsAltKeyDown(InputDeviceState) || FInputDeviceState::IsCtrlKeyDown(InputDeviceState)
+			|| FInputDeviceState::IsShiftKeyDown(InputDeviceState) || FInputDeviceState::IsCmdKeyDown(InputDeviceState))
+		{
+			return false;
+		}
+
 		bool bResult = true;
 
 		if (InputDeviceState.InputDevice == EInputDevices::Keyboard)
