[TRS Gizmos] Restore conditional logic in UEditorInteractiveGizmoManager::UpdateActiveEditorGizmos() that caused constant invalidation.

#rnx
#Synced-CL 50708145

diff --git a/Engine/Source/Editor/Experimental/EditorInteractiveToolsFramework/Private/EditorInteractiveGizmoManager.cpp b/Engine/Source/Editor/Experimental/EditorInteractiveToolsFramework/Private/EditorInteractiveGizmoManager.cpp
--- a/Engine/Source/Editor/Experimental/EditorInteractiveToolsFramework/Private/EditorInteractiveGizmoManager.cpp
+++ b/Engine/Source/Editor/Experimental/EditorInteractiveToolsFramework/Private/EditorInteractiveGizmoManager.cpp
@@ -283,9 +283,13 @@ void UEditorInteractiveGizmoManager::UpdateActiveEditorGizmos()
 		{
 			DestroyGizmo(Gizmo);
 		}
-		DestroyAllEditorGizmos();
 		
-		bShowEditorGizmos = false;
+		if (bShowEditorGizmos)
+		{
+			DestroyAllEditorGizmos();
+			bShowEditorGizmos = false;
+		}
+		
 		return;
 	}
 	
