[Viewport ITF]
- Update the various ITF Viewport camera controls to respect the `IsCameraLocked()` flag on the Editor Viewport Client.
- Extract the `UViewportClickDragBehavior` bindings configuration mechanism to its own files.
- Add an optional Get Bindings delegate to `UViewportClickDragBehavior`.

#rnx

#Synced-CL 50866064

diff --git a/Engine/Source/Editor/Experimental/EditorInteractiveToolsFramework/Private/Behaviors/ViewportClickDragBehavior.cpp b/Engine/Source/Editor/Experimental/EditorInteractiveToolsFramework/Private/Behaviors/ViewportClickDragBehavior.cpp
--- a/Engine/Source/Editor/Experimental/EditorInteractiveToolsFramework/Private/Behaviors/ViewportClickDragBehavior.cpp
+++ b/Engine/Source/Editor/Experimental/EditorInteractiveToolsFramework/Private/Behaviors/ViewportClickDragBehavior.cpp
@@ -45,7 +45,7 @@ FInputCaptureRequest UViewportClickDragBehavior::WantsCapture(const FInputDevice
 
 	if (!MouseTraveledDistance.IsSet())
 	{
-		ActiveBindingComplexity = GetBindingComplexity(EInputStage::Start, InInputState);
+		ActiveBindingComplexity = GetBindings().GetBindingComplexity(UE::Editor::ViewportInteractions::EInputStage::Start, InInputState);
 		if (ActiveBindingComplexity > 0)
 		{
 			// Make sure behavior knows we are not dragging yet
@@ -66,7 +66,7 @@ FInputCaptureRequest UViewportClickDragBehavior::WantsCapture(const FInputDevice
 	{
 		if (InInputState.IsFromDevice(EInputDevices::Mouse))
 		{
-			ActiveBindingComplexity = GetBindingComplexity(EInputStage::Continue, InInputState);
+			ActiveBindingComplexity = GetBindings().GetBindingComplexity(UE::Editor::ViewportInteractions::EInputStage::Continue, InInputState);
 			if (ActiveBindingComplexity <= 0)
 			{
 				MouseTraveledDistance = TOptional<float>();
@@ -110,12 +110,12 @@ FInputCaptureUpdate UViewportClickDragBehavior::UpdateCapture(const FInputDevice
 		return FInputCaptureUpdate::Continue();
 	}
 	
-	if (DidAnyBindingStateChange(InputState))
+	if (GetBindings().DidAnyBindingStateChange(InputState))
 	{
 		OnStateUpdatedInternal(InputState);
 	}
 	
-	if (GetBindingComplexity(EInputStage::Continue, InputState) <= 0)
+	if (GetBindings().GetBindingComplexity(UE::Editor::ViewportInteractions::EInputStage::Continue, InputState) <= 0)
 	{
 		if (bDragging)
 		{
@@ -157,7 +157,11 @@ void UViewportClickDragBehavior::ForceEndCapture(const FInputCaptureData& Captur
 void UViewportClickDragBehavior::SetBindings(const TArray<UE::Editor::ViewportInteractions::FButtonBinding>& InBindings)
 {
 	Bindings = InBindings;
-	bRequiresTrigger = Algo::FindByPredicate(Bindings, [](const UE::Editor::ViewportInteractions::FButtonBinding& InBinding) { return InBinding.bTriggersStart; }) != nullptr;
+}
+
+void UViewportClickDragBehavior::SetBindings(const FGetBindingsDelegate& InDelegate)
+{
+	GetBindingsDelegate = InDelegate;	
 }
 
 FInputDeviceRay UViewportClickDragBehavior::GetDeviceRay(const FInputDeviceState& InInputDeviceState)
@@ -200,136 +204,6 @@ void UViewportClickDragBehavior::ResetDrag(const FInputDeviceState& InputState)
 	MouseTraveledDistance = 0.0f;
 }
 
-namespace UE::Editor::ViewportInteractions::Private
-{
-	FDeviceButtonState DownState(const FKey& InKey, bool bInDown)
-	{
-		FDeviceButtonState State(InKey);
-		State.bDown = bInDown;
-		return State;
-	}
-}
-
-void UViewportClickDragBehavior::ProcessBinding(EInputStage Stage, const FInputDeviceState& InputDeviceState, const UE::Editor::ViewportInteractions::FButtonBinding& Binding, int32& OutComplexity, bool& bOutTrigger)
-{
-	if (TSharedPtr<FUICommandInfo> Command = Binding.Command.Pin())
-	{
-		for (uint32 i = 0; i < static_cast<uint8>(EMultipleKeyBindingIndex::NumChords); ++i)
-		{
-			EMultipleKeyBindingIndex ChordIndex = static_cast<EMultipleKeyBindingIndex>(i);
-			const FInputChord& Chord = *Command->GetActiveChord(ChordIndex);
-			if (Chord.IsValidChord())
-			{
-				// Use temporary state, as _either_ of the bound chords could be used
-				int32 ChordComplexity = 0;
-				bool bHasTrigger = false;
-			
-				ProcessBindingChord(Stage, InputDeviceState, Binding, Chord, ChordComplexity, bHasTrigger);
-				if (ChordComplexity == INDEX_NONE)
-				{
-					continue;
-				}
-				
-				if (ChordComplexity > 0)
-				{
-					OutComplexity += ChordComplexity;
-					bOutTrigger = bOutTrigger || bHasTrigger; 
-					return;
-				}
-			}
-		}
-		
-		// No bound chord was valid, and this binding is required, so it fails
-		if (Stage == EInputStage::Start ? Binding.bRequiredToStart : Binding.bRequiredToContinue)
-		{
-			OutComplexity = INDEX_NONE;
-		}
-	}
-	else
-	{
-		ProcessBindingChord(Stage, InputDeviceState, Binding, Binding.Chord, OutComplexity, bOutTrigger);
-	}
-}
-
-void UViewportClickDragBehavior::ProcessBindingChord(EInputStage Stage, const FInputDeviceState& InputDeviceState, const UE::Editor::ViewportInteractions::FButtonBinding& Binding, const FInputChord& Chord, int32& OutComplexity, bool& bOutTrigger)
-{
-	const FDeviceButtonState State = InputDeviceState.GetButtonState(Chord.Key);
-	
-	const bool bRequired = Stage == EInputStage::Start ? Binding.bRequiredToStart : Binding.bRequiredToContinue;
-	
-#define CHECK_REQUIRED(value) if (value) { ++OutComplexity; } else if (bRequired) { OutComplexity = INDEX_NONE; return; }
-
-	CHECK_REQUIRED(State.bDown);
-	
-	if (Chord.NeedsShift())
-	{
-		CHECK_REQUIRED(InputDeviceState.bShiftKeyDown);
-	}
-	
-	if (Chord.NeedsAlt())
-	{
-		CHECK_REQUIRED(InputDeviceState.bAltKeyDown);
-	}
-	
-	if (Chord.NeedsControl())
-	{
-		CHECK_REQUIRED(InputDeviceState.bCtrlKeyDown);
-	}
-	
-	if (Chord.NeedsCommand())
-	{
-		CHECK_REQUIRED(InputDeviceState.bCmdKeyDown);
-	}
-	
-#undef CHECK_REQUIRED
-	
-	bOutTrigger = bOutTrigger || (Binding.bTriggersStart && (Stage == EInputStage::Start ? State.bPressed : State.bDown));
-}
-
-int32 UViewportClickDragBehavior::GetBindingComplexity(EInputStage Stage, const FInputDeviceState& InputDeviceState)
-{
-	int32 Complexity = 0;
-	bool bHasTrigger = false;
-
-	for (const UE::Editor::ViewportInteractions::FButtonBinding& Binding : Bindings)
-	{
-		ProcessBinding(Stage, InputDeviceState, Binding, Complexity, bHasTrigger);
-		if (Complexity == INDEX_NONE)
-		{
-			return INDEX_NONE;
-		}
-	}
-	
-	return (!bRequiresTrigger || bHasTrigger) ? Complexity : INDEX_NONE;
-}
-
-bool UViewportClickDragBehavior::DidAnyBindingStateChange(const FInputDeviceState& InputDeviceState) const
-{
-	for (const UE::Editor::ViewportInteractions::FButtonBinding& Binding : Bindings)
-	{
-		if (TSharedPtr<FUICommandInfo> Command = Binding.Command.Pin())
-		{
-			for (uint32 i = 0; i < static_cast<uint8>(EMultipleKeyBindingIndex::NumChords); ++i)
-			{
-				EMultipleKeyBindingIndex ChordIndex = static_cast<EMultipleKeyBindingIndex>(i);
-				const FInputChord& Chord = *Command->GetActiveChord(ChordIndex);
-				if (Chord.IsValidChord())
-				{
-					const FDeviceButtonState State = InputDeviceState.GetButtonState(Chord.Key);
-					if (State.bPressed || State.bReleased) return true;
-				}
-			}
-		}
-		else
-		{
-			const FDeviceButtonState State = InputDeviceState.GetButtonState(Binding.Chord.Key);
-			if (State.bPressed || State.bReleased) return true;
-		}
-	}
-	
-	return false;
-}
-
 void UViewportClickDragBehavior::OnDragStartInternal(const FInputDeviceState& InputState)
 {
 	Target->OnDragStart(GetDeviceRay(InputState));
@@ -365,3 +239,12 @@ bool UViewportClickDragBehavior::HasMovedEnoughForDrag() const
 
 	return false;
 }
+
+const UE::Editor::ViewportInteractions::FButtonBindings& UViewportClickDragBehavior::GetBindings() const
+{
+	if (GetBindingsDelegate.IsBound())
+	{
+		return GetBindingsDelegate.Execute();
+	}
+	return Bindings;
+}
diff --git a/Engine/Source/Editor/Experimental/EditorInteractiveToolsFramework/Private/Behaviors/ViewportInteractionBindings.cpp b/Engine/Source/Editor/Experimental/EditorInteractiveToolsFramework/Private/Behaviors/ViewportInteractionBindings.cpp
new file mode 100644
--- /dev/null
+++ b/Engine/Source/Editor/Experimental/EditorInteractiveToolsFramework/Private/Behaviors/ViewportInteractionBindings.cpp
@@ -0,0 +1,142 @@
+﻿// Copyright Epic Games, Inc. All Rights Reserved.
+
+#include "ViewportInteractions/ViewportInteractionBindings.h"
+
+#include "Algo/Find.h"
+
+namespace UE::Editor::ViewportInteractions
+{
+
+FButtonBindings::FButtonBindings(const TArray<FButtonBinding>& InBindings)
+	: Bindings(InBindings)
+{
+	bRequiresTrigger = Algo::FindByPredicate(Bindings, [](const FButtonBinding& InBinding) { return InBinding.bTriggersStart; }) != nullptr;
+}
+
+int32 FButtonBindings::GetBindingComplexity(EInputStage Stage, const FInputDeviceState& InputDeviceState) const
+{
+	int32 Complexity = 0;
+	bool bHasTrigger = false;
+
+	for (const FButtonBinding& Binding : Bindings)
+	{
+		ProcessBinding(Stage, InputDeviceState, Binding, Complexity, bHasTrigger);
+		if (Complexity == INDEX_NONE)
+		{
+			return INDEX_NONE;
+		}
+	}
+	
+	return (!bRequiresTrigger || bHasTrigger) ? Complexity : INDEX_NONE;
+}
+
+bool FButtonBindings::DidAnyBindingStateChange(const FInputDeviceState& InputDeviceState) const
+{
+	for (const FButtonBinding& Binding : Bindings)
+	{
+		if (TSharedPtr<FUICommandInfo> Command = Binding.Command.Pin())
+		{
+			for (uint32 i = 0; i < static_cast<uint8>(EMultipleKeyBindingIndex::NumChords); ++i)
+			{
+				EMultipleKeyBindingIndex ChordIndex = static_cast<EMultipleKeyBindingIndex>(i);
+				const FInputChord& Chord = *Command->GetActiveChord(ChordIndex);
+				if (Chord.IsValidChord())
+				{
+					const FDeviceButtonState State = InputDeviceState.GetButtonState(Chord.Key);
+					if (State.bPressed || State.bReleased)
+					{
+						return true;
+					}
+				}
+			}
+		}
+		else
+		{
+			const FDeviceButtonState State = InputDeviceState.GetButtonState(Binding.Chord.Key);
+			if (State.bPressed || State.bReleased)
+			{
+				return true;
+			}
+		}
+	}
+	
+	return false;
+}
+
+void FButtonBindings::ProcessBinding(EInputStage Stage, const FInputDeviceState& InputDeviceState, const FButtonBinding& Binding, int32& OutComplexity, bool& bOutTrigger)
+{
+	if (TSharedPtr<FUICommandInfo> Command = Binding.Command.Pin())
+	{
+		for (uint32 i = 0; i < static_cast<uint8>(EMultipleKeyBindingIndex::NumChords); ++i)
+		{
+			EMultipleKeyBindingIndex ChordIndex = static_cast<EMultipleKeyBindingIndex>(i);
+			const FInputChord& Chord = *Command->GetActiveChord(ChordIndex);
+			if (Chord.IsValidChord())
+			{
+				// Use temporary state, as _either_ of the bound chords could be used
+				int32 ChordComplexity = 0;
+				bool bHasTrigger = false;
+			
+				ProcessBindingChord(Stage, InputDeviceState, Binding, Chord, ChordComplexity, bHasTrigger);
+				if (ChordComplexity == INDEX_NONE)
+				{
+					continue;
+				}
+				
+				if (ChordComplexity > 0)
+				{
+					OutComplexity += ChordComplexity;
+					bOutTrigger = bOutTrigger || bHasTrigger; 
+					return;
+				}
+			}
+		}
+		
+		// No bound chord was valid, and this binding is required, so it fails
+		if (Stage == EInputStage::Start ? Binding.bRequiredToStart : Binding.bRequiredToContinue)
+		{
+			OutComplexity = INDEX_NONE;
+		}
+	}
+	else
+	{
+		ProcessBindingChord(Stage, InputDeviceState, Binding, Binding.Chord, OutComplexity, bOutTrigger);
+	}
+}
+
+void FButtonBindings::ProcessBindingChord(EInputStage Stage, const FInputDeviceState& InputDeviceState, const FButtonBinding& Binding, const FInputChord& Chord, int32& OutComplexity, bool& bOutTrigger)
+{
+	const FDeviceButtonState State = InputDeviceState.GetButtonState(Chord.Key);
+	
+	const bool bRequired = Stage == EInputStage::Start ? Binding.bRequiredToStart : Binding.bRequiredToContinue;
+	
+#define CHECK_REQUIRED(value) if (value) { ++OutComplexity; } else if (bRequired) { OutComplexity = INDEX_NONE; return; }
+
+	CHECK_REQUIRED(State.bDown);
+	
+	if (Chord.NeedsShift())
+	{
+		CHECK_REQUIRED(InputDeviceState.bShiftKeyDown);
+	}
+	
+	if (Chord.NeedsAlt())
+	{
+		CHECK_REQUIRED(InputDeviceState.bAltKeyDown);
+	}
+	
+	if (Chord.NeedsControl())
+	{
+		CHECK_REQUIRED(InputDeviceState.bCtrlKeyDown);
+	}
+	
+	if (Chord.NeedsCommand())
+	{
+		CHECK_REQUIRED(InputDeviceState.bCmdKeyDown);
+	}
+	
+#undef CHECK_REQUIRED
+	
+	bOutTrigger = bOutTrigger || (Binding.bTriggersStart && (Stage == EInputStage::Start ? State.bPressed : State.bDown));
+}
+
+}
diff --git a/Engine/Source/Editor/Experimental/EditorInteractiveToolsFramework/Private/ViewportInteractions/ViewportCameraTranslateInteraction.cpp b/Engine/Source/Editor/Experimental/EditorInteractiveToolsFramework/Private/ViewportInteractions/ViewportCameraTranslateInteraction.cpp
--- a/Engine/Source/Editor/Experimental/EditorInteractiveToolsFramework/Private/ViewportInteractions/ViewportCameraTranslateInteraction.cpp
+++ b/Engine/Source/Editor/Experimental/EditorInteractiveToolsFramework/Private/ViewportInteractions/ViewportCameraTranslateInteraction.cpp
@@ -44,7 +44,15 @@ UViewportCameraTranslateInteraction::UViewportCameraTranslateInteraction()
 		}
 
 		bool bResult = true;
-
+		
+		if (FEditorViewportClient* ViewportClient = GetEditorViewportClient())
+		{
+			if (ViewportClient->IsCameraLocked())
+			{
+				return false;
+			}
+		}
+		
 		if (InputDeviceState.InputDevice == EInputDevices::Keyboard)
 		{
 			using namespace UE::Editor::ViewportInteractions;
diff --git a/Engine/Source/Editor/Experimental/EditorInteractiveToolsFramework/Private/ViewportInteractions/ViewportDollyInteraction.cpp b/Engine/Source/Editor/Experimental/EditorInteractiveToolsFramework/Private/ViewportInteractions/ViewportDollyInteraction.cpp
--- a/Engine/Source/Editor/Experimental/EditorInteractiveToolsFramework/Private/ViewportInteractions/ViewportDollyInteraction.cpp
+++ b/Engine/Source/Editor/Experimental/EditorInteractiveToolsFramework/Private/ViewportInteractions/ViewportDollyInteraction.cpp
@@ -12,10 +12,17 @@ UViewportDollyInteraction::UViewportDollyInteraction()
 	using namespace UE::Editor::ViewportInteractions;
 	InteractionName = PerspectiveDolly;
 	
-	ViewportClickDragBehavior->SetBindings({
-		EKeys::RightMouseButton,
+	NormalBindings = FButtonBindings({
+		FButtonBinding(EKeys::RightMouseButton).TriggersStart(),
 		FButtonBinding(EKeys::LeftAlt).RequiredToContinue(false)
 	});
+	
+	CameraLockedBindings = FButtonBindings({
+		FButtonBinding(EKeys::RightMouseButton).TriggersStart(),
+		FButtonBinding(EKeys::LeftAlt).Required(false)
+	});
+	
+	ViewportClickDragBehavior->SetBindings(UViewportClickDragBehavior::FGetBindingsDelegate::CreateUObject(this, &UViewportDollyInteraction::GetBindings));
 }
 
 bool UViewportDollyInteraction::CanBeActivated() const
@@ -63,3 +70,16 @@ void UViewportDollyInteraction::OnDragDelta(float InMouseDeltaX, float InMouseDe
 		}
 	}
 }
+
+const UE::Editor::ViewportInteractions::FButtonBindings& UViewportDollyInteraction::GetBindings() const
+{
+	if (FEditorViewportClient* Client = GetEditorViewportClient())
+	{
+		if (Client->IsCameraLocked())
+		{
+			return CameraLockedBindings;
+		}
+	}
+	return NormalBindings;
+}
+
diff --git a/Engine/Source/Editor/Experimental/EditorInteractiveToolsFramework/Private/ViewportInteractions/ViewportMoveYawInteraction.cpp b/Engine/Source/Editor/Experimental/EditorInteractiveToolsFramework/Private/ViewportInteractions/ViewportMoveYawInteraction.cpp
--- a/Engine/Source/Editor/Experimental/EditorInteractiveToolsFramework/Private/ViewportInteractions/ViewportMoveYawInteraction.cpp
+++ b/Engine/Source/Editor/Experimental/EditorInteractiveToolsFramework/Private/ViewportInteractions/ViewportMoveYawInteraction.cpp
@@ -26,7 +26,7 @@ bool UViewportMoveYawInteraction::CanBeActivated() const
 	{
 		if (FEditorViewportClient* EditorViewportClient = GetEditorViewportClient())
 		{
-			return EditorViewportClient->IsPerspective();
+			return EditorViewportClient->IsPerspective() && !EditorViewportClient->IsCameraLocked();
 		}
 	}
 	return false;
diff --git a/Engine/Source/Editor/Experimental/EditorInteractiveToolsFramework/Private/ViewportInteractions/ViewportOrbitInteraction.cpp b/Engine/Source/Editor/Experimental/EditorInteractiveToolsFramework/Private/ViewportInteractions/ViewportOrbitInteraction.cpp
--- a/Engine/Source/Editor/Experimental/EditorInteractiveToolsFramework/Private/ViewportInteractions/ViewportOrbitInteraction.cpp
+++ b/Engine/Source/Editor/Experimental/EditorInteractiveToolsFramework/Private/ViewportInteractions/ViewportOrbitInteraction.cpp
@@ -14,10 +14,18 @@ UViewportOrbitInteraction::UViewportOrbitInteraction()
 	Groups = { UE::Editor::ViewportInteractions::CameraDrag };
 	
 	using namespace UE::Editor::ViewportInteractions;
-	ViewportClickDragBehavior->SetBindings({
+	
+	NormalBindings = FButtonBindings({
+		FButtonBinding(EKeys::LeftMouseButton).TriggersStart(),
+    	FButtonBinding(EKeys::LeftAlt).RequiredToContinue(false)
+	});
+	
+	CameraLockedBindings = FButtonBindings({
 		FButtonBinding(EKeys::LeftMouseButton).TriggersStart(),
-		FButtonBinding(EKeys::LeftAlt).RequiredToContinue(false)
+		FButtonBinding(EKeys::LeftAlt).Required(false)
 	});
+	
+	ViewportClickDragBehavior->SetBindings(UViewportClickDragBehavior::FGetBindingsDelegate::CreateUObject(this, &UViewportOrbitInteraction::GetBindings));
 }
 
 bool UViewportOrbitInteraction::CanBeActivated() const
@@ -71,3 +79,16 @@ void UViewportOrbitInteraction::OnDragDelta(float InMouseDeltaX, float InMouseDe
 		EditorViewportClient->GetViewportNavigationHelper()->OrbitDelta += OrbitDelta;
 	}
 }
+
+const UE::Editor::ViewportInteractions::FButtonBindings& UViewportOrbitInteraction::GetBindings() const
+{
+	if (FEditorViewportClient* Client = GetEditorViewportClient())
+	{
+		if (Client->IsCameraLocked())
+		{
+			return CameraLockedBindings;
+		}
+	}
+	return NormalBindings;
+}
+
diff --git a/Engine/Source/Editor/Experimental/EditorInteractiveToolsFramework/Private/ViewportInteractions/ViewportViewAngleInteraction.cpp b/Engine/Source/Editor/Experimental/EditorInteractiveToolsFramework/Private/ViewportInteractions/ViewportViewAngleInteraction.cpp
--- a/Engine/Source/Editor/Experimental/EditorInteractiveToolsFramework/Private/ViewportInteractions/ViewportViewAngleInteraction.cpp
+++ b/Engine/Source/Editor/Experimental/EditorInteractiveToolsFramework/Private/ViewportInteractions/ViewportViewAngleInteraction.cpp
@@ -29,7 +29,7 @@ bool UViewportViewAngleInteraction::CanBeActivated() const
 	{
 		if (FEditorViewportClient* EditorViewportClient = GetEditorViewportClient())
 		{
-			return EditorViewportClient->IsPerspective();
+			return EditorViewportClient->IsPerspective() && !EditorViewportClient->IsCameraLocked();
 		}
 	}
 	return false;
diff --git a/Engine/Source/Editor/Experimental/EditorInteractiveToolsFramework/Public/Behaviors/ViewportClickDragBehavior.h b/Engine/Source/Editor/Experimental/EditorInteractiveToolsFramework/Public/Behaviors/ViewportClickDragBehavior.h
--- a/Engine/Source/Editor/Experimental/EditorInteractiveToolsFramework/Public/Behaviors/ViewportClickDragBehavior.h
+++ b/Engine/Source/Editor/Experimental/EditorInteractiveToolsFramework/Public/Behaviors/ViewportClickDragBehavior.h
@@ -5,6 +5,7 @@
 #include "MultiButtonClickDragBehavior.h"
 #include "Framework/Commands/InputChord.h"
 #include "Framework/Commands/UICommandInfo.h"
+#include "ViewportInteractions/ViewportInteractionBindings.h"
 #include "ViewportClickDragBehavior.generated.h"
 
 #define UE_API EDITORINTERACTIVETOOLSFRAMEWORK_API
@@ -64,60 +65,6 @@ public:
 	}
 };
 
-namespace UE::Editor::ViewportInteractions
-{
-	struct FButtonBinding
-	{
-		FButtonBinding(const FKey& InButton)
-			: Chord(InButton)
-			, Command(nullptr)
-		{}
-		
-		FButtonBinding(const TSharedPtr<FUICommandInfo>& InCommand)
-			: Chord()
-			, Command(InCommand)
-		{}
-		
-		FButtonBinding(const FButtonBinding& InBinding) = default;
-	
-		FInputChord Chord;
-		TWeakPtr<FUICommandInfo> Command;
-		
-		// TODO: These could be a flag-based enum, or just bitwise values
-		/** This button must be down for the behavior begin capture. */
-		bool bRequiredToStart = true;
-		/** If any binding is set to true, the behavior can capture when this button is pressed and all required buttons are active. */
-		bool bTriggersStart = false;
-		/** This button must be down for the behavior to continue capture. Should be a subset of bRequiredToStart. */
-		bool bRequiredToContinue = true;
-		
-		FButtonBinding& TriggersStart(bool bTriggers = true)
-		{
-			bTriggersStart = bTriggers;
-			return *this;
-		}
-		
-		FButtonBinding& Required(bool bRequired = true)
-		{
-			bRequiredToContinue = bRequired;
-			bRequiredToStart = bRequired;
-			return *this;
-		}
-		
-		FButtonBinding& RequiredToStart(bool bRequired = true)
-		{
-			bRequiredToStart = bRequired;
-			return *this;
-		}
-		
-		FButtonBinding& RequiredToContinue(bool bRequired = true)
-		{
-			bRequiredToContinue = bRequired;
-			return *this;
-		}
-	};
-}
-
 /**
  * A behavior dedicated to Viewport Drag Interactions, directly supporting threshold-based capture.
  * It can be used when capture stealing e.g. a single click behavior. The
@@ -130,13 +77,7 @@ class UViewportClickDragBehavior : public UInputBehavior
 	GENERATED_BODY()
 public:
 
-	enum class EInputStage
-	{
-		/** Where the behavior has not yet started. */
-		Start,
-		/** Where the behavior has started and needs to validate that it can continue. */
-		Continue
-	};
+	DECLARE_DELEGATE_RetVal(const UE::Editor::ViewportInteractions::FButtonBindings&, FGetBindingsDelegate);
 
 	UE_API void Initialize(IViewportClickDragBehaviorTarget* InTarget);
 
@@ -154,6 +95,7 @@ public:
 	UE_API virtual void ForceEndCapture(const FInputCaptureData& CaptureData) override;
 	
 	UE_API void SetBindings(const TArray<UE::Editor::ViewportInteractions::FButtonBinding>& InBindings);
+	UE_API void SetBindings(const FGetBindingsDelegate& InDelegate);
 
 	/** @return current 3D world ray and optional 2D position for Target Device */
 	UE_API virtual FInputDeviceRay GetDeviceRay(const FInputDeviceState& InInputDeviceState);
@@ -178,12 +120,6 @@ public:
 
 protected:
 
-	static void ProcessBinding(EInputStage Stage, const FInputDeviceState& InputDeviceState, const UE::Editor::ViewportInteractions::FButtonBinding& Binding, int32& OutComplexity, bool& bOutTrigger);
-	static void ProcessBindingChord(EInputStage Stage, const FInputDeviceState& InputDeviceState, const UE::Editor::ViewportInteractions::FButtonBinding& Binding, const FInputChord& Chord, int32& OutComplexity, bool& bOutTrigger);
-
-	int32 GetBindingComplexity(EInputStage Stage, const FInputDeviceState& InputDeviceState);
-	bool DidAnyBindingStateChange(const FInputDeviceState& InputDeviceState) const;
-
 	UE_API virtual void OnDragStartInternal(const FInputDeviceState& InputState);
 	UE_API virtual void OnClickDragInternal(
 		const FInputDeviceState& InputState, const FInputCaptureData& Data
@@ -218,8 +154,9 @@ protected:
 
 	IViewportClickDragBehaviorTarget* Target = nullptr;
 
-	TArray<UE::Editor::ViewportInteractions::FButtonBinding> Bindings;
-	bool bRequiresTrigger = false;
+	UE::Editor::ViewportInteractions::FButtonBindings Bindings;
+	FGetBindingsDelegate GetBindingsDelegate;
+	const UE::Editor::ViewportInteractions::FButtonBindings& GetBindings() const;
 	
 	int32 ActiveBindingComplexity = 0;
 
diff --git a/Engine/Source/Editor/Experimental/EditorInteractiveToolsFramework/Public/ViewportInteractions/ViewportDollyInteraction.h b/Engine/Source/Editor/Experimental/EditorInteractiveToolsFramework/Public/ViewportInteractions/ViewportDollyInteraction.h
--- a/Engine/Source/Editor/Experimental/EditorInteractiveToolsFramework/Public/ViewportInteractions/ViewportDollyInteraction.h
+++ b/Engine/Source/Editor/Experimental/EditorInteractiveToolsFramework/Public/ViewportInteractions/ViewportDollyInteraction.h
@@ -20,6 +20,11 @@ public:
 	UE_API virtual void OnDragEnd(const FInputDeviceRay& InDragEndPos) override;
 	UE_API virtual void OnDragDelta(float InMouseDeltaX, float InMouseDeltaY) override;
 	UE_API virtual bool CanBeActivated() const override;
+	
+private:
+	UE::Editor::ViewportInteractions::FButtonBindings NormalBindings;
+	UE::Editor::ViewportInteractions::FButtonBindings CameraLockedBindings;
+	const UE::Editor::ViewportInteractions::FButtonBindings& GetBindings() const;
 };
 
 #undef UE_API
\ No newline at end of file
diff --git a/Engine/Source/Editor/Experimental/EditorInteractiveToolsFramework/Public/ViewportInteractions/ViewportInteractionBindings.h b/Engine/Source/Editor/Experimental/EditorInteractiveToolsFramework/Public/ViewportInteractions/ViewportInteractionBindings.h
new file mode 100644
--- /dev/null
+++ b/Engine/Source/Editor/Experimental/EditorInteractiveToolsFramework/Public/ViewportInteractions/ViewportInteractionBindings.h
@@ -0,0 +1,89 @@
+﻿// Copyright Epic Games, Inc. All Rights Reserved.
+
+#pragma once
+
+#include "InputCoreTypes.h"
+#include "InputState.h"
+#include "Framework/Commands/UICommandInfo.h"
+
+namespace UE::Editor::ViewportInteractions
+{
+
+enum class EInputStage
+{
+	/** Where the behavior has not yet started. */
+	Start,
+	/** Where the behavior has started and needs to validate that it can continue. */
+	Continue
+};
+
+struct FButtonBinding
+{
+	FButtonBinding(const FKey& InButton)
+		: Chord(InButton)
+		, Command(nullptr)
+	{}
+	
+	FButtonBinding(const TSharedPtr<FUICommandInfo>& InCommand)
+		: Chord()
+		, Command(InCommand)
+	{}
+	
+	FButtonBinding(const FButtonBinding& InBinding) = default;
+
+	FInputChord Chord;
+	TWeakPtr<FUICommandInfo> Command;
+	
+	// TODO: These could be a flag-based enum, or just bitwise values
+	/** This button must be down for the behavior begin capture. */
+	bool bRequiredToStart = true;
+	/** If any binding is set to true, the behavior can capture when this button is pressed and all required buttons are active. */
+	bool bTriggersStart = false;
+	/** This button must be down for the behavior to continue capture. Should be a subset of bRequiredToStart. */
+	bool bRequiredToContinue = true;
+	
+	FButtonBinding& TriggersStart(bool bTriggers = true)
+	{
+		bTriggersStart = bTriggers;
+		return *this;
+	}
+	
+	FButtonBinding& Required(bool bRequired = true)
+	{
+		bRequiredToContinue = bRequired;
+		bRequiredToStart = bRequired;
+		return *this;
+	}
+	
+	FButtonBinding& RequiredToStart(bool bRequired = true)
+	{
+		bRequiredToStart = bRequired;
+		return *this;
+	}
+	
+	FButtonBinding& RequiredToContinue(bool bRequired = true)
+	{
+		bRequiredToContinue = bRequired;
+		return *this;
+	}
+};
+
+
+struct FButtonBindings
+{
+	FButtonBindings() = default;
+	FButtonBindings(const FButtonBindings& OtherBindings) = default;
+	FButtonBindings(const TArray<FButtonBinding>& InBindings);
+
+	TArray<FButtonBinding> Bindings;
+	bool bRequiresTrigger = false;
+	
+	int32 GetBindingComplexity(EInputStage Stage, const FInputDeviceState& InputDeviceState) const;
+	bool DidAnyBindingStateChange(const FInputDeviceState& InputDeviceState) const;
+	
+protected:
+	static void ProcessBinding(EInputStage Stage, const FInputDeviceState& InputDeviceState, const FButtonBinding& Binding, int32& OutComplexity, bool& bOutTrigger);
+	static void ProcessBindingChord(EInputStage Stage, const FInputDeviceState& InputDeviceState, const FButtonBinding& Binding, const FInputChord& Chord, int32& OutComplexity, bool& bOutTrigger);
+};
+
+}
diff --git a/Engine/Source/Editor/Experimental/EditorInteractiveToolsFramework/Public/ViewportInteractions/ViewportOrbitInteraction.h b/Engine/Source/Editor/Experimental/EditorInteractiveToolsFramework/Public/ViewportInteractions/ViewportOrbitInteraction.h
--- a/Engine/Source/Editor/Experimental/EditorInteractiveToolsFramework/Public/ViewportInteractions/ViewportOrbitInteraction.h
+++ b/Engine/Source/Editor/Experimental/EditorInteractiveToolsFramework/Public/ViewportInteractions/ViewportOrbitInteraction.h
@@ -29,6 +29,11 @@ public:
 	UE_API virtual bool CanBeActivated() const override;
 	UE_API virtual void OnDragDelta(float InMouseDeltaX, float InMouseDeltaY) override;
 	//~ End UViewportDragInteraction
+	
+private:
+	UE::Editor::ViewportInteractions::FButtonBindings NormalBindings;
+	UE::Editor::ViewportInteractions::FButtonBindings CameraLockedBindings;
+	const UE::Editor::ViewportInteractions::FButtonBindings& GetBindings() const; 
 };
 
 #undef UE_API
