[Viewport ITF] A temporary priority hack that promotes Alt-Drag actions over everything else.

#rnx

#Synced-CL 50866064

diff --git a/Engine/Source/Editor/Experimental/EditorInteractiveToolsFramework/Private/Behaviors/ViewportInteractionBindings.cpp b/Engine/Source/Editor/Experimental/EditorInteractiveToolsFramework/Private/Behaviors/ViewportInteractionBindings.cpp
--- a/Engine/Source/Editor/Experimental/EditorInteractiveToolsFramework/Private/Behaviors/ViewportInteractionBindings.cpp
+++ b/Engine/Source/Editor/Experimental/EditorInteractiveToolsFramework/Private/Behaviors/ViewportInteractionBindings.cpp
@@ -110,7 +110,7 @@ void FButtonBindings::ProcessBindingChord(EInputStage Stage, const FInputDeviceS
 	
 	const bool bRequired = Stage == EInputStage::Start ? Binding.bRequiredToStart : Binding.bRequiredToContinue;
 	
-#define CHECK_REQUIRED(value) if (value) { ++OutComplexity; } else if (bRequired) { OutComplexity = INDEX_NONE; return; }
+#define CHECK_REQUIRED(value) if (value) { OutComplexity += Binding.ButtonComplexity; } else if (bRequired) { OutComplexity = INDEX_NONE; return; }
 
 	CHECK_REQUIRED(State.bDown);
 	
diff --git a/Engine/Source/Editor/Experimental/EditorInteractiveToolsFramework/Private/EditorDragTools/FrustumSelectInteraction.cpp b/Engine/Source/Editor/Experimental/EditorInteractiveToolsFramework/Private/EditorDragTools/FrustumSelectInteraction.cpp
--- a/Engine/Source/Editor/Experimental/EditorInteractiveToolsFramework/Private/EditorDragTools/FrustumSelectInteraction.cpp
+++ b/Engine/Source/Editor/Experimental/EditorInteractiveToolsFramework/Private/EditorDragTools/FrustumSelectInteraction.cpp
@@ -139,14 +139,16 @@ void FViewportFrustum::Calculate(FConvexVolume& OutFrustum) const
 
 UFrustumSelectInteraction::UFrustumSelectInteraction()
 {
-	InteractionName = UE::Editor::ViewportInteractions::FrustumSelect;
+	using namespace UE::Editor::ViewportInteractions;
+	
+	InteractionName = FrustumSelect;
 
 	if (UViewportClickDragBehavior* ClickDragBehavior = ClickDragInputBehavior.Get())
 	{
 		ClickDragBehavior->SetBindings({
-			UE::Editor::ViewportInteractions::FButtonBinding(EKeys::LeftMouseButton).TriggersStart(),
+			FButtonBinding(EKeys::LeftMouseButton).TriggersStart(),
 			EKeys::LeftControl,
-			EKeys::LeftAlt
+			FButtonBinding(EKeys::LeftAlt).Complexity(20)
 		});
 	}
 }
diff --git a/Engine/Source/Editor/Experimental/EditorInteractiveToolsFramework/Private/ViewportInteractions/ViewportDollyInteraction.cpp b/Engine/Source/Editor/Experimental/EditorInteractiveToolsFramework/Private/ViewportInteractions/ViewportDollyInteraction.cpp
--- a/Engine/Source/Editor/Experimental/EditorInteractiveToolsFramework/Private/ViewportInteractions/ViewportDollyInteraction.cpp
+++ b/Engine/Source/Editor/Experimental/EditorInteractiveToolsFramework/Private/ViewportInteractions/ViewportDollyInteraction.cpp
@@ -14,12 +14,12 @@ UViewportDollyInteraction::UViewportDollyInteraction()
 	
 	NormalBindings = FButtonBindings({
 		FButtonBinding(EKeys::RightMouseButton).TriggersStart(),
-		FButtonBinding(EKeys::LeftAlt).RequiredToContinue(false)
+		FButtonBinding(EKeys::LeftAlt).RequiredToContinue(false).Complexity(20)
 	});
 	
 	CameraLockedBindings = FButtonBindings({
 		FButtonBinding(EKeys::RightMouseButton).TriggersStart(),
-		FButtonBinding(EKeys::LeftAlt).Required(false)
+		FButtonBinding(EKeys::LeftAlt).Required(false).Complexity(20)
 	});
 	
 	ViewportClickDragBehavior->SetBindings(UViewportClickDragBehavior::FGetBindingsDelegate::CreateUObject(this, &UViewportDollyInteraction::GetBindings));
diff --git a/Engine/Source/Editor/Experimental/EditorInteractiveToolsFramework/Private/ViewportInteractions/ViewportOrbitInteraction.cpp b/Engine/Source/Editor/Experimental/EditorInteractiveToolsFramework/Private/ViewportInteractions/ViewportOrbitInteraction.cpp
--- a/Engine/Source/Editor/Experimental/EditorInteractiveToolsFramework/Private/ViewportInteractions/ViewportOrbitInteraction.cpp
+++ b/Engine/Source/Editor/Experimental/EditorInteractiveToolsFramework/Private/ViewportInteractions/ViewportOrbitInteraction.cpp
@@ -17,12 +17,12 @@ UViewportOrbitInteraction::UViewportOrbitInteraction()
 	
 	NormalBindings = FButtonBindings({
 		FButtonBinding(EKeys::LeftMouseButton).TriggersStart(),
-    	FButtonBinding(EKeys::LeftAlt).RequiredToContinue(false)
+    	FButtonBinding(EKeys::LeftAlt).RequiredToContinue(false).Complexity(20)
 	});
 	
 	CameraLockedBindings = FButtonBindings({
 		FButtonBinding(EKeys::LeftMouseButton).TriggersStart(),
-		FButtonBinding(EKeys::LeftAlt).Required(false)
+		FButtonBinding(EKeys::LeftAlt).Required(false).Complexity(20)
 	});
 	
 	ViewportClickDragBehavior->SetBindings(UViewportClickDragBehavior::FGetBindingsDelegate::CreateUObject(this, &UViewportOrbitInteraction::GetBindings));
diff --git a/Engine/Source/Editor/Experimental/EditorInteractiveToolsFramework/Private/ViewportInteractions/ViewportPanInteraction.cpp b/Engine/Source/Editor/Experimental/EditorInteractiveToolsFramework/Private/ViewportInteractions/ViewportPanInteraction.cpp
--- a/Engine/Source/Editor/Experimental/EditorInteractiveToolsFramework/Private/ViewportInteractions/ViewportPanInteraction.cpp
+++ b/Engine/Source/Editor/Experimental/EditorInteractiveToolsFramework/Private/ViewportInteractions/ViewportPanInteraction.cpp
@@ -20,7 +20,7 @@ UViewportPanInteraction::UViewportPanInteraction()
 	using namespace UE::Editor::ViewportInteractions;
 	ViewportClickDragBehavior->SetBindings({
 		FButtonBinding(EKeys::MiddleMouseButton),
-		FButtonBinding(EKeys::LeftAlt).Required(false)
+		FButtonBinding(EKeys::LeftAlt).Required(false).Complexity(20)
 	});
 }
 
diff --git a/Engine/Source/Editor/Experimental/EditorInteractiveToolsFramework/Public/ViewportInteractions/ViewportInteractionBindings.h b/Engine/Source/Editor/Experimental/EditorInteractiveToolsFramework/Public/ViewportInteractions/ViewportInteractionBindings.h
--- a/Engine/Source/Editor/Experimental/EditorInteractiveToolsFramework/Public/ViewportInteractions/ViewportInteractionBindings.h
+++ b/Engine/Source/Editor/Experimental/EditorInteractiveToolsFramework/Public/ViewportInteractions/ViewportInteractionBindings.h
@@ -34,6 +34,13 @@ struct FButtonBinding
 	FInputChord Chord;
 	TWeakPtr<FUICommandInfo> Command;
 	
+	/**
+	 * The complexity value of each activated button.
+	 * This feature is present as an interim patch for the purposes of bespoke priority hacks.
+	 * Modifying complexity like this is high-odor system smell. This WILL go away. 
+	 */
+	int32 ButtonComplexity = 1;
+	
 	// TODO: These could be a flag-based enum, or just bitwise values
 	/** This button must be down for the behavior begin capture. */
 	bool bRequiredToStart = true;
@@ -66,6 +73,12 @@ struct FButtonBinding
 		bRequiredToContinue = bRequired;
 		return *this;
 	}
+	
+	FButtonBinding& Complexity(int32 InComplexityMultiplier)
+	{
+		ButtonComplexity = InComplexityMultiplier;
+		return *this;
+	}
 };
 
 
diff --git a/Engine/Source/Editor/SkeletalMeshEditor/Private/SkeletalMeshEditor.cpp b/Engine/Source/Editor/SkeletalMeshEditor/Private/SkeletalMeshEditor.cpp
--- a/Engine/Source/Editor/SkeletalMeshEditor/Private/SkeletalMeshEditor.cpp
+++ b/Engine/Source/Editor/SkeletalMeshEditor/Private/SkeletalMeshEditor.cpp
@@ -1631,12 +1631,6 @@ TSharedPtr<FUICommandInfo> FSkeletalMeshEditor::GetResetAllBonesTransformsComman
 void FSkeletalMeshEditor::OnBuildViewportInteractions(UViewportInteractionsBehaviorSource* Source)
 {
 	UE::Editor::ViewportInteractions::AddDefaultCameraMovementInteractions(Source);
-			
-	// Making sure viewport orbit is prioritized before regular tool drag tools	
-	for (const TObjectPtr<UInputBehavior>& InputBehavior : Source->FindInteraction(UViewportOrbitInteraction::StaticClass())->GetInputBehaviors())
-	{
-		InputBehavior->SetDefaultPriority(UE::Editor::ViewportInteractions::VIEWPORT_INTERACTIONS_DEFAULT_PRIORITY);
-	}
 }
 
 FSkeletalMeshEditorBinding::FSkeletalMeshEditorBinding(TSharedRef<FSkeletalMeshEditor> InEditor)
