[ITF Viewport Interactions] Fix an issue where the cross-left & cross-right marquee selection options were broken for Frustum Select due to incorrect swapping of the start & end positions.

#rnx
#Synced-CL 51017708

diff --git a/Engine/Source/Editor/Experimental/EditorInteractiveToolsFramework/Private/EditorDragTools/FrustumSelectInteraction.cpp b/Engine/Source/Editor/Experimental/EditorInteractiveToolsFramework/Private/EditorDragTools/FrustumSelectInteraction.cpp
--- a/Engine/Source/Editor/Experimental/EditorInteractiveToolsFramework/Private/EditorDragTools/FrustumSelectInteraction.cpp
+++ b/Engine/Source/Editor/Experimental/EditorInteractiveToolsFramework/Private/EditorDragTools/FrustumSelectInteraction.cpp
@@ -218,17 +218,6 @@ void UFrustumSelectInteraction::OnDragEnd(const FInputDeviceRay& InReleasePos)
 	);
 	FSceneView* SceneView = EditorViewportClient->CalcSceneView(&ViewFamily);
 
-	// Normalize the start and end
-	if (Start.X > End.X)
-	{
-		Swap(Start.X, End.X);
-	}
-
-	if (Start.Y > End.Y)
-	{
-		Swap(Start.Y, End.Y);
-	}
-
 	// Generate a frustum out of the dragged box
 	FConvexVolume Frustum;
 	CalculateFrustum(SceneView, Frustum, true);
@@ -296,11 +285,14 @@ void UFrustumSelectInteraction::OnDragEnd(const FInputDeviceRay& InReleasePos)
 		else
 		{
 			// Extend the endpoint of the rect to get the actual line
-
-			const int32 MinX = UE::LWC::FloatToIntCastChecked<int32>(FMath::Max<double>(0.0, Start.X));
-			const int32 MinY = UE::LWC::FloatToIntCastChecked<int32>(FMath::Max<double>(0.0, Start.Y));
-			const int32 MaxX = FMath::Min(ViewportSizeX, FMath::TruncToInt32(End.X + 1.0));
-			const int32 MaxY = FMath::Min(ViewportSizeY, FMath::TruncToInt32(End.Y + 1.0));
+			
+			const FVector MinVector = FVector::Min(Start, End);
+			const FVector MaxVector = FVector::Max(Start, End);
+
+			const int32 MinX = UE::LWC::FloatToIntCastChecked<int32>(FMath::Max<double>(0.0, MinVector.X));
+			const int32 MinY = UE::LWC::FloatToIntCastChecked<int32>(FMath::Max<double>(0.0, MinVector.Y));
+			const int32 MaxX = FMath::Min(ViewportSizeX, FMath::TruncToInt32(MaxVector.X + 1.0));
+			const int32 MaxY = FMath::Min(ViewportSizeY, FMath::TruncToInt32(MaxVector.Y + 1.0));
 
 			const FIntPoint Min{ MinX, MinY };
 			const FIntPoint Max{ MaxX, MaxY };
