[Editor TRS] Gizmo settings correctly re-applied on "Reset to Default", and "Explicit" property applied/restored

#Synced-CL 51038722

diff --git a/Engine/Source/Editor/Experimental/EditorInteractiveToolsFramework/Private/TransformGizmoEditorSettings.cpp b/Engine/Source/Editor/Experimental/EditorInteractiveToolsFramework/Private/TransformGizmoEditorSettings.cpp
--- a/Engine/Source/Editor/Experimental/EditorInteractiveToolsFramework/Private/TransformGizmoEditorSettings.cpp
+++ b/Engine/Source/Editor/Experimental/EditorInteractiveToolsFramework/Private/TransformGizmoEditorSettings.cpp
@@ -12,10 +12,7 @@ UTransformGizmoEditorSettings::UTransformGizmoEditorSettings(const FObjectInitia
 {
 	if (ULevelEditorViewportSettings* const ViewportSettings = GetMutableDefault<ULevelEditorViewportSettings>())
 	{
-		bEnableArcballRotate = ViewportSettings->bAllowArcballRotate;
-		bEnableScreenRotate = ViewportSettings->bAllowScreenRotate;
-		bEnableAxisDrawing = ViewportSettings->bAllowEditWidgetAxisDisplay;
-		bEnableCombinedTranslateRotate = ViewportSettings->bAllowTranslateRotateZWidget;
+		SyncFromLevelEditorSettings();
 
 		ViewportSettings->OnSettingChanged().AddUObject(this, &UTransformGizmoEditorSettings::OnLegacySettingChanged);
 	}
@@ -91,6 +88,22 @@ void UTransformGizmoEditorSettings::PostEditChangeProperty(FPropertyChangedEvent
 		}
 	}
 }
+
+
+void UTransformGizmoEditorSettings::PostReloadConfig(FProperty* PropertyThatWasLoaded)
+{
+	Super::PostReloadConfig(PropertyThatWasLoaded);
+
+	SyncFromLevelEditorSettings();
+
+	BroadcastNewTRSGizmoChange();
+	BroadcastGizmosParametersChange();
+
+	if (GEditor)
+	{
+		GEditor->RedrawAllViewports();
+	}
+}
 #endif // WITH_EDITOR
 
 void UTransformGizmoEditorSettings::SetUseExperimentalGizmo(bool bInUseExperimentalGizmo)
@@ -155,6 +168,17 @@ void UTransformGizmoEditorSettings::BroadcastGizmosParametersChange() const
 	UEditorInteractiveGizmoManager::OnGizmosParametersChangedDelegate().Broadcast(GizmosParameters);
 }
 
+void UTransformGizmoEditorSettings::SyncFromLevelEditorSettings()
+{
+	if (ULevelEditorViewportSettings* const ViewportSettings = GetMutableDefault<ULevelEditorViewportSettings>())
+	{
+		bEnableArcballRotate = ViewportSettings->bAllowArcballRotate;
+		bEnableScreenRotate = ViewportSettings->bAllowScreenRotate;
+		bEnableAxisDrawing = ViewportSettings->bAllowEditWidgetAxisDisplay;
+		bEnableCombinedTranslateRotate = ViewportSettings->bAllowTranslateRotateZWidget;
+	}
+}
+
 void UTransformGizmoEditorSettings::OnLegacySettingChanged(FName InPropertyName)
 {
 	// Retrieving names from ULevelEditorViewportSettings
diff --git a/Engine/Source/Editor/Experimental/EditorInteractiveToolsFramework/Public/TransformGizmoEditorSettings.h b/Engine/Source/Editor/Experimental/EditorInteractiveToolsFramework/Public/TransformGizmoEditorSettings.h
--- a/Engine/Source/Editor/Experimental/EditorInteractiveToolsFramework/Public/TransformGizmoEditorSettings.h
+++ b/Engine/Source/Editor/Experimental/EditorInteractiveToolsFramework/Public/TransformGizmoEditorSettings.h
@@ -20,6 +20,7 @@ public:
 #if WITH_EDITOR
 	UE_API virtual void PostEditChangeChainProperty(FPropertyChangedChainEvent& InPropertyChangedEvent) override;
 	UE_API virtual void PostEditChangeProperty(FPropertyChangedEvent& InPropertyChangedEvent) override;
+	UE_API virtual void PostReloadConfig(class FProperty* PropertyThatWasLoaded) override;
 #endif
 
 	UPROPERTY(Config, EditAnywhere, Category = "Transform Gizmo", DisplayName = "Transform Gizmo Size", meta = (ClampMin = "-10.0", ClampMax = "150.0", EditCondition = "!bUseExperimentalGizmo"))
@@ -74,6 +75,9 @@ private:
 	UE_API void BroadcastNewTRSGizmoChange() const;
 	UE_API void BroadcastGizmosParametersChange() const;
 
+	/** Writes to Gizmos settings from the current LevelEditorSettings. */
+	UE_API void SyncFromLevelEditorSettings();
+
 	UE_API void OnLegacySettingChanged(FName InPropertyName);
 };
 
