[ITF interaction] fixed crash when zooming without any focused vpc

#rnx

#Synced-CL 51156304

diff --git a/Engine/Source/Editor/Experimental/EditorInteractiveToolsFramework/Private/ViewportInteractions/ViewportZoomInteraction.cpp b/Engine/Source/Editor/Experimental/EditorInteractiveToolsFramework/Private/ViewportInteractions/ViewportZoomInteraction.cpp
--- a/Engine/Source/Editor/Experimental/EditorInteractiveToolsFramework/Private/ViewportInteractions/ViewportZoomInteraction.cpp
+++ b/Engine/Source/Editor/Experimental/EditorInteractiveToolsFramework/Private/ViewportInteractions/ViewportZoomInteraction.cpp
@@ -249,6 +249,10 @@ void UViewportZoomInteraction::OrthographicZoom(bool bInZoomIn, float InZoomMult
 void UViewportZoomInteraction::PerspectiveZoom(bool bInZoomIn) const
 {
 	FEditorViewportClient* EditorViewportClient = GetEditorViewportClient();
+	if (!EditorViewportClient)
+	{
+		return;
+	}
 
 	// Scrolling the mousewheel up/down moves the perspective viewport forwards/backwards.
 	FVector Drag;
@@ -278,15 +282,16 @@ void UViewportZoomInteraction::PerspectiveZoom(bool bInZoomIn) const
 
 void UViewportZoomInteraction::MouseWheelZoom(bool bInZoomIn) const
 {
-	FEditorViewportClient* EditorViewportClient = GetEditorViewportClient();
-
-	if (EditorViewportClient->IsPerspective())
+	if (const FEditorViewportClient* EditorViewportClient = GetEditorViewportClient())
 	{
-		PerspectiveZoom(bInZoomIn);
-	}
-	else
-	{
-		OrthographicZoom(bInZoomIn);
+		if (EditorViewportClient->IsPerspective())
+		{
+			PerspectiveZoom(bInZoomIn);
+		}
+		else
+		{
+			OrthographicZoom(bInZoomIn);
+		}
 	}
 }
 
