[ITF Viewpot Interactions] Reconfigure ITF binding complexity interactions so that tools, clicks, and drags coexist more equally.
- Add new UViewportClickBehavior, used by UViewportClickInteraction.
- Revert CL 50877288 that added a temporary hack where behaviors using "Alt" had set a priority multiplier.

Switch binding-based priority to use a value of 3 per bound key. A confirmed drag is also worth an additional 3 points. The base of clicks and drags are offset from DEFAULT_TOOL_PRIORITY such that a complex viewport click or drag beats out a default Tool. At the same time, inputs of the same comlexity rank out in Confirmed Drag > Tool > Click > Unconfirmed Drag.

#rnx
#Synced-CL 51094013

diff --git a/Engine/Source/Editor/Experimental/EditorInteractiveToolsFramework/Private/Behaviors/ViewportClickBehavior.cpp b/Engine/Source/Editor/Experimental/EditorInteractiveToolsFramework/Private/Behaviors/ViewportClickBehavior.cpp
new file mode 100644
--- /dev/null
+++ b/Engine/Source/Editor/Experimental/EditorInteractiveToolsFramework/Private/Behaviors/ViewportClickBehavior.cpp
@@ -0,0 +1,163 @@
+﻿// Copyright Epic Games, Inc. All Rights Reserved.
+
+#include "Behaviors/ViewportClickBehavior.h"
+
+#include "ViewportInteractions/ViewportInteraction.h"
+
+void UViewportClickBehavior::Initialize(IViewportClickBehaviorTarget* InTarget)
+{
+	check (InTarget != nullptr);
+	Target = InTarget;
+	
+	DefaultPriority = UE::Editor::ViewportInteractions::CLICK_PRIORITY;
+}
+
+FInputCapturePriority UViewportClickBehavior::GetPriority()
+{
+	if (const FInputCapturePriority* Priority = CachedCapturePriority.GetPtrOrNull())
+	{
+		return *Priority;
+	}
+	return Super::GetPriority();
+}
+
+FInputCaptureRequest UViewportClickBehavior::WantsCapture(const FInputDeviceState& InInputState)
+{
+	if (!Target)
+	{
+		return FInputCaptureRequest::Ignore();
+	}
+
+	if (!CachedCapturePriority.IsSet())
+	{
+		CachedCapturePriority = GetCapturePriority(UE::Editor::ViewportInteractions::EInputStage::Start, InInputState);
+	}
+	else if (InInputState.IsFromDevice(EInputDevices::Mouse))
+	{
+		TOptional<FInputCapturePriority> NewPriority = GetCapturePriority(UE::Editor::ViewportInteractions::EInputStage::Continue, InInputState);
+		// If in a click, WantsCapture() is occuring as this behavior is defending in a capture stealing pass.
+		// A fresh unset priority means that this click wants to resolve with the "ClickUp" event.
+		// Retaining the previous capture priority gives the behavior a chance to win and close out in UpdateCapture().
+		// A set priority always updates as any additional modifier keys needs to be tracked. 
+		if (NewPriority.IsSet() || !bIsClicking)
+		{
+			CachedCapturePriority = NewPriority;
+		}
+	}
+	
+	if (!CachedCapturePriority.IsSet())
+	{
+		return FInputCaptureRequest::Ignore();
+	}
+	
+	const FInputRayHit HitResult = Target->IsHitByClick(GetDeviceRay(InInputState));
+	if (!HitResult.bHit)
+	{
+		CachedCapturePriority.Reset();
+		return FInputCaptureRequest::Ignore();
+	}
+	
+	return FInputCaptureRequest::Begin(this, EInputCaptureSide::Any, HitResult.HitDepth);
+}
+
+FInputCaptureUpdate UViewportClickBehavior::BeginCapture(const FInputDeviceState& InputState, EInputCaptureSide eSide)
+{
+	OnStateUpdatedInternal(InputState);
+	OnClickDownInternal(InputState);
+	bIsClicking = true;
+	return FInputCaptureUpdate::Begin(this, EInputCaptureSide::Any);
+}
+
+FInputCaptureUpdate UViewportClickBehavior::UpdateCapture(const FInputDeviceState& InputState, const FInputCaptureData& CaptureData)
+{
+	// This happens when a keyboard event is routed to the active mouse capture for updating modifier states
+	// The behavior should always continue here so that subsequent mouse updates can end capture more correctly	
+	if (!InputState.IsFromDevice(EInputDevices::Mouse))
+	{
+		return FInputCaptureUpdate::Continue();
+	}
+	
+	if (Bindings.DidAnyBindingStateChange(InputState))
+	{
+		OnStateUpdatedInternal(InputState);
+	}
+	
+	CachedCapturePriority = GetCapturePriority(UE::Editor::ViewportInteractions::EInputStage::Continue, InputState);
+	if (!CachedCapturePriority.IsSet())
+	{
+		OnClickUpInternal(InputState);
+		bIsClicking = false;
+		return FInputCaptureUpdate::End();
+	}
+	
+	return FInputCaptureUpdate::Continue();
+}
+
+void UViewportClickBehavior::ForceEndCapture(const FInputCaptureData& CaptureData)
+{
+	CachedCapturePriority.Reset();
+	bIsClicking = false;
+	if (Target)
+	{
+		Target->OnForceEndCapture();
+	}
+}
+
+void UViewportClickBehavior::SetBindings(const TArray<UE::Editor::ViewportInteractions::FButtonBinding>& InBindings)
+{
+	Bindings = InBindings;
+}
+
+FInputDeviceRay UViewportClickBehavior::GetDeviceRay(const FInputDeviceState& InInputDeviceState)
+{
+	if (InInputDeviceState.IsFromDevice(EInputDevices::Mouse))
+	{
+		const FVector2D& MousePosition = GetMousePosition(InInputDeviceState);
+		const FRay& WorldRay = InInputDeviceState.Mouse.WorldRay;
+
+		return FInputDeviceRay(WorldRay, MousePosition);
+	}
+
+	return FInputDeviceRay(FRay(FVector::ZeroVector, FVector(0, 0, 1), true));
+}
+
+const FVector2D& UViewportClickBehavior::GetMousePosition(const FInputDeviceState& InInputDeviceState) const
+{
+	return InInputDeviceState.Mouse.Position2D;
+}
+
+void UViewportClickBehavior::OnClickDownInternal(const FInputDeviceState& InputState)
+{
+	if (Target)
+	{
+		Target->OnClickDown(GetDeviceRay(InputState));
+	}
+}
+
+void UViewportClickBehavior::OnClickUpInternal(const FInputDeviceState& InputState)
+{
+	if (Target)
+	{
+		Target->OnClickUp(GetDeviceRay(InputState));
+	}
+}
+
+void UViewportClickBehavior::OnStateUpdatedInternal(const FInputDeviceState& InInputDeviceState)
+{
+	if (Target)
+	{
+		Target->OnStateUpdated(InInputDeviceState);
+	}
+}
+
+TOptional<FInputCapturePriority> UViewportClickBehavior::GetCapturePriority(UE::Editor::ViewportInteractions::EInputStage Stage, const FInputDeviceState& InputState) const
+{
+	const int32 Complexity = Bindings.GetBindingComplexity(Stage, InputState);
+	if (Complexity <= 0)
+	{
+		return TOptional<FInputCapturePriority>();
+	}
+	
+	return DefaultPriority.MakeHigher(Complexity * UE::Editor::ViewportInteractions::BINDING_COMPLEXITY_PRIORITY_STEP);
+}
+
diff --git a/Engine/Source/Editor/Experimental/EditorInteractiveToolsFramework/Private/Behaviors/ViewportClickDragBehavior.cpp b/Engine/Source/Editor/Experimental/EditorInteractiveToolsFramework/Private/Behaviors/ViewportClickDragBehavior.cpp
--- a/Engine/Source/Editor/Experimental/EditorInteractiveToolsFramework/Private/Behaviors/ViewportClickDragBehavior.cpp
+++ b/Engine/Source/Editor/Experimental/EditorInteractiveToolsFramework/Private/Behaviors/ViewportClickDragBehavior.cpp
@@ -29,14 +29,11 @@ void UViewportClickDragBehavior::SetDragConfirmedPriority(FInputCapturePriority
 
 FInputCapturePriority UViewportClickDragBehavior::GetPriority()
 {
-	FInputCapturePriority BasePriority = DefaultPriority;
-	
-	if (HasMovedEnoughForDrag())
+	if (const FInputCapturePriority* Priority = CachedCapturePriority.GetPtrOrNull())
 	{
-		BasePriority = DragConfirmedPriority;
+		return *Priority;
 	}
-	
-	return BasePriority.MakeHigher(ActiveBindingComplexity);
+	return Super::GetPriority();
 }
 
 FInputCaptureRequest UViewportClickDragBehavior::WantsCapture(const FInputDeviceState& InInputState)
@@ -45,8 +42,8 @@ FInputCaptureRequest UViewportClickDragBehavior::WantsCapture(const FInputDevice
 
 	if (!MouseTraveledDistance.IsSet())
 	{
-		ActiveBindingComplexity = GetBindings().GetBindingComplexity(UE::Editor::ViewportInteractions::EInputStage::Start, InInputState);
-		if (ActiveBindingComplexity > 0)
+		CachedCapturePriority = GetCapturePriority(UE::Editor::ViewportInteractions::EInputStage::Start, InInputState);
+		if (CachedCapturePriority.IsSet())
 		{
 			// Make sure behavior knows we are not dragging yet
 			bDragging = false;
@@ -66,8 +63,8 @@ FInputCaptureRequest UViewportClickDragBehavior::WantsCapture(const FInputDevice
 	{
 		if (InInputState.IsFromDevice(EInputDevices::Mouse))
 		{
-			ActiveBindingComplexity = GetBindings().GetBindingComplexity(UE::Editor::ViewportInteractions::EInputStage::Continue, InInputState);
-			if (ActiveBindingComplexity <= 0)
+			CachedCapturePriority = GetCapturePriority(UE::Editor::ViewportInteractions::EInputStage::Continue, InInputState);
+			if (!CachedCapturePriority.IsSet())
 			{
 				MouseTraveledDistance = TOptional<float>();
 				return FInputCaptureRequest::Ignore();
@@ -81,7 +78,7 @@ FInputCaptureRequest UViewportClickDragBehavior::WantsCapture(const FInputDevice
 	if (!HitResult.bHit)
 	{
 		MouseTraveledDistance = TOptional<float>();
-		ActiveBindingComplexity = 0;
+		CachedCapturePriority.Reset();
 		return FInputCaptureRequest::Ignore();
 	}
 	
@@ -115,7 +112,8 @@ FInputCaptureUpdate UViewportClickDragBehavior::UpdateCapture(const FInputDevice
 		OnStateUpdatedInternal(InputState);
 	}
 	
-	if (GetBindings().GetBindingComplexity(UE::Editor::ViewportInteractions::EInputStage::Continue, InputState) <= 0)
+	CachedCapturePriority = GetCapturePriority(UE::Editor::ViewportInteractions::EInputStage::Continue, InputState);
+	if (!CachedCapturePriority.IsSet())
 	{
 		if (bDragging)
 		{
@@ -145,7 +143,7 @@ FInputCaptureUpdate UViewportClickDragBehavior::UpdateCapture(const FInputDevice
 void UViewportClickDragBehavior::ForceEndCapture(const FInputCaptureData& CaptureData)
 {
 	MouseTraveledDistance = TOptional<float>();
-	ActiveBindingComplexity = 0;
+	CachedCapturePriority.Reset();
 
 	if (bDragging)
 	{
@@ -248,3 +246,21 @@ const UE::Editor::ViewportInteractions::FButtonBindings& UViewportClickDragBehav
 	}
 	return Bindings;
 }
+
+TOptional<FInputCapturePriority> UViewportClickDragBehavior::GetCapturePriority(UE::Editor::ViewportInteractions::EInputStage Stage, const FInputDeviceState& InputState) const
+{
+	FInputCapturePriority BasePriority = DefaultPriority;
+	
+	if (HasMovedEnoughForDrag())
+	{
+		BasePriority = DragConfirmedPriority;
+	}
+	
+	const int32 Complexity = GetBindings().GetBindingComplexity(Stage, InputState);
+	if (Complexity <= 0)
+	{
+		return TOptional<FInputCapturePriority>();
+	}
+	
+	return BasePriority.MakeHigher(Complexity * UE::Editor::ViewportInteractions::BINDING_COMPLEXITY_PRIORITY_STEP);
+}
diff --git a/Engine/Source/Editor/Experimental/EditorInteractiveToolsFramework/Private/Behaviors/ViewportInteractionBindings.cpp b/Engine/Source/Editor/Experimental/EditorInteractiveToolsFramework/Private/Behaviors/ViewportInteractionBindings.cpp
--- a/Engine/Source/Editor/Experimental/EditorInteractiveToolsFramework/Private/Behaviors/ViewportInteractionBindings.cpp
+++ b/Engine/Source/Editor/Experimental/EditorInteractiveToolsFramework/Private/Behaviors/ViewportInteractionBindings.cpp
@@ -110,7 +110,7 @@ void FButtonBindings::ProcessBindingChord(EInputStage Stage, const FInputDeviceS
 	
 	const bool bRequired = Stage == EInputStage::Start ? Binding.bRequiredToStart : Binding.bRequiredToContinue;
 	
-#define CHECK_REQUIRED(value) if (value) { OutComplexity += Binding.ButtonComplexity; } else if (bRequired) { OutComplexity = INDEX_NONE; return; }
+#define CHECK_REQUIRED(value) if (value) { ++OutComplexity; } else if (bRequired) { OutComplexity = INDEX_NONE; return; }
 
 	CHECK_REQUIRED(State.bDown);
 	
diff --git a/Engine/Source/Editor/Experimental/EditorInteractiveToolsFramework/Private/EditorDragTools/FrustumSelectInteraction.cpp b/Engine/Source/Editor/Experimental/EditorInteractiveToolsFramework/Private/EditorDragTools/FrustumSelectInteraction.cpp
--- a/Engine/Source/Editor/Experimental/EditorInteractiveToolsFramework/Private/EditorDragTools/FrustumSelectInteraction.cpp
+++ b/Engine/Source/Editor/Experimental/EditorInteractiveToolsFramework/Private/EditorDragTools/FrustumSelectInteraction.cpp
@@ -139,16 +139,14 @@ void FViewportFrustum::Calculate(FConvexVolume& OutFrustum) const
 
 UFrustumSelectInteraction::UFrustumSelectInteraction()
 {
-	using namespace UE::Editor::ViewportInteractions;
-	
-	InteractionName = FrustumSelect;
+	InteractionName = UE::Editor::ViewportInteractions::FrustumSelect;
 
 	if (UViewportClickDragBehavior* ClickDragBehavior = ClickDragInputBehavior.Get())
 	{
 		ClickDragBehavior->SetBindings({
-			FButtonBinding(EKeys::LeftMouseButton).TriggersStart(),
+			UE::Editor::ViewportInteractions::FButtonBinding(EKeys::LeftMouseButton).TriggersStart(),
 			EKeys::LeftControl,
-			FButtonBinding(EKeys::LeftAlt).Complexity(20)
+			EKeys::LeftAlt
 		});
 	}
 }
diff --git a/Engine/Source/Editor/Experimental/EditorInteractiveToolsFramework/Private/ViewportInteractions/ViewportClickInteraction.cpp b/Engine/Source/Editor/Experimental/EditorInteractiveToolsFramework/Private/ViewportInteractions/ViewportClickInteraction.cpp
--- a/Engine/Source/Editor/Experimental/EditorInteractiveToolsFramework/Private/ViewportInteractions/ViewportClickInteraction.cpp
+++ b/Engine/Source/Editor/Experimental/EditorInteractiveToolsFramework/Private/ViewportInteractions/ViewportClickInteraction.cpp
@@ -8,32 +8,38 @@
 #include "SceneView.h"
 #include "ViewportInteractions/ViewportInteractionsBehaviorSource.h"
 
-UViewportClickInteraction::UViewportClickInteraction()
+void UViewportClickInteraction::Initialize(UViewportInteractionsBehaviorSource* const InViewportInteractionsBehaviorSource)
 {
-	USingleClickInputBehavior* SingleClickBehavior = NewObject<USingleClickInputBehavior>();
-	SingleClickBehavior->Initialize(this);
+	Super::Initialize(InViewportInteractionsBehaviorSource);
+	
+	BuildBehaviors();
+}
 
-	SingleClickBehavior->SetDefaultPriority(UE::Editor::ViewportInteractions::CLICK_PRIORITY);
-	SingleClickBehavior->SetAllowsCaptureStealing(true);
-	RegisterInputBehavior(SingleClickBehavior);
+void UViewportClickInteraction::BuildBehaviors()
+{
+	UViewportClickBehavior* ClickBehavior = NewObject<UViewportClickBehavior>();
+	ClickBehavior->Initialize(this);
+	
+	ClickBehavior->SetAllowsCaptureStealing(true);
+	RegisterInputBehavior(ClickBehavior);
 
-	SingleClickInputBehaviorWeak = SingleClickBehavior;
+	ClickBehaviorWeak = ClickBehavior;
 }
 
 FInputRayHit UViewportClickInteraction::IsHitByClick(const FInputDeviceRay& InClickPos)
+{
+	return FInputRayHit(0.0f);
+}
+
+void UViewportClickInteraction::OnClickDown(const FInputDeviceRay& InClickPos)
 {
 	if (UViewportInteractionsBehaviorSource* const ViewportInteractionBehaviorSource = GetViewportInteractionsBehaviorSource())
 	{
-		if (!ViewportInteractionBehaviorSource->IsGizmoDragging())
-		{
-			ViewportInteractionBehaviorSource->SetMouseCursorOverride(EMouseCursor::Crosshairs);
-		}
+		ViewportInteractionBehaviorSource->SetMouseCursorOverride(EMouseCursor::Crosshairs);
 	}
-
-	return FInputRayHit(0.0f);
 }
 
-void UViewportClickInteraction::OnClicked(const FInputDeviceRay& InClickPos)
+void UViewportClickInteraction::OnClickUp(const FInputDeviceRay& InClickPos)
 {
 	if (UViewportInteractionsBehaviorSource* const ViewportInteractionBehaviorSource = GetViewportInteractionsBehaviorSource())
 	{
@@ -81,11 +87,19 @@ void UViewportClickInteraction::OnClicked(const FInputDeviceRay& InClickPos)
 			}
 		}
 	}
+	
+	if (UViewportInteractionsBehaviorSource* const ViewportInteractionBehaviorSource = GetViewportInteractionsBehaviorSource())
+	{
+		ViewportInteractionBehaviorSource->ClearMouseCursorOverride();
+	}
 }
 
-FDeviceButtonState UViewportClickInteraction::GetActiveMouseButtonState(const FInputDeviceState& Input)
+void UViewportClickInteraction::OnForceEndCapture()
 {
-	return Input.Mouse.Left;
+	if (UViewportInteractionsBehaviorSource* const ViewportInteractionBehaviorSource = GetViewportInteractionsBehaviorSource())
+	{
+		ViewportInteractionBehaviorSource->ClearMouseCursorOverride();
+	}
 }
 
 FKey UViewportClickInteraction::GetMouseKey() const
diff --git a/Engine/Source/Editor/Experimental/EditorInteractiveToolsFramework/Private/ViewportInteractions/ViewportDollyInteraction.cpp b/Engine/Source/Editor/Experimental/EditorInteractiveToolsFramework/Private/ViewportInteractions/ViewportDollyInteraction.cpp
--- a/Engine/Source/Editor/Experimental/EditorInteractiveToolsFramework/Private/ViewportInteractions/ViewportDollyInteraction.cpp
+++ b/Engine/Source/Editor/Experimental/EditorInteractiveToolsFramework/Private/ViewportInteractions/ViewportDollyInteraction.cpp
@@ -14,12 +14,12 @@ UViewportDollyInteraction::UViewportDollyInteraction()
 	
 	NormalBindings = FButtonBindings({
 		FButtonBinding(EKeys::RightMouseButton).TriggersStart(),
-		FButtonBinding(EKeys::LeftAlt).RequiredToContinue(false).Complexity(20)
+		FButtonBinding(EKeys::LeftAlt).RequiredToContinue(false)
 	});
 	
 	CameraLockedBindings = FButtonBindings({
 		FButtonBinding(EKeys::RightMouseButton).TriggersStart(),
-		FButtonBinding(EKeys::LeftAlt).Required(false).Complexity(20)
+		FButtonBinding(EKeys::LeftAlt).Required(false)
 	});
 	
 	ViewportClickDragBehavior->SetBindings(UViewportClickDragBehavior::FGetBindingsDelegate::CreateUObject(this, &UViewportDollyInteraction::GetBindings));
diff --git a/Engine/Source/Editor/Experimental/EditorInteractiveToolsFramework/Private/ViewportInteractions/ViewportOrbitInteraction.cpp b/Engine/Source/Editor/Experimental/EditorInteractiveToolsFramework/Private/ViewportInteractions/ViewportOrbitInteraction.cpp
--- a/Engine/Source/Editor/Experimental/EditorInteractiveToolsFramework/Private/ViewportInteractions/ViewportOrbitInteraction.cpp
+++ b/Engine/Source/Editor/Experimental/EditorInteractiveToolsFramework/Private/ViewportInteractions/ViewportOrbitInteraction.cpp
@@ -17,12 +17,12 @@ UViewportOrbitInteraction::UViewportOrbitInteraction()
 	
 	NormalBindings = FButtonBindings({
 		FButtonBinding(EKeys::LeftMouseButton).TriggersStart(),
-    	FButtonBinding(EKeys::LeftAlt).RequiredToContinue(false).Complexity(20)
+    	FButtonBinding(EKeys::LeftAlt).RequiredToContinue(false)
 	});
 	
 	CameraLockedBindings = FButtonBindings({
 		FButtonBinding(EKeys::LeftMouseButton).TriggersStart(),
-		FButtonBinding(EKeys::LeftAlt).Required(false).Complexity(20)
+		FButtonBinding(EKeys::LeftAlt).Required(false)
 	});
 	
 	ViewportClickDragBehavior->SetBindings(UViewportClickDragBehavior::FGetBindingsDelegate::CreateUObject(this, &UViewportOrbitInteraction::GetBindings));
diff --git a/Engine/Source/Editor/Experimental/EditorInteractiveToolsFramework/Private/ViewportInteractions/ViewportPanInteraction.cpp b/Engine/Source/Editor/Experimental/EditorInteractiveToolsFramework/Private/ViewportInteractions/ViewportPanInteraction.cpp
--- a/Engine/Source/Editor/Experimental/EditorInteractiveToolsFramework/Private/ViewportInteractions/ViewportPanInteraction.cpp
+++ b/Engine/Source/Editor/Experimental/EditorInteractiveToolsFramework/Private/ViewportInteractions/ViewportPanInteraction.cpp
@@ -20,7 +20,7 @@ UViewportPanInteraction::UViewportPanInteraction()
 	using namespace UE::Editor::ViewportInteractions;
 	ViewportClickDragBehavior->SetBindings({
 		FButtonBinding(EKeys::MiddleMouseButton),
-		FButtonBinding(EKeys::LeftAlt).Required(false).Complexity(20)
+		FButtonBinding(EKeys::LeftAlt).Required(false)
 	});
 }
 
diff --git a/Engine/Source/Editor/Experimental/EditorInteractiveToolsFramework/Public/Behaviors/ViewportClickBehavior.h b/Engine/Source/Editor/Experimental/EditorInteractiveToolsFramework/Public/Behaviors/ViewportClickBehavior.h
new file mode 100644
--- /dev/null
+++ b/Engine/Source/Editor/Experimental/EditorInteractiveToolsFramework/Public/Behaviors/ViewportClickBehavior.h
@@ -0,0 +1,85 @@
+﻿// Copyright Epic Games, Inc. All Rights Reserved.
+
+#pragma once
+
+#include "CoreMinimal.h"
+#include "InputBehavior.h"
+#include "ViewportInteractions/ViewportInteractionBindings.h"
+
+#include "ViewportClickBehavior.generated.h"
+
+#define UE_API EDITORINTERACTIVETOOLSFRAMEWORK_API
+
+class IViewportClickBehaviorTarget
+{
+public:
+	virtual ~IViewportClickBehaviorTarget() = default;
+
+	/**
+	 * Test if target can begin the click interaction at this point
+	 * @param PressPos device position/ray at click point
+	 * @return hit information at this point
+	 */
+	virtual FInputRayHit IsHitByClick(const FInputDeviceRay& PressPos) = 0;
+
+	/**
+	 * Notify the target that the click has started
+	 */
+	virtual void OnClickDown(const FInputDeviceRay& InClickPos) = 0;
+
+	/**
+	 * Notify the target that the click has completed
+	 */
+	virtual void OnClickUp(const FInputDeviceRay& InClickPos) = 0;
+
+	/**
+	 * Notify Target button state has changed
+	 */
+	virtual void OnStateUpdated(const FInputDeviceState& InInputDeviceState) = 0;
+
+	/**
+	 * Notify Target that capture is ending
+	 */
+	virtual void OnForceEndCapture()
+	{
+	}
+};
+
+/** A Click behavior designed to compete alongside UViewportClickDragBehaviors */
+UCLASS(MinimalAPI)
+class UViewportClickBehavior : public UInputBehavior
+{
+	GENERATED_BODY()
+	
+public:
+	UE_API virtual void Initialize(IViewportClickBehaviorTarget* InTarget);
+
+	UE_API virtual FInputCapturePriority GetPriority() override;
+	UE_API virtual FInputCaptureRequest WantsCapture(const FInputDeviceState& InputState) override;
+	UE_API virtual FInputCaptureUpdate BeginCapture(const FInputDeviceState& InputState, EInputCaptureSide eSide) override;
+	UE_API virtual FInputCaptureUpdate UpdateCapture(const FInputDeviceState& InputState, const FInputCaptureData& CaptureData) override;
+	
+	UE_API virtual void ForceEndCapture(const FInputCaptureData& CaptureData) override;
+	
+	UE_API void SetBindings(const TArray<UE::Editor::ViewportInteractions::FButtonBinding>& InBindings);
+	
+	/** @return current 3D world ray and optional 2D position for Target Device */
+	UE_API virtual FInputDeviceRay GetDeviceRay(const FInputDeviceState& InInputDeviceState);
+
+	UE_API virtual const FVector2D& GetMousePosition(const FInputDeviceState& InInputDeviceState) const;
+	
+protected:
+	UE_API void OnStateUpdatedInternal(const FInputDeviceState& InInputDeviceState);
+	UE_API void OnClickDownInternal(const FInputDeviceState& InputState);
+	UE_API void OnClickUpInternal(const FInputDeviceState& InputState);
+	
+	UE_API virtual TOptional<FInputCapturePriority> GetCapturePriority(UE::Editor::ViewportInteractions::EInputStage Stage, const FInputDeviceState& InputState) const;
+
+	IViewportClickBehaviorTarget* Target = nullptr;
+	UE::Editor::ViewportInteractions::FButtonBindings Bindings;
+	
+	TOptional<FInputCapturePriority> CachedCapturePriority;
+	bool bIsClicking = false;
+};
+
+#undef UE_API
\ No newline at end of file
diff --git a/Engine/Source/Editor/Experimental/EditorInteractiveToolsFramework/Public/Behaviors/ViewportClickDragBehavior.h b/Engine/Source/Editor/Experimental/EditorInteractiveToolsFramework/Public/Behaviors/ViewportClickDragBehavior.h
--- a/Engine/Source/Editor/Experimental/EditorInteractiveToolsFramework/Public/Behaviors/ViewportClickDragBehavior.h
+++ b/Engine/Source/Editor/Experimental/EditorInteractiveToolsFramework/Public/Behaviors/ViewportClickDragBehavior.h
@@ -130,11 +130,13 @@ protected:
 
 	virtual void OnStateUpdatedInternal(const FInputDeviceState& InInputDeviceState);
 	
+	bool HasMovedEnoughForDrag() const;
+	const UE::Editor::ViewportInteractions::FButtonBindings& GetBindings() const;
+	UE_API virtual TOptional<FInputCapturePriority> GetCapturePriority(UE::Editor::ViewportInteractions::EInputStage Stage, const FInputDeviceState& InputState) const;
+	
 	/** Distance after which a drag is considered such */
 	float DragDistanceThreshold = 2.0f;
 	
-	bool HasMovedEnoughForDrag() const;
-	
 	/** Cached cursor position */
 	FVector2D MouseDownPosition;
 
@@ -156,9 +158,8 @@ protected:
 
 	UE::Editor::ViewportInteractions::FButtonBindings Bindings;
 	FGetBindingsDelegate GetBindingsDelegate;
-	const UE::Editor::ViewportInteractions::FButtonBindings& GetBindings() const;
 	
-	int32 ActiveBindingComplexity = 0;
+	TOptional<FInputCapturePriority> CachedCapturePriority;
 
 	/** True when capturing and dragging */
 	bool bDragging = false;
diff --git a/Engine/Source/Editor/Experimental/EditorInteractiveToolsFramework/Public/ViewportInteractions/ViewportClickInteraction.h b/Engine/Source/Editor/Experimental/EditorInteractiveToolsFramework/Public/ViewportInteractions/ViewportClickInteraction.h
--- a/Engine/Source/Editor/Experimental/EditorInteractiveToolsFramework/Public/ViewportInteractions/ViewportClickInteraction.h
+++ b/Engine/Source/Editor/Experimental/EditorInteractiveToolsFramework/Public/ViewportInteractions/ViewportClickInteraction.h
@@ -4,6 +4,7 @@
 
 #include "ViewportInteraction.h"
 #include "BaseBehaviors/BehaviorTargetInterfaces.h"
+#include "Behaviors/ViewportClickBehavior.h"
 #include "ViewportClickInteraction.generated.h"
 
 #define UE_API EDITORINTERACTIVETOOLSFRAMEWORK_API
@@ -19,7 +20,6 @@ enum class EViewportClickButton : uint8
 } 
 
 class HHitProxy;
-class USingleClickInputBehavior;
 struct FViewportClick;
 
 /**
@@ -28,30 +28,34 @@ struct FViewportClick;
 UCLASS(MinimalAPI, Abstract)
 class UViewportClickInteraction
 	: public UViewportInteraction
-	, public IClickBehaviorTarget
+	, public IViewportClickBehaviorTarget
 {
 	GENERATED_BODY()
 
 public:
-	UE_API UViewportClickInteraction();
-
-	// ~Begin IClickBehaviorTarget
+	UE_API virtual void Initialize(UViewportInteractionsBehaviorSource* const InViewportInteractionsBehaviorSource) override;
+	
+	// ~Begin IViewportClickBehaviorTarget
 	UE_API virtual FInputRayHit IsHitByClick(const FInputDeviceRay& InClickPos) override;
-	UE_API virtual void OnClicked(const FInputDeviceRay& InClickPos) override;
-	// ~End IClickBehaviorTarget
+	UE_API virtual void OnClickDown(const FInputDeviceRay& InClickPos) override; 
+	UE_API virtual void OnClickUp(const FInputDeviceRay& InClickPos) override;
+	UE_API virtual void OnForceEndCapture() override;
+	
+	UE_API virtual void OnStateUpdated(const FInputDeviceState& InInputDeviceState) override {}
+	// ~End IViewportClickBehaviorTarget
 
 	FKey MouseButton = EKeys::LeftMouseButton;
 
 protected:
+	UE_API virtual void BuildBehaviors();
+
 	/** Override this function to provide custom functionality */
 	virtual void ProcessClick_Internal(const FViewportClick& InViewportClick, HHitProxy* InHitProxy, FSceneView& View) {}
 
-	UE_API virtual FDeviceButtonState GetActiveMouseButtonState(const FInputDeviceState& Input);
-
 	UE_API FKey GetMouseKey() const;
 	UE_API virtual EInputEvent GetMouseEvent() const;
 
-	TWeakObjectPtr<USingleClickInputBehavior> SingleClickInputBehaviorWeak;
+	TWeakObjectPtr<UViewportClickBehavior> ClickBehaviorWeak;
 };
 
 #undef UE_API
diff --git a/Engine/Source/Editor/Experimental/EditorInteractiveToolsFramework/Public/ViewportInteractions/ViewportInteraction.h b/Engine/Source/Editor/Experimental/EditorInteractiveToolsFramework/Public/ViewportInteractions/ViewportInteraction.h
--- a/Engine/Source/Editor/Experimental/EditorInteractiveToolsFramework/Public/ViewportInteractions/ViewportInteraction.h
+++ b/Engine/Source/Editor/Experimental/EditorInteractiveToolsFramework/Public/ViewportInteractions/ViewportInteraction.h
@@ -29,9 +29,32 @@ constexpr int RightMouseButtonMod = 5;
 constexpr int MiddleMouseButtonMod = 6;
 
 static constexpr int VIEWPORT_INTERACTIONS_DEFAULT_PRIORITY = FInputCapturePriority::DEFAULT_TOOL_PRIORITY - 5;
-static constexpr int DRAG_CONFIRMED_PRIORITY = FInputCapturePriority::DEFAULT_TOOL_PRIORITY + 5;
-static constexpr int CLICK_PRIORITY = FInputCapturePriority::DEFAULT_TOOL_PRIORITY + 10;
-static constexpr int DRAG_PRIORITY = FInputCapturePriority::DEFAULT_TOOL_PRIORITY + 15;
+
+/*
+ * Binding complexity is chosen so that "tool" behaviors win out over viewport interactions on the same complexity footing.
+ * Tool priority in this diagram is assumed to have a grandfathered single point of complexity.
+ *
+ * These statements should be true:
+ *		- A tool beats an unmodified click.
+ *		- A modified click can beat a tool.
+ *		- An unmodified drag loses to a click until it is confirmed.
+ *		- A modified, unconfirmed drag can beat an unmodified click.
+ *
+ *     |--   Tool
+ *      |--  Click
+ *   |--|--  Modified Click
+ *       |-- Drag (unconfirmed)
+ *    |--|-- Drag confirmed (or modified unconfirmed drag)
+ * |--|--|-- Modified drag confirmed
+ *
+ * BINDING_COMPLEXITY_PRIORITY_STEP can be increased to give more wiggle room.
+ * Making it cleanly divisible by 3 gives wiggle room within configurations.  
+ */
+static constexpr int BINDING_COMPLEXITY_PRIORITY_THIRD_STEP = 3;
+static constexpr int BINDING_COMPLEXITY_PRIORITY_STEP = BINDING_COMPLEXITY_PRIORITY_THIRD_STEP * 3;
+static constexpr int CLICK_PRIORITY = FInputCapturePriority::DEFAULT_TOOL_PRIORITY + BINDING_COMPLEXITY_PRIORITY_STEP + BINDING_COMPLEXITY_PRIORITY_THIRD_STEP;
+static constexpr int DRAG_PRIORITY = FInputCapturePriority::DEFAULT_TOOL_PRIORITY + BINDING_COMPLEXITY_PRIORITY_STEP + BINDING_COMPLEXITY_PRIORITY_THIRD_STEP * 2;
+static constexpr int DRAG_CONFIRMED_PRIORITY = DRAG_PRIORITY - BINDING_COMPLEXITY_PRIORITY_STEP;
 
 /**
  * Possible types of Viewport Interactions.
diff --git a/Engine/Source/Editor/Experimental/EditorInteractiveToolsFramework/Public/ViewportInteractions/ViewportInteractionBindings.h b/Engine/Source/Editor/Experimental/EditorInteractiveToolsFramework/Public/ViewportInteractions/ViewportInteractionBindings.h
--- a/Engine/Source/Editor/Experimental/EditorInteractiveToolsFramework/Public/ViewportInteractions/ViewportInteractionBindings.h
+++ b/Engine/Source/Editor/Experimental/EditorInteractiveToolsFramework/Public/ViewportInteractions/ViewportInteractionBindings.h
@@ -34,13 +34,6 @@ struct FButtonBinding
 	FInputChord Chord;
 	TWeakPtr<FUICommandInfo> Command;
 	
-	/**
-	 * The complexity value of each activated button.
-	 * This feature is present as an interim patch for the purposes of bespoke priority hacks.
-	 * Modifying complexity like this is high-odor system smell. This WILL go away. 
-	 */
-	int32 ButtonComplexity = 1;
-	
 	// TODO: These could be a flag-based enum, or just bitwise values
 	/** This button must be down for the behavior begin capture. */
 	bool bRequiredToStart = true;
@@ -73,12 +66,6 @@ struct FButtonBinding
 		bRequiredToContinue = bRequired;
 		return *this;
 	}
-	
-	FButtonBinding& Complexity(int32 InComplexityMultiplier)
-	{
-		ButtonComplexity = InComplexityMultiplier;
-		return *this;
-	}
 };
 
 
diff --git a/Engine/Source/Editor/LevelEditor/Private/ViewportInteractions/LevelViewportClickSelection.cpp b/Engine/Source/Editor/LevelEditor/Private/ViewportInteractions/LevelViewportClickSelection.cpp
--- a/Engine/Source/Editor/LevelEditor/Private/ViewportInteractions/LevelViewportClickSelection.cpp
+++ b/Engine/Source/Editor/LevelEditor/Private/ViewportInteractions/LevelViewportClickSelection.cpp
@@ -95,7 +95,25 @@ void ProcessClick(FEditorViewportClient* InEditorViewportClient, const FViewport
 ULevelViewportClickSelection::ULevelViewportClickSelection()
 {
 	InteractionName = TEXT("Level Viewport Click Selection");
-	Groups = { UE::Editor::ViewportInteractions::ViewportClick };
+    Groups = { UE::Editor::ViewportInteractions::ViewportClick };
+}
+
+void ULevelViewportClickSelection::BuildBehaviors()
+{
+	Super::BuildBehaviors();
+
+	using namespace UE::Editor::ViewportInteractions;
+	
+	if (UViewportClickBehavior* ClickBehavior = ClickBehaviorWeak.Get())
+	{
+		ClickBehavior->SetBindings({
+			EKeys::LeftMouseButton,
+			FButtonBinding(EKeys::LeftControl).Required(false),
+			FButtonBinding(EKeys::LeftCommand).Required(false),
+			FButtonBinding(EKeys::LeftAlt).Required(false),
+			FButtonBinding(EKeys::LeftShift).Required(false)
+		});
+	}
 }
 
 void ULevelViewportClickSelection::ProcessClick_Internal(const FViewportClick& InViewportClick, HHitProxy* InHitProxy, FSceneView& View)
@@ -110,27 +128,44 @@ ULevelViewportClickContextMenu::ULevelViewportClickContextMenu()
 	Groups = { UE::Editor::ViewportInteractions::ViewportClick };
 
 	MouseButton = EKeys::RightMouseButton;
+}
 
-	if (TStrongObjectPtr<USingleClickInputBehavior> SingleClickInputBehavior = SingleClickInputBehaviorWeak.Pin())
+void ULevelViewportClickContextMenu::BuildBehaviors()
+{
+	Super::BuildBehaviors();
+	
+	using namespace UE::Editor::ViewportInteractions;
+	
+	if (UViewportClickBehavior* ClickBehavior = ClickBehaviorWeak.Get())
 	{
-		SingleClickInputBehavior->SetUseRightMouseButton();
+		ClickBehavior->SetBindings({
+			EKeys::RightMouseButton,
+			FButtonBinding(EKeys::LeftControl).Required(false),
+			FButtonBinding(EKeys::LeftCommand).Required(false),
+			FButtonBinding(EKeys::LeftAlt).Required(false),
+			FButtonBinding(EKeys::LeftShift).Required(false)
+		});
 	}
 }
 
 ULevelViewportSetPivot::ULevelViewportSetPivot()
 {
-	InteractionName = TEXT("Level Viewport Middle Click");
+	InteractionName = TEXT("Level Viewport Set Pivot");
 	Groups = { UE::Editor::ViewportInteractions::ViewportClick };
 
 	MouseButton = EKeys::MiddleMouseButton;
+}
+
+void ULevelViewportSetPivot::BuildBehaviors()
+{
+	Super::BuildBehaviors();
 
-	if (TStrongObjectPtr<USingleClickInputBehavior> SingleClickInputBehavior = SingleClickInputBehaviorWeak.Pin())
+	if (UViewportClickBehavior* ClickBehavior = ClickBehaviorWeak.Get())
 	{
-		SingleClickInputBehavior->SetUseMiddleMouseButton();
-		SingleClickInputBehavior->ModifierCheckFunc = [](const FInputDeviceState& InputDeviceState)
-		{
-			return InputDeviceState.bAltKeyDown;
-		};
+		ClickBehavior->SetBindings({
+			EKeys::MiddleMouseButton,
+			EKeys::LeftAlt
+		});
 	}
 }
 
@@ -145,14 +180,28 @@ void ULevelViewportSetPivot::ProcessClick_Internal(const FViewportClick& InViewp
 	}
 }
 
-ULevelViewportDoubleClick::ULevelViewportDoubleClick()
+TOptional<FInputCapturePriority> ULevelViewportDoubleClickBehavior::GetCapturePriority(UE::Editor::ViewportInteractions::EInputStage Stage, const FInputDeviceState& InputState) const
 {
-	UDoubleClickInputBehavior* DoubleClickInputBehavior = NewObject<UDoubleClickInputBehavior>();
-	DoubleClickInputBehavior->Initialize(this);
-	DoubleClickInputBehavior->SetDefaultPriority(FInputCapturePriority::DEFAULT_GIZMO_PRIORITY - 1);
+	if (InputState.Mouse.Left.bDoubleClicked)
+	{
+		return FInputCapturePriority(FInputCapturePriority::DEFAULT_GIZMO_PRIORITY - 1);
+	}
+	return TOptional<FInputCapturePriority>();
+}
 
+ULevelViewportDoubleClick::ULevelViewportDoubleClick()
+{
 	InteractionName = TEXT("Level Viewport Double Click");
-	Groups = { UE::Editor::ViewportInteractions::ViewportClick };
+    Groups = { UE::Editor::ViewportInteractions::ViewportClick };
+}
+
+void ULevelViewportDoubleClick::BuildBehaviors()
+{
+	ULevelViewportDoubleClickBehavior* ClickBehavior = NewObject<ULevelViewportDoubleClickBehavior>();
+	ClickBehavior->Initialize(this);
+	
+	ClickBehavior->SetAllowsCaptureStealing(true);
+	RegisterInputBehavior(ClickBehavior);
 
-	RegisterInputBehavior(DoubleClickInputBehavior);
+	ClickBehaviorWeak = ClickBehavior;
 }
diff --git a/Engine/Source/Editor/LevelEditor/Private/ViewportInteractions/LevelViewportClickSelection.h b/Engine/Source/Editor/LevelEditor/Private/ViewportInteractions/LevelViewportClickSelection.h
--- a/Engine/Source/Editor/LevelEditor/Private/ViewportInteractions/LevelViewportClickSelection.h
+++ b/Engine/Source/Editor/LevelEditor/Private/ViewportInteractions/LevelViewportClickSelection.h
@@ -20,6 +20,7 @@ public:
 	UE_API ULevelViewportClickSelection();
 
 protected:
+	UE_API virtual void BuildBehaviors() override;
 	UE_API virtual void ProcessClick_Internal(const FViewportClick& InViewportClick, HHitProxy* InHitProxy, FSceneView& View) override;
 };
 
@@ -33,6 +34,9 @@ class ULevelViewportClickContextMenu : public ULevelViewportClickSelection
 
 public:
 	UE_API ULevelViewportClickContextMenu();
+	
+protected:
+	UE_API virtual void BuildBehaviors() override;
 };
 
 /**
@@ -47,9 +51,20 @@ public:
 	UE_API ULevelViewportSetPivot();
 	
 protected:
+	UE_API virtual void BuildBehaviors() override;
+	
 	UE_API virtual void ProcessClick_Internal(const FViewportClick& InViewportClick, HHitProxy* InHitProxy, FSceneView& View) override;
 };
 
+UCLASS(MinimalAPI)
+class ULevelViewportDoubleClickBehavior : public UViewportClickBehavior
+{
+	GENERATED_BODY()
+
+protected:
+	UE_API virtual TOptional<FInputCapturePriority> GetCapturePriority(UE::Editor::ViewportInteractions::EInputStage Stage, const FInputDeviceState& InputState) const;
+};
+
 /**
  * Left-mouse double click
  */
@@ -65,6 +80,9 @@ public:
 	{
 		return IE_DoubleClick;
 	}
+	
+protected:
+	UE_API virtual void BuildBehaviors() override; 
 };
 
 #undef UE_API
diff --git a/Engine/Source/Editor/Persona/Private/ViewportInteractions/SkeletonSelectionClickInteraction.cpp b/Engine/Source/Editor/Persona/Private/ViewportInteractions/SkeletonSelectionClickInteraction.cpp
--- a/Engine/Source/Editor/Persona/Private/ViewportInteractions/SkeletonSelectionClickInteraction.cpp
+++ b/Engine/Source/Editor/Persona/Private/ViewportInteractions/SkeletonSelectionClickInteraction.cpp
@@ -15,6 +15,21 @@ void USkeletonSelectionClickInteraction::Bind(TSharedPtr<FSkeletonSelectionEditM
 	EditMode = InEditMode;
 }
 
+void USkeletonSelectionClickInteraction::BuildBehaviors()
+{
+	Super::BuildBehaviors();
+	
+	using namespace UE::Editor::ViewportInteractions;
+	
+	ClickBehaviorWeak->SetBindings({
+		EKeys::LeftMouseButton,
+		FButtonBinding(EKeys::LeftControl).Required(false),
+		FButtonBinding(EKeys::LeftCommand).Required(false),
+		FButtonBinding(EKeys::LeftAlt).Required(false),
+		FButtonBinding(EKeys::LeftShift).Required(false)
+	});
+}
+
 FInputRayHit USkeletonSelectionClickInteraction::IsHitByClick(const FInputDeviceRay& ClickPos)
 {
 	if (IsEnabled())
@@ -36,6 +51,19 @@ USkeletonSelectionRightClickInteraction::USkeletonSelectionRightClickInteraction
 	Groups = { UE::Editor::ViewportInteractions::ViewportClick };
 
 	MouseButton = EKeys::RightMouseButton;
+}
+
+void USkeletonSelectionRightClickInteraction::BuildBehaviors()
+{
+	Super::BuildBehaviors();
+	
+	using namespace UE::Editor::ViewportInteractions;
 	
-	SingleClickInputBehaviorWeak->SetUseRightMouseButton();
+	ClickBehaviorWeak->SetBindings({
+		EKeys::RightMouseButton,
+		FButtonBinding(EKeys::LeftControl).Required(false),
+		FButtonBinding(EKeys::LeftCommand).Required(false),
+		FButtonBinding(EKeys::LeftAlt).Required(false),
+		FButtonBinding(EKeys::LeftShift).Required(false)
+	});
 }
diff --git a/Engine/Source/Editor/Persona/Public/ViewportInteractions/SkeletonSelectionClickInteraction.h b/Engine/Source/Editor/Persona/Public/ViewportInteractions/SkeletonSelectionClickInteraction.h
--- a/Engine/Source/Editor/Persona/Public/ViewportInteractions/SkeletonSelectionClickInteraction.h
+++ b/Engine/Source/Editor/Persona/Public/ViewportInteractions/SkeletonSelectionClickInteraction.h
@@ -20,7 +20,9 @@ public:
 	virtual FInputRayHit IsHitByClick(const FInputDeviceRay& ClickPos) override;
 
 protected:
+	virtual void BuildBehaviors() override;
 	virtual void ProcessClick_Internal(const FViewportClick& InViewportClick, HHitProxy* InHitProxy, FSceneView& View) override;
+	
 
 	TWeakPtr<FSkeletonSelectionEditMode> EditMode;
 };
@@ -35,4 +37,6 @@ class USkeletonSelectionRightClickInteraction
 public:
 	USkeletonSelectionRightClickInteraction();
 	
+protected:
+	virtual void BuildBehaviors() override;
 };
