[ITF Gizmos] Make turning on & off the ITF gizmo stack based on a legacy editor mode requiring the legacy systems only affect its editor mode rather than all editor modes.

#rnx

#Synced-CL 51158691

diff --git a/Engine/Source/Editor/Experimental/EditorInteractiveToolsFramework/Private/EditorGizmos/EditorTransformGizmoUtil.cpp b/Engine/Source/Editor/Experimental/EditorInteractiveToolsFramework/Private/EditorGizmos/EditorTransformGizmoUtil.cpp
--- a/Engine/Source/Editor/Experimental/EditorInteractiveToolsFramework/Private/EditorGizmos/EditorTransformGizmoUtil.cpp
+++ b/Engine/Source/Editor/Experimental/EditorInteractiveToolsFramework/Private/EditorGizmos/EditorTransformGizmoUtil.cpp
@@ -291,32 +291,49 @@ void UEditorTransformGizmoContextObject::UpdateGizmo(const TArray<FEditorViewpor
 	}
 }
 
+void UEditorTransformGizmoContextObject::SetGizmosEnabled(bool bInEnabled)
+{
+	if (!bInEnabled)
+	{
+		// swap back default mode
+		(void)SwapDefaultMode(FAssetEdModes::EM_AssetDefault, FBuiltinEditorModes::EM_Default);
+
+		// reset the coordinate system if the legacy gizmo doesn't support it
+		if (ModeTools->GetCoordSystem() > COORD_Local)
+		{
+			ModeTools->SetCoordSystem(COORD_Local);
+		}
+			
+		// remove viewports' binding + gizmos as they are useless
+		const UModeManagerInteractiveToolsContext* ToolsContext = ModeTools->GetInteractiveToolsContext();
+		const TObjectPtr<UInteractiveToolManager> ToolManager = ToolsContext->ToolManager;
+		UE::EditorTransformGizmoUtil::RemoveDefaultTransformGizmo(ToolManager);
+			
+		RemoveViewportsBinding();
+			
+		return;
+	}
+
+	InitializeViewportsBinding();
+}
+
 void UEditorTransformGizmoContextObject::InitializeGizmoManagerBinding()
 {
 	auto OnGizmoVariableChanged = [this](const bool bUseNewTRSGizmos)
 	{
-		if (!bUseNewTRSGizmos)
+		bool bEnableGizmos = bUseNewTRSGizmos;
+	
+		if (UModeManagerInteractiveToolsContext* ToolsContext = ModeTools->GetInteractiveToolsContext())
 		{
-			// swap back default mode
-			(void)SwapDefaultMode(FAssetEdModes::EM_AssetDefault, FBuiltinEditorModes::EM_Default);
-
-			// reset the coordinate system if the legacy gizmo doesn't support it
-			if (ModeTools->GetCoordSystem() > COORD_Local)
+			if (UEditorInteractiveGizmoManager* GizmoManager = Cast<UEditorInteractiveGizmoManager>(ToolsContext->GizmoManager))
 			{
-				ModeTools->SetCoordSystem(COORD_Local);
+				// When the full context of the manager is available, rely on its state.
+				// The setting to allow `NewTRSGizmo` may be on while a particular mode manager has disabled them.
+				bEnableGizmos = GizmoManager->AreEditorGizmosAllowed();
 			}
-			
-			// remove viewports' binding + gizmos as they are useless
-			const UModeManagerInteractiveToolsContext* ToolsContext = ModeTools->GetInteractiveToolsContext();
-			const TObjectPtr<UInteractiveToolManager> ToolManager = ToolsContext->ToolManager;
-			UE::EditorTransformGizmoUtil::RemoveDefaultTransformGizmo(ToolManager);
-			
-			RemoveViewportsBinding();
-			
-			return;
 		}
-
-		InitializeViewportsBinding();
+		
+		SetGizmosEnabled(bEnableGizmos);
 	};
 
 	// bind UseNewGizmosChangedHandled if needed
diff --git a/Engine/Source/Editor/Experimental/EditorInteractiveToolsFramework/Private/EditorInteractiveGizmoManager.cpp b/Engine/Source/Editor/Experimental/EditorInteractiveToolsFramework/Private/EditorInteractiveGizmoManager.cpp
--- a/Engine/Source/Editor/Experimental/EditorInteractiveToolsFramework/Private/EditorInteractiveGizmoManager.cpp
+++ b/Engine/Source/Editor/Experimental/EditorInteractiveToolsFramework/Private/EditorInteractiveGizmoManager.cpp
@@ -3,6 +3,7 @@
 
 #include "EditorInteractiveGizmoManager.h"
 
+#include "ContextObjectStore.h"
 #include "Editor.h"
 #include "Editor/EditorEngine.h"
 #include "EditorInteractiveGizmoConditionalBuilder.h"
@@ -23,6 +24,7 @@
 #include "Templates/Casts.h"
 #include "ToolContextInterfaces.h"
 #include "TransformGizmoEditorSettings.h"
+#include "EditorGizmos/EditorTransformGizmoUtil.h"
 #include "EditorGizmos/TransformGizmo.h"
 
 #include UE_INLINE_GENERATED_CPP_BY_NAME(EditorInteractiveGizmoManager)
@@ -275,17 +277,23 @@ bool UEditorInteractiveGizmoManager::GetShowEditorGizmosForView(IToolsContextRen
 	return bShowEditorGizmos && bEngineShowFlagsModeWidget;
 }
 
+bool UEditorInteractiveGizmoManager::AreEditorGizmosAllowed() const
+{
+	return bAllowEditorGizmos && UsesNewTRSGizmos();
+}
+
 void UEditorInteractiveGizmoManager::SetAllowEditorGizmos(bool bInAllowed)
 {
 	if (bAllowEditorGizmos != bInAllowed)
 	{
 		bAllowEditorGizmos = bInAllowed;
 		
-		OnUsesNewTRSGizmosChangedDelegate().Broadcast(bAllowEditorGizmos && UsesNewTRSGizmos());
-		
-		if (bAllowEditorGizmos && UsesNewTRSGizmos())
+		if (UContextObjectStore* Store = GetContextObjectStore())
 		{
-			
+			if (UEditorTransformGizmoContextObject* GizmoContext = Store->FindContext<UEditorTransformGizmoContextObject>())
+			{
+				GizmoContext->SetGizmosEnabled(AreEditorGizmosAllowed());
+			}
 		}
 		
 		UpdateActiveEditorGizmos();
@@ -294,7 +302,7 @@ void UEditorInteractiveGizmoManager::SetAllowEditorGizmos(bool bInAllowed)
 
 void UEditorInteractiveGizmoManager::UpdateActiveEditorGizmos()
 {
-	if (!bAllowEditorGizmos || !UsesNewTRSGizmos())
+	if (!AreEditorGizmosAllowed())
 	{
 		if (UTransformGizmo* Gizmo = FindDefaultTransformGizmo())
 		{
diff --git a/Engine/Source/Editor/Experimental/EditorInteractiveToolsFramework/Public/EditorGizmos/EditorTransformGizmoUtil.h b/Engine/Source/Editor/Experimental/EditorInteractiveToolsFramework/Public/EditorGizmos/EditorTransformGizmoUtil.h
--- a/Engine/Source/Editor/Experimental/EditorInteractiveToolsFramework/Public/EditorGizmos/EditorTransformGizmoUtil.h
+++ b/Engine/Source/Editor/Experimental/EditorInteractiveToolsFramework/Public/EditorGizmos/EditorTransformGizmoUtil.h
@@ -83,6 +83,8 @@ public:
 
 	EDITORINTERACTIVETOOLSFRAMEWORK_API void Initialize(FEditorModeTools* InModeTools);
 	EDITORINTERACTIVETOOLSFRAMEWORK_API void Shutdown();
+	
+	EDITORINTERACTIVETOOLSFRAMEWORK_API void SetGizmosEnabled(bool bInEnabled);
 
 	EDITORINTERACTIVETOOLSFRAMEWORK_API UTransformGizmo* CreateTransformGizmo(
 		UEditorInteractiveGizmoManager* InGizmoManager, const FString& InInstanceIdentifier, void* InOwner) const;
diff --git a/Engine/Source/Editor/Experimental/EditorInteractiveToolsFramework/Public/EditorInteractiveGizmoManager.h b/Engine/Source/Editor/Experimental/EditorInteractiveToolsFramework/Public/EditorInteractiveGizmoManager.h
--- a/Engine/Source/Editor/Experimental/EditorInteractiveToolsFramework/Public/EditorInteractiveGizmoManager.h
+++ b/Engine/Source/Editor/Experimental/EditorInteractiveToolsFramework/Public/EditorInteractiveGizmoManager.h
@@ -113,6 +113,11 @@ public:
 		return bSearchLocalBuildersOnly;
 	}
 
+	/**
+	 * Gets whether editor gizmos are allowed
+	 */
+	UE_API bool AreEditorGizmosAllowed() const; 
+
 	/**
 	 * Sets whether editor gizmos are allowed
 	 */
