[ITF Viewport Interactions]
- Provide a default set of viewport interactions in FAssetEditorToolkit::OnBuildViewportInteractions().
- Refactor ViewportClickInteraction to not require manual configuration of what button or event is tracked.
- Add UViewportLegacyClickInteraction that passes click ITF clicks to the mode manager to match the same functionality as the legacy EditorViewportClient input path.
- Remove SkeletonSelectionClickInteraction and SKMModelingToolsClickSelection as it is no longer necessary.

#rnx

#Synced-CL 51211095

diff --git a/Engine/Plugins/Animation/SkeletalMeshModelingTools/Source/SkeletalMeshModelingTools/Private/SkeletalMeshModelingToolsEditorMode.cpp b/Engine/Plugins/Animation/SkeletalMeshModelingTools/Source/SkeletalMeshModelingTools/Private/SkeletalMeshModelingToolsEditorMode.cpp
--- a/Engine/Plugins/Animation/SkeletalMeshModelingTools/Source/SkeletalMeshModelingTools/Private/SkeletalMeshModelingToolsEditorMode.cpp
+++ b/Engine/Plugins/Animation/SkeletalMeshModelingTools/Source/SkeletalMeshModelingTools/Private/SkeletalMeshModelingToolsEditorMode.cpp
@@ -16,7 +16,6 @@
 
 #include "SkeletalMeshModelingToolsModule.h"
 #include "Selection/SKMGeometrySelectionTransformTweaker.h"
-#include "ViewportInteractions/SKMModelingToolsClickSelection.h"
 #include "ViewportInteractions/SKMModelingToolsGeometryClickSelection.h"
 #include "ViewportInteractions/SKMModelingToolsGeometryFrustumSelection.h"
 #include "ViewportInteractions/SKMModelingToolsGeometryHoverInteraction.h"
@@ -1075,10 +1074,6 @@ void USkeletalMeshModelingToolsEditorMode::InitializeGeometrySelectionSystems()
 
 void USkeletalMeshModelingToolsEditorMode::OnBuildViewportInteractions(UViewportInteractionsBehaviorSource* Source)
 {
-	USkeletalMeshModelingToolsClickSelection* BaseClickSelection = 
-		Source->AddInteraction<USkeletalMeshModelingToolsClickSelection>();
-	BaseClickSelection->Bind(this);
-
 	USkeletalMeshModelingToolsGeometryFrustumSelection* GeometryFrustumSelection =
 		Source->AddInteraction<USkeletalMeshModelingToolsGeometryFrustumSelection>();
 	GeometryFrustumSelection->BindSelectionManager(SelectionManager);
@@ -1123,8 +1118,7 @@ void USkeletalMeshModelingToolsEditorMode::ToggleGeometrySelectionViewportIntera
 		BehaviorSource->FindInteraction(USkeletalMeshModelingToolsGeometryClickSelection::StaticClass())->SetEnabled(bEnable);
 		BehaviorSource->FindInteraction(USkeletalMeshModelingToolsGeometryHoverInteraction::StaticClass())->SetEnabled(bEnable);
 		BehaviorSource->FindInteraction(USkeletalMeshModelingToolsGeometrySoftSelectRadiusDragInteraction::StaticClass())->SetEnabled(bEnable);
-				
-		BehaviorSource->FindInteraction(USkeletalMeshModelingToolsClickSelection::StaticClass())->SetEnabled(!bEnable);
+		
 		BehaviorSource->FindInteraction(UViewportMoveYawInteraction::StaticClass())->SetEnabled(!bEnable);	
 	}	
 }
diff --git a/Engine/Plugins/Animation/SkeletalMeshModelingTools/Source/SkeletalMeshModelingTools/Private/ViewportInteractions/SKMModelingToolsClickSelection.cpp b/Engine/Plugins/Animation/SkeletalMeshModelingTools/Source/SkeletalMeshModelingTools/Private/ViewportInteractions/SKMModelingToolsClickSelection.cpp
deleted file mode 100644
diff --git a/Engine/Plugins/Animation/SkeletalMeshModelingTools/Source/SkeletalMeshModelingTools/Private/ViewportInteractions/SKMModelingToolsClickSelection.h b/Engine/Plugins/Animation/SkeletalMeshModelingTools/Source/SkeletalMeshModelingTools/Private/ViewportInteractions/SKMModelingToolsClickSelection.h
deleted file mode 100644
diff --git a/Engine/Source/Editor/Experimental/EditorInteractiveToolsFramework/Private/ViewportInteractions/ViewportClickInteraction.cpp b/Engine/Source/Editor/Experimental/EditorInteractiveToolsFramework/Private/ViewportInteractions/ViewportClickInteraction.cpp
--- a/Engine/Source/Editor/Experimental/EditorInteractiveToolsFramework/Private/ViewportInteractions/ViewportClickInteraction.cpp
+++ b/Engine/Source/Editor/Experimental/EditorInteractiveToolsFramework/Private/ViewportInteractions/ViewportClickInteraction.cpp
@@ -28,7 +28,11 @@ void UViewportClickInteraction::BuildBehaviors()
 
 FInputRayHit UViewportClickInteraction::IsHitByClick(const FInputDeviceRay& InClickPos)
 {
-	return FInputRayHit(0.0f);
+	if (IsEnabled())
+	{
+		return FInputRayHit(0.0f);
+	}
+	return FInputRayHit();
 }
 
 void UViewportClickInteraction::OnClickDown(const FInputDeviceRay& InClickPos)
@@ -39,6 +43,34 @@ void UViewportClickInteraction::OnClickDown(const FInputDeviceRay& InClickPos)
 	}
 }
 
+void UViewportClickInteraction::OnStateUpdated(const FInputDeviceState& InInputDeviceState)
+{
+	auto CacheChange = [this](const FDeviceButtonState& ButtonState) -> bool
+	{
+		if (ButtonState.bDoubleClicked)
+		{
+			LastChangedMouseButton = ButtonState.Button;
+			LastInputEvent = IE_DoubleClick;
+			return true;
+		}
+		if (ButtonState.bPressed)
+		{
+			LastChangedMouseButton = ButtonState.Button;
+			LastInputEvent = IE_Pressed;
+			return true;
+		}
+		if (ButtonState.bReleased)
+		{
+			LastChangedMouseButton = ButtonState.Button;
+			LastInputEvent = IE_Released;
+			return true;
+		}
+		return false;
+	};
+
+	CacheChange(InInputDeviceState.Mouse.Left) || CacheChange(InInputDeviceState.Mouse.Right) || CacheChange(InInputDeviceState.Mouse.Middle);
+}
+
 void UViewportClickInteraction::OnClickUp(const FInputDeviceRay& InClickPos)
 {
 	if (UViewportInteractionsBehaviorSource* const ViewportInteractionBehaviorSource = GetViewportInteractionsBehaviorSource())
@@ -80,7 +112,7 @@ void UViewportClickInteraction::OnClickUp(const FInputDeviceRay& InClickPos)
 			if (FSceneView* View = EditorViewportClient->CalcSceneView(&ViewFamily))
 			{
 				const FViewportClick ViewportClick(
-					View, EditorViewportClient, GetMouseKey(), GetMouseEvent(), CurrentClickPosX, CurrentClickPosY
+					View, EditorViewportClient, LastChangedMouseButton, LastInputEvent, CurrentClickPosX, CurrentClickPosY
 				);
 
 				ProcessClick_Internal(ViewportClick, HitProxy, *View);
@@ -101,13 +133,3 @@ void UViewportClickInteraction::OnForceEndCapture()
 		ViewportInteractionBehaviorSource->ClearMouseCursorOverride();
 	}
 }
-
-FKey UViewportClickInteraction::GetMouseKey() const
-{
-	return MouseButton;
-}
-
-EInputEvent UViewportClickInteraction::GetMouseEvent() const
-{
-	return IE_Released;
-}
diff --git a/Engine/Source/Editor/Experimental/EditorInteractiveToolsFramework/Private/ViewportInteractions/ViewportLegacyClickInteraction.cpp b/Engine/Source/Editor/Experimental/EditorInteractiveToolsFramework/Private/ViewportInteractions/ViewportLegacyClickInteraction.cpp
new file mode 100644
--- /dev/null
+++ b/Engine/Source/Editor/Experimental/EditorInteractiveToolsFramework/Private/ViewportInteractions/ViewportLegacyClickInteraction.cpp
@@ -0,0 +1,47 @@
+﻿// Copyright Epic Games, Inc. All Rights Reserved.
+
+#include "ViewportInteractions/ViewportLegacyClickInteraction.h"
+
+#include "EditorModeManager.h"
+
+UViewportLegacyClickInteraction::UViewportLegacyClickInteraction()
+{
+	InteractionName = TEXT("Legacy Viewport Click");
+	Groups = { UE::Editor::ViewportInteractions::ViewportClick };
+}
+
+void UViewportLegacyClickInteraction::BuildBehaviors()
+{
+	Super::BuildBehaviors();
+	
+	if (UViewportClickBehavior* ClickBehavior = ClickBehaviorWeak.Get())
+	{
+		using namespace UE::Editor::ViewportInteractions;
+	
+		// Lower priority by one step so that custom UViewportClickInteractions do not have to compete.
+		ClickBehavior->SetDefaultPriority(FInputCapturePriority(CLICK_PRIORITY).MakeLower());
+		
+		ClickBehavior->SetBindings({
+			FButtonBinding(EKeys::LeftMouseButton).Required(false).TriggersStart(true),
+			FButtonBinding(EKeys::RightMouseButton).Required(false).TriggersStart(true),
+			FButtonBinding(EKeys::MiddleMouseButton).Required(false).TriggersStart(true),
+			
+			FButtonBinding(EKeys::LeftControl).Required(false),
+			FButtonBinding(EKeys::LeftCommand).Required(false),
+			FButtonBinding(EKeys::LeftAlt).Required(false),
+			FButtonBinding(EKeys::LeftShift).Required(false)
+		});
+	}
+}
+
+void UViewportLegacyClickInteraction::ProcessClick_Internal(const FViewportClick& InViewportClick, HHitProxy* InHitProxy, FSceneView& View)
+{
+	if (FEditorModeTools* ModeTools = GetModeTools())
+	{
+		if (FEditorViewportClient* ViewportClient = ModeTools->GetFocusedViewportClient())
+		{
+			ModeTools->HandleClick(ViewportClient, InHitProxy, InViewportClick);
+		}
+	}
+}
+
diff --git a/Engine/Source/Editor/Experimental/EditorInteractiveToolsFramework/Public/ViewportInteractions/ViewportClickInteraction.h b/Engine/Source/Editor/Experimental/EditorInteractiveToolsFramework/Public/ViewportInteractions/ViewportClickInteraction.h
--- a/Engine/Source/Editor/Experimental/EditorInteractiveToolsFramework/Public/ViewportInteractions/ViewportClickInteraction.h
+++ b/Engine/Source/Editor/Experimental/EditorInteractiveToolsFramework/Public/ViewportInteractions/ViewportClickInteraction.h
@@ -23,7 +23,7 @@ class HHitProxy;
 struct FViewportClick;
 
 /**
- *
+ * Provides click-on-up interactions with binding-complexity-based priority compatible with FViewportClick.
  */
 UCLASS(MinimalAPI, Abstract)
 class UViewportClickInteraction
@@ -41,21 +41,19 @@ public:
 	UE_API virtual void OnClickUp(const FInputDeviceRay& InClickPos) override;
 	UE_API virtual void OnForceEndCapture() override;
 	
-	UE_API virtual void OnStateUpdated(const FInputDeviceState& InInputDeviceState) override {}
+	UE_API virtual void OnStateUpdated(const FInputDeviceState& InInputDeviceState) override;
 	// ~End IViewportClickBehaviorTarget
 
-	FKey MouseButton = EKeys::LeftMouseButton;
-
 protected:
 	UE_API virtual void BuildBehaviors();
 
 	/** Override this function to provide custom functionality */
 	virtual void ProcessClick_Internal(const FViewportClick& InViewportClick, HHitProxy* InHitProxy, FSceneView& View) {}
 
-	UE_API FKey GetMouseKey() const;
-	UE_API virtual EInputEvent GetMouseEvent() const;
-
 	TWeakObjectPtr<UViewportClickBehavior> ClickBehaviorWeak;
+	
+	FKey LastChangedMouseButton = FKey();
+	EInputEvent LastInputEvent = EInputEvent::IE_MAX;
 };
 
 #undef UE_API
diff --git a/Engine/Source/Editor/Experimental/EditorInteractiveToolsFramework/Public/ViewportInteractions/ViewportLegacyClickInteraction.h b/Engine/Source/Editor/Experimental/EditorInteractiveToolsFramework/Public/ViewportInteractions/ViewportLegacyClickInteraction.h
new file mode 100644
--- /dev/null
+++ b/Engine/Source/Editor/Experimental/EditorInteractiveToolsFramework/Public/ViewportInteractions/ViewportLegacyClickInteraction.h
@@ -0,0 +1,30 @@
+﻿// Copyright Epic Games, Inc. All Rights Reserved.
+
+#pragma once
+
+#include "CoreMinimal.h"
+#include "ViewportClickInteraction.h"
+#include "UObject/Object.h"
+#include "ViewportLegacyClickInteraction.generated.h"
+
+#define UE_API EDITORINTERACTIVETOOLSFRAMEWORK_API
+
+/**
+ * An interaction that passes through ITF clicks through to the legacy viewport click handlers.
+ */
+UCLASS(MinimalAPI)
+class UViewportLegacyClickInteraction : public UViewportClickInteraction
+{
+	GENERATED_BODY()
+	
+public:
+	UViewportLegacyClickInteraction();	
+
+protected:
+
+	virtual void BuildBehaviors() override;
+	
+	virtual void ProcessClick_Internal(const FViewportClick& InViewportClick, HHitProxy* InHitProxy, FSceneView& View) override;
+};
+
+#undef UE_API
diff --git a/Engine/Source/Editor/LevelEditor/Private/ViewportInteractions/LevelViewportClickSelection.cpp b/Engine/Source/Editor/LevelEditor/Private/ViewportInteractions/LevelViewportClickSelection.cpp
--- a/Engine/Source/Editor/LevelEditor/Private/ViewportInteractions/LevelViewportClickSelection.cpp
+++ b/Engine/Source/Editor/LevelEditor/Private/ViewportInteractions/LevelViewportClickSelection.cpp
@@ -126,8 +126,6 @@ ULevelViewportClickContextMenu::ULevelViewportClickContextMenu()
 {
 	InteractionName = TEXT("Level Viewport Context Menu");
 	Groups = { UE::Editor::ViewportInteractions::ViewportClick };
-
-	MouseButton = EKeys::RightMouseButton;
 }
 
 void ULevelViewportClickContextMenu::BuildBehaviors()
@@ -152,8 +150,6 @@ ULevelViewportSetPivot::ULevelViewportSetPivot()
 {
 	InteractionName = TEXT("Level Viewport Set Pivot");
 	Groups = { UE::Editor::ViewportInteractions::ViewportClick };
-
-	MouseButton = EKeys::MiddleMouseButton;
 }
 
 void ULevelViewportSetPivot::BuildBehaviors()
diff --git a/Engine/Source/Editor/LevelEditor/Private/ViewportInteractions/LevelViewportClickSelection.h b/Engine/Source/Editor/LevelEditor/Private/ViewportInteractions/LevelViewportClickSelection.h
--- a/Engine/Source/Editor/LevelEditor/Private/ViewportInteractions/LevelViewportClickSelection.h
+++ b/Engine/Source/Editor/LevelEditor/Private/ViewportInteractions/LevelViewportClickSelection.h
@@ -76,11 +76,6 @@ class ULevelViewportDoubleClick : public ULevelViewportClickSelection
 public:
 	UE_API ULevelViewportDoubleClick();
 
-	virtual EInputEvent GetMouseEvent() const override
-	{
-		return IE_DoubleClick;
-	}
-	
 protected:
 	UE_API virtual void BuildBehaviors() override; 
 };
diff --git a/Engine/Source/Editor/Persona/Private/SkeletonSelectionEditMode.cpp b/Engine/Source/Editor/Persona/Private/SkeletonSelectionEditMode.cpp
--- a/Engine/Source/Editor/Persona/Private/SkeletonSelectionEditMode.cpp
+++ b/Engine/Source/Editor/Persona/Private/SkeletonSelectionEditMode.cpp
@@ -17,9 +17,6 @@
 #include "IEditableSkeleton.h"
 
 #include "SceneView.h"
-#include "Tools/EdModeInteractiveToolsContext.h"
-#include "ViewportInteractions/SkeletonSelectionClickInteraction.h"
-#include "ViewportInteractions/ViewportInteractionsBehaviorSource.h"
 
 #define LOCTEXT_NAMESPACE "SkeletonSelectionEditMode"
 
@@ -44,16 +41,6 @@ const FReferenceSkeleton& FSkeletonSelectionEditMode::GetReferenceSkeletonForCom
 	return Component->GetSkeletalMeshAsset() && !Component->GetSkeletalMeshAsset()->IsCompiling() ? Component->GetSkeletalMeshAsset()->GetRefSkeleton() : GetAnimPreviewScene().GetPersonaToolkit()->GetSkeleton()->GetReferenceSkeleton();
 }
 
-void FSkeletonSelectionEditMode::OnBuildViewportInteractions(UViewportInteractionsBehaviorSource* Source)
-{
-	USkeletonSelectionClickInteraction* ClickInteraction = Source->AddInteraction<USkeletonSelectionClickInteraction>();
-	ClickInteraction->Bind(SharedThis(this));
-
-	USkeletonSelectionRightClickInteraction* RightClickInteraction = Source->AddInteraction<USkeletonSelectionRightClickInteraction>();
-	RightClickInteraction->Bind(SharedThis(this));
-}
-
-
 bool FSkeletonSelectionEditMode::GetCameraTarget(FSphere& OutTarget) const
 {
 	bool bHandled = false;
@@ -99,26 +86,6 @@ void FSkeletonSelectionEditMode::GetOnScreenDebugInfo(TArray<FText>& OutDebugInf
 
 }
 
-void FSkeletonSelectionEditMode::Enter()
-{
-	IPersonaEditMode::Enter();
-
-	if (UModeManagerInteractiveToolsContext* const ModeManagerInteractiveToolsContext = GetModeManager()->GetInteractiveToolsContext())
-	{
-		ModeManagerInteractiveToolsContext->OnBuildViewportInteractions().AddSP(this, &FSkeletonSelectionEditMode::OnBuildViewportInteractions);
-	}	
-}
-
-void FSkeletonSelectionEditMode::Exit()
-{
-	if (UModeManagerInteractiveToolsContext* const ModeManagerInteractiveToolsContext = GetModeManager()->GetInteractiveToolsContext())
-	{
-		ModeManagerInteractiveToolsContext->OnBuildViewportInteractions().RemoveAll(this);
-	}
-	
-	IPersonaEditMode::Exit();
-}
-
 FSelectedSocketInfo FSkeletonSelectionEditMode::DuplicateAndSelectSocket(const FSelectedSocketInfo& SocketInfoToDuplicate)
 {
 	USkeletalMesh* SkeletalMesh = GetAnimPreviewScene().GetPreviewMeshComponent()->GetSkeletalMeshAsset();
diff --git a/Engine/Source/Editor/Persona/Private/ViewportInteractions/SkeletonSelectionClickInteraction.cpp b/Engine/Source/Editor/Persona/Private/ViewportInteractions/SkeletonSelectionClickInteraction.cpp
deleted file mode 100644
diff --git a/Engine/Source/Editor/Persona/Public/SkeletonSelectionEditMode.h b/Engine/Source/Editor/Persona/Public/SkeletonSelectionEditMode.h
--- a/Engine/Source/Editor/Persona/Public/SkeletonSelectionEditMode.h
+++ b/Engine/Source/Editor/Persona/Public/SkeletonSelectionEditMode.h
@@ -31,8 +31,6 @@ public:
 	virtual void GetOnScreenDebugInfo(TArray<FText>& OutDebugInfo) const override;
 
 	/** FEdMode interface */
-	virtual void Enter() override;
-	virtual void Exit() override;
 	virtual bool StartTracking(FEditorViewportClient* InViewportClient, FViewport* InViewport) override;
 	virtual bool EndTracking(FEditorViewportClient* InViewportClient, FViewport* InViewport) override;
 	virtual bool InputDelta(FEditorViewportClient* InViewportClient, FViewport* InViewport, FVector& InDrag, FRotator& InRot, FVector& InScale) override;
@@ -56,7 +54,6 @@ protected:
 	FTransform GetSocketTransform(const USkeletalMeshSocket* Socket) const;
 	FTransform GetWorldSpaceBoneTransform(const FReferenceSkeleton& ReferenceSkeleton, const int32 BoneIndex) const;
 	const FReferenceSkeleton& GetReferenceSkeletonForComponent(const USkeletalMeshComponent* Component) const;
-	void OnBuildViewportInteractions(UViewportInteractionsBehaviorSource* Source);
 	
 private:
 	/** Duplicates and selects a socket when we alt-drag */
diff --git a/Engine/Source/Editor/Persona/Public/ViewportInteractions/SkeletonSelectionClickInteraction.h b/Engine/Source/Editor/Persona/Public/ViewportInteractions/SkeletonSelectionClickInteraction.h
deleted file mode 100644
diff --git a/Engine/Source/Editor/SkeletalMeshEditor/Private/SkeletalMeshEditor.cpp b/Engine/Source/Editor/SkeletalMeshEditor/Private/SkeletalMeshEditor.cpp
--- a/Engine/Source/Editor/SkeletalMeshEditor/Private/SkeletalMeshEditor.cpp
+++ b/Engine/Source/Editor/SkeletalMeshEditor/Private/SkeletalMeshEditor.cpp
@@ -1628,11 +1628,6 @@ TSharedPtr<FUICommandInfo> FSkeletalMeshEditor::GetResetAllBonesTransformsComman
 	return Commands.ResetAllBonesTransforms;
 }
 
-void FSkeletalMeshEditor::OnBuildViewportInteractions(UViewportInteractionsBehaviorSource* Source)
-{
-	UE::Editor::ViewportInteractions::AddDefaultCameraMovementInteractions(Source);
-}
-
 FSkeletalMeshEditorBinding::FSkeletalMeshEditorBinding(TSharedRef<FSkeletalMeshEditor> InEditor)
 	: Editor(InEditor)
 {}
diff --git a/Engine/Source/Editor/SkeletalMeshEditor/Private/SkeletalMeshEditor.h b/Engine/Source/Editor/SkeletalMeshEditor/Private/SkeletalMeshEditor.h
--- a/Engine/Source/Editor/SkeletalMeshEditor/Private/SkeletalMeshEditor.h
+++ b/Engine/Source/Editor/SkeletalMeshEditor/Private/SkeletalMeshEditor.h
@@ -137,9 +137,6 @@ public:
 	virtual TSharedPtr<FUICommandInfo> GetResetBoneTransformsCommand() override;
 	virtual TSharedPtr<FUICommandInfo> GetResetAllBonesTransformsCommand() override;
 	
-protected:
-	void OnBuildViewportInteractions(UViewportInteractionsBehaviorSource* Source) override;
-
 private:
 	void HandleObjectSelected(UObject* InObject);
 	
diff --git a/Engine/Source/Editor/UnrealEd/Private/Toolkits/AssetEditorToolkit.cpp b/Engine/Source/Editor/UnrealEd/Private/Toolkits/AssetEditorToolkit.cpp
--- a/Engine/Source/Editor/UnrealEd/Private/Toolkits/AssetEditorToolkit.cpp
+++ b/Engine/Source/Editor/UnrealEd/Private/Toolkits/AssetEditorToolkit.cpp
@@ -40,6 +40,9 @@
 #include "Toolkits/SStandaloneAssetEditorToolkitHost.h"
 #include "Toolkits/ToolkitManager.h"
 #include "Tools/EdModeInteractiveToolsContext.h"
+#include "ViewportInteractions/ViewportInteractionsBehaviorSource.h"
+#include "ViewportInteractions/ViewportInteractionsUtilities.h"
+#include "ViewportInteractions/ViewportLegacyClickInteraction.h"
 #include "Widgets/SToolTip.h"
 #include "Widgets/Docking/SDockTab.h"
 #include "Widgets/Images/SImage.h"
@@ -921,6 +924,13 @@ void FAssetEditorToolkit::RestoreFromLayout(const TSharedRef<FTabManager::FLayou
 	}
 }
 
+void FAssetEditorToolkit::OnBuildViewportInteractions(UViewportInteractionsBehaviorSource* Source)
+{
+	UE::Editor::ViewportInteractions::AddDefaultCameraMovementInteractions(Source);
+	
+	Source->AddInteraction<UViewportLegacyClickInteraction>();
+}
+
 FAssetEditorModeManager* FAssetEditorToolkit::GetAssetEditorModeManager() const
 {
 	PRAGMA_DISABLE_DEPRECATION_WARNINGS
diff --git a/Engine/Source/Editor/UnrealEd/Public/Toolkits/AssetEditorToolkit.h b/Engine/Source/Editor/UnrealEd/Public/Toolkits/AssetEditorToolkit.h
--- a/Engine/Source/Editor/UnrealEd/Public/Toolkits/AssetEditorToolkit.h
+++ b/Engine/Source/Editor/UnrealEd/Public/Toolkits/AssetEditorToolkit.h
@@ -486,7 +486,7 @@ protected:
 	virtual FName GetOverrideTabId() const { return NAME_None; }
 	
 	UE_EXPERIMENTAL(5.8, "Viewport Interactions are Experimental")
-	virtual void OnBuildViewportInteractions(UViewportInteractionsBehaviorSource* Source) {};
+	UNREALED_API virtual void OnBuildViewportInteractions(UViewportInteractionsBehaviorSource* Source);
 	
 	UE_EXPERIMENTAL(5.8, "Viewport Interactions are Experimental")
 	virtual void PostBuildViewportInteractions(const UViewportInteractionsBehaviorSource* Source) {};
