[ITF Gizmos] Fix an issue where "animator gizmo" functionality was dependent on the full ITF viewport interaction stack being enabled. Consolidate local editor gizmo logic into a single static helper function.

#rnx
#Synced-CL 51211095

diff --git a/Engine/Plugins/Animation/ControlRig/Source/ControlRigEditor/Private/EditMode/ControlRigEditMode.cpp b/Engine/Plugins/Animation/ControlRig/Source/ControlRigEditor/Private/EditMode/ControlRigEditMode.cpp
--- a/Engine/Plugins/Animation/ControlRig/Source/ControlRigEditor/Private/EditMode/ControlRigEditMode.cpp
+++ b/Engine/Plugins/Animation/ControlRig/Source/ControlRigEditor/Private/EditMode/ControlRigEditMode.cpp
@@ -2155,7 +2155,7 @@ bool FControlRigEditMode::InputDelta(FEditorViewportClient* InViewportClient, FV
 	//button down if left and ctrl and right is down, needed for indirect posting
 
 	// enable MMB with the new TRS gizmos
-	const bool bEnableMMB = InViewportClient->GetModeTools()->UsesViewportITF();
+	const bool bEnableMMB = UEditorInteractiveGizmoManager::AreEditorGizmosAllowed(InViewportClient->GetModeTools());
 	
 	const bool bMouseButtonDown =
 		InViewport->KeyState(EKeys::LeftMouseButton) ||
diff --git a/Engine/Source/Editor/Experimental/EditorInteractiveToolsFramework/Private/EditorGizmos/EditorTransformGizmoSource.cpp b/Engine/Source/Editor/Experimental/EditorInteractiveToolsFramework/Private/EditorGizmos/EditorTransformGizmoSource.cpp
--- a/Engine/Source/Editor/Experimental/EditorInteractiveToolsFramework/Private/EditorGizmos/EditorTransformGizmoSource.cpp
+++ b/Engine/Source/Editor/Experimental/EditorInteractiveToolsFramework/Private/EditorGizmos/EditorTransformGizmoSource.cpp
@@ -98,7 +98,7 @@ bool UEditorTransformGizmoSource::CanInteract(const EViewportContext InViewportC
 			bool bUseLegacyWidget = (WidgetMode == UE::Widget::WM_TranslateRotateZ || WidgetMode == UE::Widget::WM_2D);
 			if (!bUseLegacyWidget)
 			{
-				bUseLegacyWidget = !ModeTools.UsesViewportITF();
+				bUseLegacyWidget = !UEditorInteractiveGizmoManager::AreEditorGizmosAllowed(&ModeTools);
 			}
 			return !bUseLegacyWidget;
 		}
diff --git a/Engine/Source/Editor/Experimental/EditorInteractiveToolsFramework/Private/EditorInteractiveGizmoManager.cpp b/Engine/Source/Editor/Experimental/EditorInteractiveToolsFramework/Private/EditorInteractiveGizmoManager.cpp
--- a/Engine/Source/Editor/Experimental/EditorInteractiveToolsFramework/Private/EditorInteractiveGizmoManager.cpp
+++ b/Engine/Source/Editor/Experimental/EditorInteractiveToolsFramework/Private/EditorInteractiveGizmoManager.cpp
@@ -300,6 +300,22 @@ void UEditorInteractiveGizmoManager::SetAllowEditorGizmos(bool bInAllowed)
 	}
 }
 
+bool UEditorInteractiveGizmoManager::AreEditorGizmosAllowed(const FEditorModeTools* ModeTools)
+{
+	if (ModeTools && ModeTools->SupportsViewportITF())
+	{
+		if (const UModeManagerInteractiveToolsContext* Context = ModeTools->GetInteractiveToolsContext())
+		{
+			if (const UEditorInteractiveGizmoManager* GizmoManager = Cast<UEditorInteractiveGizmoManager>(Context->GizmoManager))
+			{
+				return GizmoManager->AreEditorGizmosAllowed();
+			}
+		}
+	}
+
+	return false;
+}
+
 void UEditorInteractiveGizmoManager::UpdateActiveEditorGizmos()
 {
 	if (!AreEditorGizmosAllowed())
diff --git a/Engine/Source/Editor/Experimental/EditorInteractiveToolsFramework/Public/EditorInteractiveGizmoManager.h b/Engine/Source/Editor/Experimental/EditorInteractiveToolsFramework/Public/EditorInteractiveGizmoManager.h
--- a/Engine/Source/Editor/Experimental/EditorInteractiveToolsFramework/Public/EditorInteractiveGizmoManager.h
+++ b/Engine/Source/Editor/Experimental/EditorInteractiveToolsFramework/Public/EditorInteractiveGizmoManager.h
@@ -146,6 +146,11 @@ public:
 	static UE_API const FString& TransformInstanceIdentifier();
 	static UE_API const FString& TransformBuilderIdentifier();
 
+	/**
+	 * Returns true if editor gizmos are enabled for the provided ModeTools instance.
+	 */
+	static UE_API bool AreEditorGizmosAllowed(const FEditorModeTools* ModeTools);
+
 	/**
 	 * Returns true if the new TRS gizmos are used.
 	 */
diff --git a/Engine/Source/Editor/UnrealEd/Private/EditorModeManager.cpp b/Engine/Source/Editor/UnrealEd/Private/EditorModeManager.cpp
--- a/Engine/Source/Editor/UnrealEd/Private/EditorModeManager.cpp
+++ b/Engine/Source/Editor/UnrealEd/Private/EditorModeManager.cpp
@@ -146,7 +146,7 @@ void FEditorModeTools::LoadConfig(void)
 	if (CoordinateSystem == COORD_Parent || CoordinateSystem == COORD_Explicit)
 	{
 		// parent mode is only supported with new trs gizmos for now
-		if (!UsesViewportITF())
+		if (!(bViewportITF && UEditorInteractiveGizmoManager::UsesNewTRSGizmos()))
 		{
 			CoordSystemAsInt = static_cast<int32>(COORD_Local);
 		}
diff --git a/Engine/Source/Editor/UnrealEd/Private/LevelEditorViewport.cpp b/Engine/Source/Editor/UnrealEd/Private/LevelEditorViewport.cpp
--- a/Engine/Source/Editor/UnrealEd/Private/LevelEditorViewport.cpp
+++ b/Engine/Source/Editor/UnrealEd/Private/LevelEditorViewport.cpp
@@ -3320,7 +3320,7 @@ bool FLevelEditorViewportClient::InputWidgetDelta(FViewport* InViewport, EAxisLi
 				// enable MMB for New TRS Gizmos
 				const bool bModifyTransform = [this, MiddleMouseButtonDown]() -> bool
 				{
-					if (!ModeTools->UsesViewportITF())
+					if (!UEditorInteractiveGizmoManager::AreEditorGizmosAllowed(ModeTools.Get()))
 					{
 						return !MiddleMouseButtonDown;
 					}
diff --git a/Engine/Source/Editor/UnrealEd/Private/SEditorViewport.cpp b/Engine/Source/Editor/UnrealEd/Private/SEditorViewport.cpp
--- a/Engine/Source/Editor/UnrealEd/Private/SEditorViewport.cpp
+++ b/Engine/Source/Editor/UnrealEd/Private/SEditorViewport.cpp
@@ -1198,7 +1198,7 @@ void SEditorViewport::OnCycleCoordinateSystem()
 
 	// parent mode is only supported with new trs gizmos for now
 	int CoordMax = COORD_Parent;
-	if (Client->GetModeTools()->UsesViewportITF())
+	if (UEditorInteractiveGizmoManager::AreEditorGizmosAllowed(Client->GetModeTools()))
 	{
 		CoordMax = UEditorInteractiveGizmoManager::IsExplicitModeEnabled() ? COORD_Max : COORD_Explicit;
 	}
diff --git a/Engine/Source/Editor/UnrealEd/Private/UnrealWidgetRender.cpp b/Engine/Source/Editor/UnrealEd/Private/UnrealWidgetRender.cpp
--- a/Engine/Source/Editor/UnrealEd/Private/UnrealWidgetRender.cpp
+++ b/Engine/Source/Editor/UnrealEd/Private/UnrealWidgetRender.cpp
@@ -16,6 +16,7 @@
 #include "EditorInteractiveGizmoManager.h"
 #include "TransformGizmoEditorSettings.h"
 #include "Misc/AxisDisplayInfo.h"
+#include "Tools/EdModeInteractiveToolsContext.h"
 
 /*
  *  Simple struct used to create and group data related to the current window's / viewport's space,
@@ -190,7 +191,7 @@ void FWidget::Render( const FSceneView* View,FPrimitiveDrawInterface* PDI, FEdit
 	bool bUseLegacyWidget = (WidgetMode == UE::Widget::EWidgetMode::WM_TranslateRotateZ || WidgetMode == UE::Widget::EWidgetMode::WM_2D);
 	if (!bUseLegacyWidget)
 	{
-		bUseLegacyWidget = !ViewportClient->GetModeTools()->UsesViewportITF();
+		bUseLegacyWidget = !UEditorInteractiveGizmoManager::AreEditorGizmosAllowed(ViewportClient->GetModeTools());
 	}
 
 	bool bDrawWidget;
diff --git a/Engine/Source/Editor/UnrealEd/Private/ViewportToolbar/UnrealEdViewportToolbar.cpp b/Engine/Source/Editor/UnrealEd/Private/ViewportToolbar/UnrealEdViewportToolbar.cpp
--- a/Engine/Source/Editor/UnrealEd/Private/ViewportToolbar/UnrealEdViewportToolbar.cpp
+++ b/Engine/Source/Editor/UnrealEd/Private/ViewportToolbar/UnrealEdViewportToolbar.cpp
@@ -1064,7 +1064,7 @@ FToolMenuEntry CreateTransformsSubmenu()
 
 													if (TSharedPtr<FEditorViewportClient> Client = UUnrealEdViewportToolbarContext::GetEditorViewportClient(InSubmenu->Context))
 													{
-														if (FEditorModeTools* ModeTools = Client->GetModeTools(); ModeTools && ModeTools->UsesViewportITF())
+														if (UEditorInteractiveGizmoManager::AreEditorGizmosAllowed(Client->GetModeTools()))
 														{
 															FToolMenuEntry ParentCoords = FToolMenuEntry::InitMenuEntry(
 																FEditorViewportCommands::Get().RelativeCoordinateSystem_Parent
