[Gizmos] Modify viewport click interaction to just pass through to the VPC process click function for legacy click behavior

#Synced-CL 51343360

diff --git a/Engine/Source/Editor/Experimental/EditorInteractiveToolsFramework/Private/ViewportInteractions/ViewportLegacyClickInteraction.cpp b/Engine/Source/Editor/Experimental/EditorInteractiveToolsFramework/Private/ViewportInteractions/ViewportLegacyClickInteraction.cpp
--- a/Engine/Source/Editor/Experimental/EditorInteractiveToolsFramework/Private/ViewportInteractions/ViewportLegacyClickInteraction.cpp
+++ b/Engine/Source/Editor/Experimental/EditorInteractiveToolsFramework/Private/ViewportInteractions/ViewportLegacyClickInteraction.cpp
@@ -3,6 +3,7 @@
 #include "ViewportInteractions/ViewportLegacyClickInteraction.h"
 
 #include "EditorModeManager.h"
+#include "EditorViewportClient.h"
 
 UViewportLegacyClickInteraction::UViewportLegacyClickInteraction()
 {
@@ -36,12 +37,9 @@ void UViewportLegacyClickInteraction::BuildBehaviors()
 
 void UViewportLegacyClickInteraction::ProcessClick_Internal(const FViewportClick& InViewportClick, HHitProxy* InHitProxy, FSceneView& View)
 {
-	if (FEditorModeTools* ModeTools = GetModeTools())
+	if (FEditorViewportClient* ViewportClient = InViewportClick.GetViewportClient())
 	{
-		if (FEditorViewportClient* ViewportClient = ModeTools->GetFocusedViewportClient())
-		{
-			ModeTools->HandleClick(ViewportClient, InHitProxy, InViewportClick);
-		}
+		ViewportClient->ProcessClick(View, InHitProxy, InViewportClick.GetKey(), InViewportClick.GetEvent(), InViewportClick.GetClickPos().X, InViewportClick.GetClickPos().Y);
 	}
 }
 
