[ITF Viewport Interactions] Zoom interaction captures mouse wheel when not necessary

#Synced-CL 51539428

diff --git a/Engine/Source/Editor/Experimental/EditorInteractiveToolsFramework/Private/ViewportInteractions/ViewportZoomInteraction.cpp b/Engine/Source/Editor/Experimental/EditorInteractiveToolsFramework/Private/ViewportInteractions/ViewportZoomInteraction.cpp
--- a/Engine/Source/Editor/Experimental/EditorInteractiveToolsFramework/Private/ViewportInteractions/ViewportZoomInteraction.cpp
+++ b/Engine/Source/Editor/Experimental/EditorInteractiveToolsFramework/Private/ViewportInteractions/ViewportZoomInteraction.cpp
@@ -22,6 +22,16 @@ UViewportZoomInteraction::UViewportZoomInteraction()
 	UMouseWheelInputBehavior* MouseWheelInputBehavior = NewObject<UMouseWheelInputBehavior>();
 	MouseWheelInputBehavior->Initialize(this);
 	MouseWheelInputBehavior->SetDefaultPriority(UE::Editor::ViewportInteractions::VIEWPORT_INTERACTIONS_DEFAULT_PRIORITY);
+	MouseWheelInputBehavior->ModifierCheckFunc = [](const FInputDeviceState& InputDeviceState)
+	{
+		if (FInputDeviceState::IsAltKeyDown(InputDeviceState) || FInputDeviceState::IsCtrlKeyDown(InputDeviceState)
+			|| FInputDeviceState::IsShiftKeyDown(InputDeviceState) || FInputDeviceState::IsCmdKeyDown(InputDeviceState))
+		{
+			return false;
+		}
+
+		return true;
+	};
 
 	// Keyboard zoom
 	UKeyInputBehavior* KeyInputBehavior = NewObject<UKeyInputBehavior>();
