[ITF Viewport Interactions] Update and bring back Ortho Viewport switching (Viewport Change Tool)

#Synced-CL 51759566

diff --git a/Engine/Source/Editor/Experimental/EditorInteractiveToolsFramework/Private/EditorDragTools/ViewportChangeInteraction.cpp b/Engine/Source/Editor/Experimental/EditorInteractiveToolsFramework/Private/EditorDragTools/ViewportChangeInteraction.cpp
--- a/Engine/Source/Editor/Experimental/EditorInteractiveToolsFramework/Private/EditorDragTools/ViewportChangeInteraction.cpp
+++ b/Engine/Source/Editor/Experimental/EditorInteractiveToolsFramework/Private/EditorDragTools/ViewportChangeInteraction.cpp
@@ -1,21 +1,32 @@
 ﻿// Copyright Epic Games, Inc. All Rights Reserved.
 
 #include "EditorDragTools/ViewportChangeInteraction.h"
-#include "BaseBehaviors/ClickDragBehavior.h"
 #include "CanvasItem.h"
 #include "CanvasTypes.h"
 #include "EditorViewportClient.h"
-#include "Framework/Commands/InputChord.h"
 #include "Settings/EditorStyleSettings.h"
 #include "SnappingUtils.h"
 
 UViewportChangeInteraction::UViewportChangeInteraction()
 {
+	using namespace UE::Editor::ViewportInteractions;
+
 	InteractionName = UE::Editor::ViewportInteractions::ViewportChange;
 
 	ViewOptionOffset = FVector2D::ZeroVector;
 	bUseSnapping = true;
 	bConvertDelta = false;
+
+	UViewportClickDragBehavior* MouseBehavior = NewObject<UViewportClickDragBehavior>();
+	MouseBehavior->Initialize(this);
+	MouseBehavior->SetBindings({ 
+		FButtonBinding(EKeys::LeftControl).RequiredToContinue(false),
+		FButtonBinding(EKeys::LeftAlt).RequiredToContinue(false),
+		FButtonBinding(EKeys::MiddleMouseButton).TriggersStart() 
+	}
+	);
+
+	RegisterInputBehavior(MouseBehavior);
 }
 
 void UViewportChangeInteraction::Render(const FSceneView* View, FCanvas* Canvas, EViewInteractionState InteractionState)
diff --git a/Engine/Source/Editor/Experimental/EditorInteractiveToolsFramework/Private/ViewportInteractions/ViewportInteractionsUtilities.cpp b/Engine/Source/Editor/Experimental/EditorInteractiveToolsFramework/Private/ViewportInteractions/ViewportInteractionsUtilities.cpp
--- a/Engine/Source/Editor/Experimental/EditorInteractiveToolsFramework/Private/ViewportInteractions/ViewportInteractionsUtilities.cpp
+++ b/Engine/Source/Editor/Experimental/EditorInteractiveToolsFramework/Private/ViewportInteractions/ViewportInteractionsUtilities.cpp
@@ -5,6 +5,7 @@
 #include "EditorDragTools/BoxSelectInteraction.h"
 #include "EditorDragTools/FrustumSelectInteraction.h"
 #include "EditorDragTools/MeasureToolInteraction.h"
+#include "EditorDragTools/ViewportChangeInteraction.h"
 #include "ViewportInteractions/ViewportCameraRotateInteraction.h"
 #include "ViewportInteractions/ViewportCameraSpeedMouseWheelInteraction.h"
 #include "ViewportInteractions/ViewportCameraTranslateInteraction.h"
@@ -44,7 +45,11 @@ void UE::Editor::ViewportInteractions::AddDefaultCameraMovementInteractions(
 
 void UE::Editor::ViewportInteractions::AddDefaultDragToolsInteractions(UViewportInteractionsBehaviorSource* InInteractionsBehaviorSource)
 {
-	InInteractionsBehaviorSource->AddInteraction<UFrustumSelectInteraction>();
-	InInteractionsBehaviorSource->AddInteraction<UBoxSelectInteraction>();
-	InInteractionsBehaviorSource->AddInteraction<UMeasureToolInteraction>();
+	if (InInteractionsBehaviorSource)
+	{
+		InInteractionsBehaviorSource->AddInteraction<UFrustumSelectInteraction>();
+		InInteractionsBehaviorSource->AddInteraction<UBoxSelectInteraction>();
+		InInteractionsBehaviorSource->AddInteraction<UMeasureToolInteraction>();
+		InInteractionsBehaviorSource->AddInteraction<UViewportChangeInteraction>();
+	}
 }
