[Modes][Gizmos] Update default mode manager implementation to use the asset editor default mode, and the level editor one uses the level editor default ed mode

Removed swapping of modes - this wasn't having an effect on gizmo visibility in asset editors

This fixes an issue where when legacy gizmos are toggled on globally, a selection must be active in the level editor for the gizmos to display in an asset editor that was using the default editor modes

#Synced-CL 51761123

diff --git a/Engine/Source/Editor/Experimental/EditorInteractiveToolsFramework/Private/EditorGizmos/EditorTransformGizmoUtil.cpp b/Engine/Source/Editor/Experimental/EditorInteractiveToolsFramework/Private/EditorGizmos/EditorTransformGizmoUtil.cpp
--- a/Engine/Source/Editor/Experimental/EditorInteractiveToolsFramework/Private/EditorGizmos/EditorTransformGizmoUtil.cpp
+++ b/Engine/Source/Editor/Experimental/EditorInteractiveToolsFramework/Private/EditorGizmos/EditorTransformGizmoUtil.cpp
@@ -263,31 +263,16 @@ void UEditorTransformGizmoContextObject::UpdateGizmo(const TArray<FEditorViewpor
 	
 	const UModeManagerInteractiveToolsContext* ToolsContext = ModeTools->GetInteractiveToolsContext();
 	const TObjectPtr<UInteractiveToolManager> ToolManager = ToolsContext->ToolManager;
-
-	// swap default mode if not dealing with the level editor mode manager
-	const bool bIsLevelViewportClient = [&ViewportClients]()
-	{
-		if (GEditor)
-		{
-			return GEditor->GetLevelViewportClients().ContainsByPredicate([&ViewportClients](const FEditorViewportClient* LevelViewportClient)
-			{
-				return ViewportClients.Contains(LevelViewportClient);
-			});
-		}
-		return false;
-	}();
-	
-	const bool Swapped = !bIsLevelViewportClient ? SwapDefaultMode(FBuiltinEditorModes::EM_Default, FAssetEdModes::EM_AssetDefault) : false;
 	
 	if (UE::EditorTransformGizmoUtil::FindDefaultTransformGizmo(ToolManager))
 	{
 		return;
 	}
-				
+
 	// create a new one
 	if (UTransformGizmo* TransformGizmo = UE::EditorTransformGizmoUtil::GetDefaultTransformGizmo(ToolManager))
 	{
-		TransformGizmo->SetVisibility(Swapped ? ModeTools->UsesViewportITF() : false);
+		TransformGizmo->SetVisibility(!ModeTools->IsModeActive(FBuiltinEditorModes::EM_Default) ? ModeTools->UsesViewportITF() : false);
 	}
 }
 
@@ -295,9 +280,6 @@ void UEditorTransformGizmoContextObject::SetGizmosEnabled(bool bInEnabled)
 {
 	if (!bInEnabled)
 	{
-		// swap back default mode
-		(void)SwapDefaultMode(FAssetEdModes::EM_AssetDefault, FBuiltinEditorModes::EM_Default);
-
 		// reset the coordinate system if the legacy gizmo doesn't support it
 		if (ModeTools->GetCoordSystem() > COORD_Local)
 		{
@@ -464,25 +446,3 @@ void UEditorTransformGizmoContextObject::RemoveViewportsBinding()
 		LevelViewportClientsChangedHandle.Reset();
 	}
 }
-
-bool UEditorTransformGizmoContextObject::SwapDefaultMode(const FEditorModeID InCurrentDefaultMode, const FEditorModeID InNewDefaultMode) const
-{
-	// swap only if we're not dealing with the level editor + InCurrentDefaultMode is one of the default modes
-	if (!ModeTools->IsDefaultMode(InCurrentDefaultMode))
-	{
-		return false;
-	}
-
-	// replace default mode
-	ModeTools->RemoveDefaultMode(InCurrentDefaultMode);
-	ModeTools->AddDefaultMode(InNewDefaultMode);
-	
-	if (ModeTools->IsModeActive(InCurrentDefaultMode))
-	{
-		// activate it if the previous one was 
-		ModeTools->DeactivateMode(InCurrentDefaultMode);
-		ModeTools->ActivateMode(InNewDefaultMode);
-	}
-
-	return true;
-}
diff --git a/Engine/Source/Editor/Experimental/EditorInteractiveToolsFramework/Public/EditorGizmos/EditorTransformGizmoUtil.h b/Engine/Source/Editor/Experimental/EditorInteractiveToolsFramework/Public/EditorGizmos/EditorTransformGizmoUtil.h
--- a/Engine/Source/Editor/Experimental/EditorInteractiveToolsFramework/Public/EditorGizmos/EditorTransformGizmoUtil.h
+++ b/Engine/Source/Editor/Experimental/EditorInteractiveToolsFramework/Public/EditorGizmos/EditorTransformGizmoUtil.h
@@ -111,9 +111,6 @@ private:
 	EDITORINTERACTIVETOOLSFRAMEWORK_API void RemoveGizmoManagerBinding();
 	EDITORINTERACTIVETOOLSFRAMEWORK_API void RemoveViewportsBinding();
 
-	typedef FName FEditorModeID;
-	bool SwapDefaultMode(const FEditorModeID InCurrentDefaultMode, const FEditorModeID InNewDefaultMode) const;
-
 	FOnGizmoCreated OnGizmoCreated;
 	
 	FEditorModeTools* ModeTools = nullptr;
diff --git a/Engine/Source/Editor/UnrealEd/Private/EditorModeManager.cpp b/Engine/Source/Editor/UnrealEd/Private/EditorModeManager.cpp
--- a/Engine/Source/Editor/UnrealEd/Private/EditorModeManager.cpp
+++ b/Engine/Source/Editor/UnrealEd/Private/EditorModeManager.cpp
@@ -25,6 +25,7 @@
 #include "Toolkits/BaseToolkit.h"
 #include "Subsystems/AssetEditorSubsystem.h"
 #include "Subsystems/BrushEditingSubsystem.h"
+#include "Tools/DefaultEdMode.h"
 #include "Tools/UEdMode.h"
 #include "InputRouter.h"
 #include "Tools/EdModeInteractiveToolsContext.h"
@@ -70,7 +71,7 @@ FEditorModeTools::FEditorModeTools()
 	, GridBase(ForceInitToZero)
 	, TranslateRotateXAxisAngle(0.0f)
 	, TranslateRotate2DAngle(0.0f)
-	, DefaultModeIDs()
+	, DefaultModeIDs({ FAssetEdModes::EM_AssetDefault })
 	, WidgetMode(UE::Widget::WM_None)
 	, OverrideWidgetMode(UE::Widget::WM_None)
 	, bShowWidget(true)
@@ -80,8 +81,6 @@ FEditorModeTools::FEditorModeTools()
 	, CoordSystem(COORD_World)
 	, bIsTracking(false)
 {
-	DefaultModeIDs.Add( FBuiltinEditorModes::EM_Default );
-
 	InteractiveToolsContext = NewObject<UModeManagerInteractiveToolsContext>(GetTransientPackage(), UModeManagerInteractiveToolsContext::StaticClass(), NAME_None, RF_Transient);
 	InteractiveToolsContext->InitializeContextWithEditorModeManager(this);
 	InteractiveToolsContext->Activate();
@@ -106,7 +105,7 @@ FEditorModeTools::FEditorModeTools()
 
 FEditorModeTools::~FEditorModeTools()
 {
-	SetDefaultMode(FBuiltinEditorModes::EM_Default);
+	SetDefaultMode( FAssetEdModes::EM_AssetDefault );
 	DeactivateAllModes();
 
 	RemoveAllDelegateHandlers();
@@ -822,7 +821,7 @@ void FEditorModeTools::DestroyMode( FEditorModeID InID )
 	// Add back the default default mode if we just removed the last valid default.
 	if ( DefaultModeIDs.Num() == 0 )
 	{
-		AddDefaultMode( FBuiltinEditorModes::EM_Default );
+		AddDefaultMode( FAssetEdModes::EM_AssetDefault );
 	}
 
 	// Find the mode from the ID and exit it.
diff --git a/Engine/Source/Editor/UnrealEd/Private/Tools/DefaultEdMode.cpp b/Engine/Source/Editor/UnrealEd/Private/Tools/DefaultEdMode.cpp
--- a/Engine/Source/Editor/UnrealEd/Private/Tools/DefaultEdMode.cpp
+++ b/Engine/Source/Editor/UnrealEd/Private/Tools/DefaultEdMode.cpp
@@ -57,17 +57,7 @@ namespace FAssetEdModes
 
 class FAssetEdModeWidgetHelper : public FLegacyEdModeWidgetHelper
 {
-public:
-	virtual bool ShouldDrawWidget() const override
-	{
-		if (Owner && Owner->SelectionHasSceneComponent())
-		{
-			return true;
-		}
-
-		return false;
-	}
-	
+public:	
 	virtual bool RequiresLegacyViewportInteractions() const override { return false; }
 };
 
diff --git a/Engine/Source/Editor/UnrealEd/Private/Tools/LegacyEdModeWidgetHelpers.cpp b/Engine/Source/Editor/UnrealEd/Private/Tools/LegacyEdModeWidgetHelpers.cpp
--- a/Engine/Source/Editor/UnrealEd/Private/Tools/LegacyEdModeWidgetHelpers.cpp
+++ b/Engine/Source/Editor/UnrealEd/Private/Tools/LegacyEdModeWidgetHelpers.cpp
@@ -416,7 +416,7 @@ bool FLegacyEdModeWidgetHelper::ShouldDrawWidget() const
 
 	bool bDrawWidget = false;
 	bool bHadSelectableComponents = false;
-	if (Owner->GetSelectedComponents()->Num() > 0)
+	if (Owner->GetSelectedComponents() && Owner->GetSelectedComponents()->Num() > 0)
 	{
 		// when components are selected, only show the widget when one or more are scene components
 		for (FSelectedEditableComponentIterator It(*Owner->GetSelectedComponents()); It; ++It)
diff --git a/Engine/Source/Editor/UnrealEd/Private/UnrealEdGlobals.cpp b/Engine/Source/Editor/UnrealEd/Private/UnrealEdGlobals.cpp
--- a/Engine/Source/Editor/UnrealEd/Private/UnrealEdGlobals.cpp
+++ b/Engine/Source/Editor/UnrealEd/Private/UnrealEdGlobals.cpp
@@ -54,6 +54,7 @@ namespace Internal
 	public:
 		FLevelEditorModeTools()
 		{
+			DefaultModeIDs = { FBuiltinEditorModes::EM_Default };
 			USelection::SelectNoneEvent.AddRaw(this, &FLevelEditorModeTools::OnEditorSelectNone);
 			USelection::SelectionChangedEvent.AddRaw(this, &FLevelEditorModeTools::OnEditorSelectionChanged);
 			USelection::SelectObjectEvent.AddRaw(this, &FLevelEditorModeTools::OnEditorSelectionChanged);
