[Editor TRS] Move from EditorTRSGizmo -> EditorInteractiveToolsFramework

#Synced-CL 51889604

diff --git a/Engine/Source/Editor/Experimental/EditorInteractiveToolsFramework/EditorInteractiveToolsFramework.Build.cs b/Engine/Source/Editor/Experimental/EditorInteractiveToolsFramework/EditorInteractiveToolsFramework.Build.cs
--- a/Engine/Source/Editor/Experimental/EditorInteractiveToolsFramework/EditorInteractiveToolsFramework.Build.cs
+++ b/Engine/Source/Editor/Experimental/EditorInteractiveToolsFramework/EditorInteractiveToolsFramework.Build.cs
@@ -36,12 +36,17 @@ public class EditorInteractiveToolsFramework : ModuleRules
 				"EditorFramework",
                 "ContentBrowser",
                 "ApplicationCore",
+				"Json",
+				"JsonUtilities",
 				"MeshDescription",
 				"StaticMeshDescription",
                 "EditorSubsystem",
+                "PropertyEditor",
+                "ToolWidgets",
                 "TypedElementRuntime",
                 "AnimationCore",
                 "ViewportSnapping", // For ISnappingPolicy support
+                "WorkspaceMenuStructure",
                 // ... add private dependencies that you statically link with here ...
 			}
             );
diff --git a/Engine/Source/Editor/Experimental/EditorInteractiveToolsFramework/Private/EditorGizmos/TransformGizmo.cpp b/Engine/Source/Editor/Experimental/EditorInteractiveToolsFramework/Private/EditorGizmos/TransformGizmo.cpp
--- a/Engine/Source/Editor/Experimental/EditorInteractiveToolsFramework/Private/EditorGizmos/TransformGizmo.cpp
+++ b/Engine/Source/Editor/Experimental/EditorInteractiveToolsFramework/Private/EditorGizmos/TransformGizmo.cpp
@@ -23,6 +23,7 @@
 #include "PrimitiveDrawingUtils.h"
 #include "SceneQueries/SceneSnappingManager.h"
 #include "ToolDataVisualizer.h"
+#include "TransformGizmoEditorSettings.h"
 
 #include UE_INLINE_GENERATED_CPP_BY_NAME(TransformGizmo)
 
@@ -158,6 +159,51 @@ namespace UE::EditorInteractiveToolsFramework::Private
 	}
 };
 
+FTransformGizmoStyle::FTransformGizmoStyle()
+{
+}
+
+void FTransformGizmoStyle::MakeValid()
+{
+	constexpr float AxisSizeMultiplierMin = 0.1f;
+	AxisSizeMultiplier = FMath::Max(AxisSizeMultiplierMin, AxisSizeMultiplier);
+}
+
+float FTransformGizmoStyle::GetModifiedAxisSizeMultiplier(const TOptional<TUniqueFunction<float(const float)>>& InModifierOverride) const
+{
+	if (InModifierOverride.IsSet())
+	{
+		return InModifierOverride.GetValue()(AxisSizeMultiplier);
+	}
+
+	/** Adjustment for TransformGizmoSize to apply to AxisSizeMultiplier. */
+	constexpr float TransformGizmoSizeToMultiplierDivisor = 0.25f;
+
+	const UTransformGizmoEditorSettings* TransformGizmoSettings = GetDefault<UTransformGizmoEditorSettings>();
+
+	constexpr float MinSizeMultiplier = 0.1f;
+
+	// Legacy Gizmo Size is an offset (0-based, default 1), in world units. New is a multiplier, so we need to compensate (1.0 based), such that scaling is in TransformGizmoSizeToMultiplierDivisor increments.
+	return FMath::Max(MinSizeMultiplier, AxisSizeMultiplier * (1.0f + (TransformGizmoSettings->TransformGizmoSize - 1.0f) * TransformGizmoSizeToMultiplierDivisor));
+}
+
+bool FTransformGizmoInteraction::IsCoordinateSystemEnabled(const ECoordSystem InCoordinateSystem) const
+{
+	return (AdditionalCoordinateSystems & (1 << InCoordinateSystem)) != 0;
+}
+
+void FTransformGizmoInteraction::SetCoordinateSystemEnabled(const ECoordSystem InCoordinateSystem, const bool bInIsEnabled)
+{
+	if (bInIsEnabled)
+	{
+		AdditionalCoordinateSystems |= (1 << InCoordinateSystem);
+	}
+	else
+	{
+		AdditionalCoordinateSystems &= ~(1 << InCoordinateSystem);
+	}
+}
+
 void UTransformGizmo::SetDisallowNegativeScaling(bool bDisallow)
 {
 	bDisallowNegativeScaling = bDisallow;
@@ -1305,15 +1351,21 @@ void UTransformGizmo::OnParametersChanged(const FGizmosParameters& InParameters)
 	}
 
 	DefaultRotateMode = InParameters.RotateMode;
+	
+	// Unlike the EditorTRSGizmo, these are fixed (used to be in the now-deleted FGizmoDevParameters)
+	constexpr float LineThicknessMultiplier = 1.0f;
+	constexpr float TranslateAxisLength = 70.0f;
+	constexpr float ScaleAxisLength = 70.0f;
+	constexpr float RotateAxisRadius = 70.0f;
 
 	// @note: The clamp values should match the property metadata (ClampMin)
 	Style.AxisSizeMultiplier = FMath::Max(0.01f, InParameters.AxisSizeMultiplier);
-	Style.LineThicknessMultiplier = FMath::Max(0.01f, InParameters.DevParameters.LineThicknessMultiplier);
-	Style.TranslateAxisLength = FMath::Max(1.0f, InParameters.DevParameters.TranslateAxisLength);
-	Style.RotateArcballSphereRadius = FMath::Max(1.0f, InParameters.DevParameters.RotateAxisRadius);
+	Style.LineThicknessMultiplier = FMath::Max(0.01f, LineThicknessMultiplier);
+	Style.TranslateAxisLength = FMath::Max(1.0f, TranslateAxisLength);
+	Style.RotateArcballSphereRadius = FMath::Max(1.0f, RotateAxisRadius);
 	Style.RotateAxisRadius = Style.RotateArcballSphereRadius;
-	Style.RotateCircleRadius =Style. RotateArcballSphereRadius;
-	Style.ScaleAxisLength = FMath::Max(1.0f, InParameters.DevParameters.ScaleAxisLength);
+	Style.RotateCircleRadius = Style.RotateArcballSphereRadius;
+	Style.ScaleAxisLength = FMath::Max(1.0f, ScaleAxisLength);
 
 	UpdateElements();
 }
diff --git a/Engine/Source/Editor/Experimental/EditorInteractiveToolsFramework/Private/EditorInteractiveGizmoManager.cpp b/Engine/Source/Editor/Experimental/EditorInteractiveToolsFramework/Private/EditorInteractiveGizmoManager.cpp
--- a/Engine/Source/Editor/Experimental/EditorInteractiveToolsFramework/Private/EditorInteractiveGizmoManager.cpp
+++ b/Engine/Source/Editor/Experimental/EditorInteractiveToolsFramework/Private/EditorInteractiveGizmoManager.cpp
@@ -33,6 +33,8 @@ class FCanvas;
 
 #define LOCTEXT_NAMESPACE "UEditorInteractiveGizmoManager"
 
+DEFINE_LOG_CATEGORY_STATIC(LogEditorInteractiveGizmoManager, Log, All);
+
 namespace GizmoManagerLocals
 {
 	static TOptional<FGizmosParameters> OptDefaultParameters;
@@ -42,23 +44,69 @@ namespace GizmoManagerLocals
 
 bool UEditorInteractiveGizmoManager::UsesNewTRSGizmos()
 {
-	if (const UTransformGizmoEditorSettings* TransformGizmoEditorSettings = GetDefault<UTransformGizmoEditorSettings>())
+	return GetTransformGizmoVersion() >= 1;
+}
+
+void UEditorInteractiveGizmoManager::SetUsesNewTRSGizmos(const bool bUseNewTRSGizmos)
+{
+	constexpr int32 LegacyGizmoVersion = 0;
+	constexpr int32 NewGizmoVersion = 2;
+	SetTransformGizmoVersion(bUseNewTRSGizmos ? NewGizmoVersion : LegacyGizmoVersion);
+}
+
+int32 UEditorInteractiveGizmoManager::GetTransformGizmoVersion()
+{
+	if (const UTransformGizmoEditorSettings* Settings = GetDefault<UTransformGizmoEditorSettings>())
 	{
-		return TransformGizmoEditorSettings->UsesNewTRSGizmo();
+		if (Settings->UsesNewTRSGizmo())
+		{
+			return 2; // New/5.8
+		}
+
+		if (Settings->UsesExperimentalGizmo())
+		{
+			return 1; // New/Experimental
+		}
 	}
 
-	return false;
+	return 0; // Legacy
 }
 
-void UEditorInteractiveGizmoManager::SetUsesNewTRSGizmos(const bool bUseNewTRSGizmos)
+void UEditorInteractiveGizmoManager::SetTransformGizmoVersion(const int32 InVersion)
 {
-	if (UTransformGizmoEditorSettings* const TransformGizmoEditorSettings = GetMutableDefault<UTransformGizmoEditorSettings>())
+	if (GetTransformGizmoVersion() == InVersion)
+	{
+		return;
+	}
+
+	if (!FMath::IsWithinInclusive(InVersion, 0, 2))
 	{
-		if (bUseNewTRSGizmos != TransformGizmoEditorSettings->UsesNewTRSGizmo())
+		UE_LOG(LogEditorInteractiveGizmoManager, Error, TEXT("Invalid GizmoVersion %d, expected value between 0 and 2 inclusive. No changes applied."), InVersion);
+		return;
+	}
+
+	if (UTransformGizmoEditorSettings* Settings = GetMutableDefault<UTransformGizmoEditorSettings>())
+	{
+		const bool bUseExperimentalGizmo = InVersion == 1;
+		const bool bUseNew58Gizmo = InVersion == 2;
+
+		// Ordering is important here, so we consider what transitions to what
+
+		const bool bFromExperimentalToNew = GetTransformGizmoVersion() == 1 && bUseNew58Gizmo;
+
+		if (bFromExperimentalToNew)
 		{
-			TransformGizmoEditorSettings->SetUseExperimentalGizmo(bUseNewTRSGizmos);
+			Settings->SetUseExperimentalGizmo(bUseExperimentalGizmo);
+			Settings->SetUseNewTRSGizmo(bUseNew58Gizmo);	
+		}
+		else // bFromNewToExperimental or Legacy
+		{
+			Settings->SetUseNewTRSGizmo(bUseNew58Gizmo);
+			Settings->SetUseExperimentalGizmo(bUseExperimentalGizmo);
 		}
 	}
+
+	// If legacy is specified, the two above will have been disabled, leaving us with legacy enabled
 }
 
 UEditorInteractiveGizmoManager::FOnUsesNewTRSGizmosChanged& UEditorInteractiveGizmoManager::OnUsesNewTRSGizmosChangedDelegate()
@@ -86,7 +134,7 @@ const TOptional<FGizmosParameters>& UEditorInteractiveGizmoManager::GetDefaultGi
 		GizmoManagerLocals::OptDefaultParameters = TransformGizmoEditorSettings->GizmosParameters;
 	}
 
-	static const TOptional<FGizmosParameters> Invalid;
+	static TOptional<FGizmosParameters> Invalid;
 	return GizmoManagerLocals::OptDefaultParameters.IsSet() ? GizmoManagerLocals::OptDefaultParameters : Invalid;
 }
 
@@ -95,8 +143,7 @@ bool UEditorInteractiveGizmoManager::IsExplicitModeEnabled()
 	return GetDefaultGizmosParameters().IsSet() ? GetDefaultGizmosParameters()->bEnableExplicit : false;
 }
 
-UEditorInteractiveGizmoManager::UEditorInteractiveGizmoManager() :
-	UInteractiveGizmoManager()
+UEditorInteractiveGizmoManager::UEditorInteractiveGizmoManager()
 {
 	Registry = NewObject<UEditorInteractiveGizmoRegistry>();
 	bShowEditorGizmos = UsesNewTRSGizmos();
diff --git a/Engine/Source/Editor/Experimental/EditorInteractiveToolsFramework/Private/EditorInteractiveGizmoSubsystem.cpp b/Engine/Source/Editor/Experimental/EditorInteractiveToolsFramework/Private/EditorInteractiveGizmoSubsystem.cpp
--- a/Engine/Source/Editor/Experimental/EditorInteractiveToolsFramework/Private/EditorInteractiveGizmoSubsystem.cpp
+++ b/Engine/Source/Editor/Experimental/EditorInteractiveToolsFramework/Private/EditorInteractiveGizmoSubsystem.cpp
@@ -2,7 +2,7 @@
 
 #include "EditorInteractiveGizmoSubsystem.h"
 
-#include "EditorGizmos/EditorTransformGizmoBuilder.h"
+#include "EditorGizmos/EditorTRSGizmoBuilder.h"
 #include "Engine/Engine.h"
 #include "Logging/LogMacros.h"
 #include "Misc/AssertionMacros.h"
@@ -45,8 +45,9 @@ void UEditorInteractiveGizmoSubsystem::Deinitialize()
 
 void UEditorInteractiveGizmoSubsystem::RegisterBuiltinEditorGizmoTypes()
 {
-	// Setup gizmo transform builder
-	RegisterTransformGizmoBuilder(NewObject<UEditorTransformGizmoBuilder>());
+	// Setup gizmo transform builder (UEditorTRSGizmoBuilder extends UEditorTransformGizmoBuilder
+	// and falls back to it when bUseEditorTRSGizmo is false)
+	RegisterTransformGizmoBuilder(NewObject<UEditorTRSGizmoBuilder>());
 
 	// Register built-in gizmo types here
 
diff --git a/Engine/Source/Editor/Experimental/EditorInteractiveToolsFramework/Private/EditorInteractiveToolsFrameworkModule.cpp b/Engine/Source/Editor/Experimental/EditorInteractiveToolsFramework/Private/EditorInteractiveToolsFrameworkModule.cpp
--- a/Engine/Source/Editor/Experimental/EditorInteractiveToolsFramework/Private/EditorInteractiveToolsFrameworkModule.cpp
+++ b/Engine/Source/Editor/Experimental/EditorInteractiveToolsFramework/Private/EditorInteractiveToolsFrameworkModule.cpp
@@ -3,37 +3,340 @@
 #include "EditorInteractiveToolsFrameworkModule.h"
 
 #include "ComponentSourceInterfaces.h"
+#include "Customizations/GizmoPerStateValueCustomization.h"
+#include "Customizations/TransformGizmoEditorSettingsCustomization.h"
 #include "EditorGizmos/GizmoRotationUtil.h"
-#include "Modules/ModuleManager.h"
+#include "EditorInteractiveGizmoManager.h"
+#include "Engine/Engine.h"
+#include "Framework/Application/SlateApplication.h"
+#include "PropertyEditorModule.h"
+#include "Styling/AppStyle.h"
 #include "Templates/UniquePtr.h"
 #include "Tools/EditorComponentSourceFactory.h"
+#include "TransformGizmoEditorSettings.h"
+#include "Widgets/Docking/SDockTab.h"
+#include "Widgets/SGizmoSettings.h"
+#include "Widgets/SGizmoTree.h"
+#include "WorkspaceMenuStructure.h"
+#include "WorkspaceMenuStructureModule.h"
 
 #define LOCTEXT_NAMESPACE "FEditorInteractiveToolsFrameworkModule"
 
+DEFINE_LOG_CATEGORY_STATIC(LogEditorInteractiveToolsFrameworkModule, Log, All);
+
+namespace UE::Editor::InteractiveToolsFramework::Private
+{
+	constexpr const TCHAR* GizmoVersionCVarName = TEXT("Gizmos.Version");
+}
+
 int32 FEditorInteractiveToolsFrameworkGlobals::RegisteredStaticMeshTargetFactoryKey = -1;
 
+const FLazyName FEditorInteractiveToolsFrameworkModule::GizmoSettingsTabId = TEXT("GizmoSettings");
+const FLazyName FEditorInteractiveToolsFrameworkModule::GizmoTreeTabId = TEXT("GizmoTree");
+
+// ---------------------------------------------------------------------------
+// TRS Gizmo CVar management (merged from EditorTRSGizmoModule)
+// ---------------------------------------------------------------------------
+
+void FEditorInteractiveToolsFrameworkModule::OnUsesNewTRSGizmosChanged(bool bInUsesNewTRSGizmos)
+{
+	const bool bWillEditorTRSGizmo = bInUsesNewTRSGizmos && GetDefault<UTransformGizmoEditorSettings>()->bUseEditorTRSGizmo;
+
+	static bool bUsesEditorTRSGizmo = false;
+	if (bUsesEditorTRSGizmo == bWillEditorTRSGizmo)
+	{
+		return;
+	}
+
+	bUsesEditorTRSGizmo = bWillEditorTRSGizmo;
+
+	constexpr const TCHAR* ITFGizmoEnabledCVars[] = {
+		TEXT("Editor.EnableITFCursorOverrideSupport"),
+		TEXT("DragTools.EnableITFTools"),
+		TEXT("ViewportInteractions.EnableITFInteractions"),
+		TEXT("Editor.UnboundedCursorDrag")
+	};
+
+	constexpr const TCHAR* ITFGizmoDisabledCVars[] = {
+		TEXT("Editor.AlwaysSnapPercentageBasedScale")
+	};
+
+	auto ToggleCVars = [](TConstArrayView<const TCHAR*> CVars, bool bEnable)
+	{
+		for (const TCHAR* CVarToToggle : CVars)
+		{
+			if (IConsoleVariable* CVar = IConsoleManager::Get().FindConsoleVariable(CVarToToggle))
+			{
+				if (CVar->GetBool() != bEnable)
+				{
+					CVar->Set(bEnable);
+				}
+			}
+		}
+	};
+
+	ToggleCVars(ITFGizmoEnabledCVars, bWillEditorTRSGizmo);
+	ToggleCVars(ITFGizmoDisabledCVars, !bWillEditorTRSGizmo);
+}
+
+void FEditorInteractiveToolsFrameworkModule::OnPostEngineInit()
+{
+	// Sync CVars to initial state and subscribe to future changes.
+	// Builder registration is handled by UEditorInteractiveGizmoSubsystem::RegisterBuiltinEditorGizmoTypes.
+	OnUsesNewTRSGizmosChanged(UEditorInteractiveGizmoManager::UsesNewTRSGizmos());
+	TRSGizmoChangeDelegate = UEditorInteractiveGizmoManager::OnUsesNewTRSGizmosChangedDelegate().AddStatic(
+		&FEditorInteractiveToolsFrameworkModule::OnUsesNewTRSGizmosChanged
+	);
+}
+
+// ---------------------------------------------------------------------------
+// Settings UI (merged from EditorTRSGizmoSettingsModule)
+// ---------------------------------------------------------------------------
+
+void FEditorInteractiveToolsFrameworkModule::RegisterConsoleCommands()
+{
+	ConsoleCommands.Add(IConsoleManager::Get().RegisterConsoleCommand(
+		TEXT("Gizmos.ShowSettings"),
+		TEXT("Displays a window listing hidden settings for Gizmos"),
+		FConsoleCommandWithArgsDelegate::CreateRaw(this, &FEditorInteractiveToolsFrameworkModule::ExecuteShowSettings),
+		ECVF_Default
+	));
+
+	ConsoleCommands.Add(IConsoleManager::Get().RegisterConsoleCommand(
+		TEXT("Gizmos.ShowTree"),
+		TEXT("Displays a window with various debug options for Gizmos"),
+		FConsoleCommandWithArgsDelegate::CreateRaw(this, &FEditorInteractiveToolsFrameworkModule::ExecuteShowGizmoTree),
+		ECVF_Default
+	));
+
+	using namespace UE::Editor::InteractiveToolsFramework::Private;
+	constexpr int32 DefaultGizmoVersion = 2; // 2 = New TRS Gizmo
+	ConsoleCommands.Add(IConsoleManager::Get().RegisterDelegatedConsoleVariable(
+		GizmoVersionCVarName,
+		DefaultGizmoVersion,
+		[&]() { return UEditorInteractiveGizmoManager::GetTransformGizmoVersion(); },
+		[&](int32 Version) { UEditorInteractiveGizmoManager::SetTransformGizmoVersion(Version); },
+		TEXT("0 = Legacy\n1 = Experimental (<5.8)\n2 = New (5.8+)"),
+		ECVF_Default
+	));
+}
+
+void FEditorInteractiveToolsFrameworkModule::UnregisterConsoleCommands()
+{
+	for (IConsoleObject* Cmd : ConsoleCommands)
+	{
+		IConsoleManager::Get().UnregisterConsoleObject(Cmd);
+	}
+
+	ConsoleCommands.Empty();
+}
+
+void FEditorInteractiveToolsFrameworkModule::RegisterCustomizations()
+{
+	FPropertyEditorModule& PropertyEditorModule = FModuleManager::GetModuleChecked<FPropertyEditorModule>("PropertyEditor");
+
+	PropertyEditorModule.RegisterCustomClassLayout(
+		UTransformGizmoEditorSettings::StaticClass()->GetFName(),
+		FOnGetDetailCustomizationInstance::CreateStatic(&FTransformGizmoEditorSettingsCustomization::MakeInstance));
+
+	// Notify PropertyEditor to pick up the customization if any detail panels were already created
+	PropertyEditorModule.NotifyCustomizationModuleChanged();
+
+	// Sections
+	{
+		// Shared
+		{
+			TSharedRef<FPropertySection> SharedSection = PropertyEditorModule.FindOrCreateSection(
+			UTransformGizmoEditorSettings::StaticClass()->GetFName(),
+			TEXT("Shared"),
+			FText::FromString(TEXT("Shared")),
+			-5);
+
+			SharedSection->AddCategory("Shared");
+		}
+
+		// Translate
+		{
+			TSharedRef<FPropertySection> TranslateSection = PropertyEditorModule.FindOrCreateSection(
+			UTransformGizmoEditorSettings::StaticClass()->GetFName(),
+				TEXT("Translate"),
+				FText::FromString(TEXT("Translate")),
+				-1);
+
+			TranslateSection->AddCategory("Translate");
+		}
+
+		// Rotate
+		{
+			TSharedRef<FPropertySection> RotateSection = PropertyEditorModule.FindOrCreateSection(
+			UTransformGizmoEditorSettings::StaticClass()->GetFName(),
+				TEXT("Rotate"),
+				FText::FromString(TEXT("Rotate")));
+
+			RotateSection->AddCategory("Rotate");
+		}
+
+		// Scale
+		{
+			TSharedRef<FPropertySection> ScaleSection = PropertyEditorModule.FindOrCreateSection(
+				UTransformGizmoEditorSettings::StaticClass()->GetFName(),
+				TEXT("Scale"),
+				FText::FromString(TEXT("Scale")),
+				1);
+
+			ScaleSection->AddCategory("Scale");
+		}
+	}
+}
+
+void FEditorInteractiveToolsFrameworkModule::UnregisterCustomizations()
+{
+	const FName PropertyEditorModuleName("PropertyEditor");
+	if (FModuleManager::Get().IsModuleLoaded(PropertyEditorModuleName))
+	{
+		if (FPropertyEditorModule* PropertyEditorModule = FModuleManager::Get().GetModulePtr<FPropertyEditorModule>(PropertyEditorModuleName))
+		{
+			const FName TransformGizmoEditorSettingsClassName = UTransformGizmoEditorSettings::StaticClass()->GetFName();
+			PropertyEditorModule->UnregisterCustomClassLayout(TransformGizmoEditorSettingsClassName);
+
+			PropertyEditorModule->NotifyCustomizationModuleChanged();
+		}
+	}
+}
+
+void FEditorInteractiveToolsFrameworkModule::RegisterTabs()
+{
+	TSharedRef<FGlobalTabmanager> GlobalTabManager = FGlobalTabmanager::Get();
+
+	GlobalTabManager->RegisterNomadTabSpawner(
+		GizmoSettingsTabId,
+		FOnSpawnTab::CreateLambda(
+			[this](const FSpawnTabArgs&)
+			{
+				TSharedRef<SDockTab> DockTab =
+					SNew(SDockTab)
+					.TabRole(ETabRole::NomadTab)
+					.Content()
+					[
+						SNew(SGizmoSettings)
+					];
+
+				return DockTab;
+			}
+		)
+	)
+	.SetDisplayName(LOCTEXT("GizmoSettingsTabName", "Gizmo Settings"))
+	.SetGroup(WorkspaceMenu::GetMenuStructure().GetLevelEditorCategory())
+	.SetIcon(FSlateIcon(FAppStyle::GetAppStyleSetName(), "LevelEditor.Tabs.Details"))
+	.SetAutoGenerateMenuEntry(false);
+
+	GlobalTabManager->RegisterNomadTabSpawner(
+		GizmoTreeTabId,
+		FOnSpawnTab::CreateLambda(
+			[this](const FSpawnTabArgs&)
+			{
+				TSharedRef<SDockTab> DockTab =
+					SNew(SDockTab)
+					.TabRole(ETabRole::NomadTab)
+					.Content()
+					[
+						SNew(SGizmoTree)
+					];
+
+				return DockTab;
+			}
+		)
+	)
+	.SetDisplayName(LOCTEXT("GizmoTreeTabName", "Gizmo Tree"))
+	.SetGroup(WorkspaceMenu::GetMenuStructure().GetLevelEditorCategory())
+	.SetIcon(FSlateIcon(FAppStyle::GetAppStyleSetName(), "LevelEditor.Tabs.Outliner"))
+	.SetAutoGenerateMenuEntry(false);
+}
+
+void FEditorInteractiveToolsFrameworkModule::UnregisterTabs()
+{
+	if (FSlateApplication::IsInitialized())
+	{
+		const TSharedRef<FGlobalTabmanager> GlobalTabManager = FGlobalTabmanager::Get();
+
+		GlobalTabManager->UnregisterNomadTabSpawner(GizmoSettingsTabId);
+		GlobalTabManager->UnregisterNomadTabSpawner(GizmoTreeTabId);
+
+		WeakGizmoSettingsTab.Reset();
+		WeakGizmoTreeTab.Reset();
+	}
+}
+
+void FEditorInteractiveToolsFrameworkModule::ExecuteShowSettings(const TArray<FString>& InArgs)
+{
+	if (!WeakGizmoSettingsTab.IsValid())
+	{
+		WeakGizmoSettingsTab = FGlobalTabmanager::Get()->TryInvokeTab(GizmoSettingsTabId.Resolve());
+	}
+}
+
+void FEditorInteractiveToolsFrameworkModule::ExecuteShowGizmoTree(const TArray<FString>& InArgs)
+{
+	if (!WeakGizmoTreeTab.IsValid())
+	{
+		WeakGizmoTreeTab = FGlobalTabmanager::Get()->TryInvokeTab(GizmoTreeTabId.Resolve());
+	}
+}
+
+// ---------------------------------------------------------------------------
+// Module startup / shutdown
+// ---------------------------------------------------------------------------
+
 void FEditorInteractiveToolsFrameworkModule::StartupModule()
 {
-	// This code will execute after your module is loaded into memory; the exact timing is specified in the .uplugin file per-module
-	FEditorInteractiveToolsFrameworkGlobals::RegisteredStaticMeshTargetFactoryKey 
+	FEditorInteractiveToolsFrameworkGlobals::RegisteredStaticMeshTargetFactoryKey
 		= AddComponentTargetFactory( TUniquePtr<FComponentTargetFactory>{new FStaticMeshComponentTargetFactory{} } );
 
 	// The Interactive Tools Framework needs to be explicitly loaded, at the very least to make sure that
 	// the InteractiveToolsSelectionStoreSubsystem gets initialized.
 	FModuleManager::Get().LoadModule(TEXT("InteractiveToolsFramework"));
 
-	// register relative transform interface
+	// Register relative transform interface
 	using namespace UE::GizmoRotationUtil;
 	FRelativeTransformInterfaceRegistry& Registry = FRelativeTransformInterfaceRegistry::Get();
 	Registry.RegisterDefaultInterfaces();
+
+	if (!IsRunningCommandlet())
+	{
+		// TRS gizmo builder registration and CVar sync (from EditorTRSGizmoModule)
+		if (GEngine && GEngine->IsInitialized())
+		{
+			OnPostEngineInit();
+		}
+		else
+		{
+			FCoreDelegates::GetOnPostEngineInit().AddRaw(this, &FEditorInteractiveToolsFrameworkModule::OnPostEngineInit);
+		}
+
+		// Settings UI (from EditorTRSGizmoSettingsModule)
+		RegisterConsoleCommands();
+		RegisterCustomizations();
+		RegisterTabs();
+	}
 }
 
 void FEditorInteractiveToolsFrameworkModule::ShutdownModule()
 {
-	// This function may be called during shutdown to clean up your module.  For modules that support dynamic reloading,
-	// we call this function before unloading the module.
+	if (!IsRunningCommandlet())
+	{
+		FCoreDelegates::GetOnPostEngineInit().RemoveAll(this);
+
+		UnregisterTabs();
+		UnregisterCustomizations();
+		UnregisterConsoleCommands();
+	}
+
+	if (TRSGizmoChangeDelegate.IsValid())
+	{
+		UEditorInteractiveGizmoManager::OnUsesNewTRSGizmosChangedDelegate().Remove(TRSGizmoChangeDelegate);
+	}
 }
 
 #undef LOCTEXT_NAMESPACE
-	
+
 IMPLEMENT_MODULE(FEditorInteractiveToolsFrameworkModule, EditorInteractiveToolsFramework)
diff --git a/Engine/Source/Editor/Experimental/EditorInteractiveToolsFramework/Private/TransformGizmoEditorSettings.cpp b/Engine/Source/Editor/Experimental/EditorInteractiveToolsFramework/Private/TransformGizmoEditorSettings.cpp
--- a/Engine/Source/Editor/Experimental/EditorInteractiveToolsFramework/Private/TransformGizmoEditorSettings.cpp
+++ b/Engine/Source/Editor/Experimental/EditorInteractiveToolsFramework/Private/TransformGizmoEditorSettings.cpp
@@ -18,6 +18,37 @@ UTransformGizmoEditorSettings::UTransformGizmoEditorSettings(const FObjectInitia
 	}
 }
 
+void UTransformGizmoEditorSettings::PostInitProperties()
+{
+	Super::PostInitProperties();
+	
+	// Freeze Camera CVar
+	{
+		constexpr const TCHAR* CVarName = TEXT("Gizmos.FreezeView");
+		if (IConsoleVariable* FoundCVar = IConsoleManager::Get().FindConsoleVariable(CVarName, false))
+		{
+			FoundCVar->Set(GizmosParameters.Debug.bFreezeCamera);
+		}
+	}
+
+	// Debug Toggle CVar
+	{
+		constexpr const TCHAR* CVarName = TEXT("Gizmos.DebugDraw");
+		if (IConsoleVariable* FoundCVar = IConsoleManager::Get().FindConsoleVariable(CVarName, false))
+		{
+			FoundCVar->Set(GizmosParameters.Debug.bDrawDebug);
+			FoundCVar->OnChangedDelegate().AddWeakLambda(this, [WeakThis = MakeWeakObjectPtr(this)](IConsoleVariable* InCVar)
+			{
+				if (WeakThis.IsValid())
+				{
+					WeakThis->GizmosParameters.Debug.bDrawDebug = InCVar->GetBool();
+					WeakThis->SaveConfig();
+				}
+			});
+		}
+	}
+}
+
 #if WITH_EDITOR
 void UTransformGizmoEditorSettings::PostEditChangeChainProperty(FPropertyChangedChainEvent& InPropertyChangedEvent)
 {
@@ -32,6 +63,22 @@ void UTransformGizmoEditorSettings::PostEditChangeChainProperty(FPropertyChanged
 
 			if (GizmosParametersName == ChangedPropertyName)
 			{
+				// Sync FreezeView CVar when debug settings change
+				if (const FEditPropertyChain& Chain = InPropertyChangedEvent.PropertyChain;
+					Chain.Num() >= 2)
+				{
+					static const FName FreezeCameraName = GET_MEMBER_NAME_CHECKED(FGizmoDebugSettings, bFreezeCamera);
+					if (const FProperty* LeafProperty = Chain.GetTail()->GetValue();
+						LeafProperty && LeafProperty->GetFName() == FreezeCameraName)
+					{
+						constexpr const TCHAR* CVarName = TEXT("Gizmos.FreezeView");
+						if (IConsoleVariable* FoundCVar = IConsoleManager::Get().FindConsoleVariable(CVarName, false))
+						{
+							FoundCVar->Set(GizmosParameters.Debug.bFreezeCamera);
+						}
+					}
+				}
+
 				BroadcastGizmosParametersChange();
 			}
 		}
@@ -43,6 +90,7 @@ void UTransformGizmoEditorSettings::PostEditChangeProperty(FPropertyChangedEvent
 	Super::PostEditChangeProperty(InPropertyChangedEvent);
 
 	static const FName UseExperimentalGizmoName = GET_MEMBER_NAME_CHECKED(UTransformGizmoEditorSettings, bUseExperimentalGizmo);
+	static const FName UseEditorTRSGizmoName = GET_MEMBER_NAME_CHECKED(UTransformGizmoEditorSettings, bUseEditorTRSGizmo);
 	static const FName EnableArcballRotateName = GET_MEMBER_NAME_CHECKED(UTransformGizmoEditorSettings, bEnableArcballRotate);
 	static const FName EnableScreenRotateName = GET_MEMBER_NAME_CHECKED(UTransformGizmoEditorSettings, bEnableScreenRotate);
 	static const FName EnableAxisDrawingName = GET_MEMBER_NAME_CHECKED(UTransformGizmoEditorSettings, bEnableAxisDrawing);
@@ -59,6 +107,14 @@ void UTransformGizmoEditorSettings::PostEditChangeProperty(FPropertyChangedEvent
 	{
 		BroadcastNewTRSGizmoChange();
 	}
+	else if (ChangedPropertyName == UseEditorTRSGizmoName)
+	{
+		// Force gizmo rebuild by cycling experimental gizmos off/on.
+		// This tears down existing gizmos and recreates them with the new builder path.
+		const bool bWasUsingExperimental = bUseExperimentalGizmo;
+		SetUseExperimentalGizmo(false);
+		SetUseExperimentalGizmo(bWasUsingExperimental);
+	}
 	else if (ChangedPropertyName == EnableArcballRotateName)
 	{
 		if (ULevelEditorViewportSettings* const ViewportSettings = GetMutableDefault<ULevelEditorViewportSettings>())
@@ -122,16 +178,37 @@ void UTransformGizmoEditorSettings::SetUseExperimentalGizmo(bool bInUseExperimen
 	}
 }
 
+void UTransformGizmoEditorSettings::SetUseNewTRSGizmo(bool bInUseNewTRSGizmo)
+{
+	if (bInUseNewTRSGizmo != bUseEditorTRSGizmo)
+	{
+		bUseEditorTRSGizmo = bInUseNewTRSGizmo;
+		SaveConfig();
+
+		BroadcastNewTRSGizmoChange();
+
+		if (GEditor)
+		{
+			GEditor->RedrawAllViewports();
+		}
+	}
+}
+
 bool UTransformGizmoEditorSettings::UsesLegacyGizmo() const
 {
-	return !UsesNewTRSGizmo();
+	return !(UsesExperimentalGizmo() || UsesNewTRSGizmo());
 }
 
-bool UTransformGizmoEditorSettings::UsesNewTRSGizmo() const
+bool UTransformGizmoEditorSettings::UsesExperimentalGizmo() const
 {
 	return bUseExperimentalGizmo;
 }
 
+bool UTransformGizmoEditorSettings::UsesNewTRSGizmo() const
+{
+	return bUseEditorTRSGizmo;
+}
+
 void UTransformGizmoEditorSettings::SetGizmosParameters(const FGizmosParameters& InGizmosParameters)
 {
 	GizmosParameters = InGizmosParameters;
@@ -160,7 +237,7 @@ void UTransformGizmoEditorSettings::SetTransformGizmoSize(float InTransformGizmo
 
 void UTransformGizmoEditorSettings::BroadcastNewTRSGizmoChange() const
 {
-	UEditorInteractiveGizmoManager::OnUsesNewTRSGizmosChangedDelegate().Broadcast(UsesNewTRSGizmo());
+	UEditorInteractiveGizmoManager::OnUsesNewTRSGizmosChangedDelegate().Broadcast(UsesExperimentalGizmo() || UsesNewTRSGizmo());
 }
 
 void UTransformGizmoEditorSettings::BroadcastGizmosParametersChange() const
diff --git a/Engine/Source/Editor/Experimental/EditorInteractiveToolsFramework/Public/EditorGizmos/TransformGizmo.h b/Engine/Source/Editor/Experimental/EditorInteractiveToolsFramework/Public/EditorGizmos/TransformGizmo.h
--- a/Engine/Source/Editor/Experimental/EditorInteractiveToolsFramework/Public/EditorGizmos/TransformGizmo.h
+++ b/Engine/Source/Editor/Experimental/EditorInteractiveToolsFramework/Public/EditorGizmos/TransformGizmo.h
@@ -10,6 +10,9 @@
 #include "CoreMinimal.h"
 #include "BaseGizmos/GizmoElementShared.h"
 #include "EditorGizmos/TransformGizmoInterfaces.h"
+#include "GizmoElementRotate.h"
+#include "GizmoElementScale.h"
+#include "GizmoElementTranslate.h"
 #include "InputState.h"
 #include "InteractiveGizmo.h"
 #include "InteractiveToolChange.h"
@@ -107,34 +110,183 @@ namespace UE::Editor::InteractiveToolsFramework
 	struct FTransformGizmoAccessor;
 }
 
-UENUM()
-namespace EAxisRotateMode
+/** Controls the (fixed) visibility of various visual elements in the TransformGizmo. */
+UENUM(meta = (Bitflags, UseEnumValuesAsMaskValuesInEditor = "true"))
+enum class ETransformGizmoShowFlags : uint32
 {
-	enum Type : uint8
-	{
-		Pull = 0,
-		Arc = 1,
-		ScreenArc = 2, // Like Arc, but in screen space, rather than aligned to the associated axis
-	};
-}
+	None = 0 UMETA(Hidden),
+	DeltaLabel = 1 << 0,
+	DeltaLine = 1 << 1,
+	DeltaArc = 1 << 2,
+	DeltaOrigin = 1 << 3,
 
-/** @todo: To be removed for release, for dev/tweaking only. */
-USTRUCT()
-struct FGizmoDevParameters
+	All = DeltaLabel | DeltaLine | DeltaArc | DeltaOrigin
+};
+ENUM_CLASS_FLAGS(ETransformGizmoShowFlags);
+
+/** Contains all styling properties for TransformGizmo. */
+USTRUCT(MinimalAPI)
+struct FTransformGizmoStyle
 {
 	GENERATED_BODY()
 
-	UPROPERTY(EditAnywhere, Category = "Transform Gizmo", meta = (ClampMin = 0.01f))
+public:
+	PRAGMA_DISABLE_DEPRECATION_WARNINGS
+	UE_API FTransformGizmoStyle();
+	FTransformGizmoStyle(const FTransformGizmoStyle&) = default;
+	FTransformGizmoStyle& operator=(const FTransformGizmoStyle&) = default;
+	PRAGMA_ENABLE_DEPRECATION_WARNINGS
+
+	/** Clamps and otherwise corrects values within the style. */
+	UE_API void MakeValid();
+
+public:
+	// Uniforms
+
+	UE_DEPRECATED(5.6, "Use AxisDisplayInfo::GetAxisColor(EAxisList::X) instead")
+	static constexpr FLinearColor AxisColorX = FLinearColor(0.594f, 0.0197f, 0.0f);
+	UE_DEPRECATED(5.6, "Use AxisDisplayInfo::GetAxisColor(EAxisList::Y) instead")
+	static constexpr FLinearColor AxisColorY = FLinearColor(0.1349f, 0.3959f, 0.0f);
+	UE_DEPRECATED(5.6, "Use AxisDisplayInfo::GetAxisColor(EAxisList::Z) instead")
+	static constexpr FLinearColor AxisColorZ = FLinearColor(0.0251f, 0.207f, 0.85f);
+
+public:
+	// Varying
+
+#pragma region Shared
+	/** Applies to axis lengths and radii. Use the GetModifiedAxisSizeMultiplier to account for various other project/editor settings. */
+	UPROPERTY(EditAnywhere, Category = "Shared", meta = (ClampMin = 0.1f, DisplayPriority = -1))
+	float AxisSizeMultiplier = 1.0f;
+
+	/** Applies to arrow heads, planar and uniform handles. */
+	UPROPERTY(EditAnywhere, Category = "Shared", meta = (ClampMin = 0.01f, DisplayPriority = -1))
+	float HandleSizeMultiplier = 1.0f;
+
+	/** Applies to both translate and scale planar handles. */
+	UPROPERTY(EditAnywhere, Category = "Shared", meta = (ClampMin = 0.01f, DisplayPriority = -1))
+	float PlanarAxisOffsetFromOrigin = 50.0f;
+
+	/** Used to prevent sub-pixel line sizes. */
+	UPROPERTY(EditAnywhere, Category = "Shared", meta = (ClampMin = 0.01f))
+	float MinLineThickness = 1.0f;
+
+	UPROPERTY(EditAnywhere, Category = "Shared", meta = (ClampMin = 0.01f))
 	float LineThicknessMultiplier = 1.0f;
 
-	UPROPERTY(EditAnywhere, Category = "Transform Gizmo", meta = (ClampMin = 1.0f))
-	float TranslateAxisLength = 70.0f;
+	UPROPERTY(EditAnywhere, Category = "Shared", meta = (ClampMin = 0.01f))
+	float HoverLineThicknessMultiplier = 1.5f;
+
+	/** The current color to use. */
+	UE_DEPRECATED(5.7, "Use HoverColor instead")
+	UPROPERTY(meta = (HideAlphaChannel, DeprecatedProperty, DeprecationMessage = "Use HoverColor instead."))
+	FLinearColor CurrentColor = FLinearColor(1.0f, 1.0f, 0.0f);
+
+	/** The color to use on hover. */
+	UPROPERTY(EditAnywhere, Category = "Shared", meta = (HideAlphaChannel))
+	FLinearColor HoverColor = FLinearColor(1.0f, 0.552619f, 0.024453f);
+
+	/** The color to use on select. */
+	UPROPERTY(EditAnywhere, Category = "Shared", meta = (HideAlphaChannel))
+	FLinearColor SelectColor = FLinearColor(0.8f, 0.8f, 0.0f);
+
+	/** The color to use on interact. */
+	UPROPERTY(EditAnywhere, Category = "Shared", meta = (HideAlphaChannel))
+	FLinearColor InteractColor = FLinearColor(1.0f, 1.0f, 0.0f);
+
+	/** The color to use on subdue. Depending on the effect, we may only use the alpha value. */
+	UPROPERTY(EditAnywhere, Category = "Shared")
+	FLinearColor SubdueColor = FLinearColor::Black.CopyWithNewOpacity(0.7f);
+
+	UPROPERTY(EditAnywhere, Category = "Shared", meta = (HideAlphaChannel))
+	FLinearColor GreyColor = FLinearColor(0.6f, 0.6f, 0.6f);
+
+	UPROPERTY(EditAnywhere, Category = "Shared", meta = (HideAlphaChannel))
+	FLinearColor WhiteColor = FLinearColor(0.9f, 0.9f, 0.9f);
+
+	/** Whether to show shading on the gizmo, or unlit (if false). */
+	UPROPERTY(EditAnywhere, Category = "Shared")
+	bool bUseShading = true;
 
-	UPROPERTY(EditAnywhere, Category = "Transform Gizmo", meta = (ClampMin = 1.0f))
-	float ScaleAxisLength = 70.0f;
+	UPROPERTY(EditAnywhere, Category = "Shared", meta = (EditCondition = "bUseShading", ClampMin = 0.0f, ClampMax = 2.0f))
+	float ShadingAmbient = 0.8f;
 
-	UPROPERTY(EditAnywhere, Category = "Transform Gizmo", meta = (ClampMin = 1.0f))
-	float RotateAxisRadius = 70.0f;
+	UPROPERTY(EditAnywhere, Category = "Shared", meta = (EditCondition = "bUseShading", ClampMin = 0.0f, ClampMax = 1.0f))
+	float ShadingSpecularGlossiness = 0.1f;
+
+	/* Pixel hit distance threshold - where the default is generally 7px. */
+	UPROPERTY(EditAnywhere, Category = "Shared", meta = (ClampMin = 0.1f))
+	float PixelHitDistanceThreshold = 15.0f;
+
+	/* Minimum pixel hit distance threshold. */
+	UPROPERTY(EditAnywhere, Category = "Shared", meta = (ClampMin = 0.0f))
+	float MinimumPixelHitDistanceThreshold = 0.0f;
+
+	/** Represents the background color for text labels. */
+	UPROPERTY(EditAnywhere, Category = "Shared")
+	FLinearColor DeltaTextBackgroundColor = FLinearColor(0.0f, 0.0f, 0.0f, 0.65f);
+
+	/** Synchronizes the cursor size to that specified for the platform/OS. If false, CursorSize is used. */
+	UPROPERTY(EditAnywhere, Category = "Shared")
+	bool bUsePlatformCursorSize = true;
+
+	/** Defines the size of the custom cursor. */
+	UPROPERTY(EditAnywhere, Category = "Shared", meta = (EditCondition = "!bUsePlatformCursorSize"))
+	float CursorSize = 16.0f;
+
+	/** Sets the fixed visibility of various elements. */
+	UPROPERTY(EditAnywhere, Category = "Shared", meta = (Bitmask, BitmaskEnum = "/Script/EditorInteractiveToolsFramework.ETransformGizmoShowFlags"))
+	uint32 ShowFlags = static_cast<uint32>(ETransformGizmoShowFlags::All);
+#pragma endregion Shared
+
+#pragma region Translate
+	UPROPERTY(EditAnywhere, Category = "Translate", meta = (ShowOnlyInnerProperties, DisplayPriority = 10))
+	FGizmoElementTranslateStyle TranslateStyle;
+#pragma endregion Translate
+
+#pragma region Rotate
+	UPROPERTY(EditAnywhere, Category = "Rotate", meta = (ShowOnlyInnerProperties, DisplayPriority = 11))
+	FGizmoElementRotateStyle RotateStyle;
+#pragma endregion Rotate
+
+#pragma region Scale
+	UPROPERTY(EditAnywhere, Category = "Scale", meta = (ShowOnlyInnerProperties, DisplayPriority = 12))
+	FGizmoElementScaleStyle ScaleStyle;
+#pragma endregion Scale
+
+public:
+#pragma region Shared
+	/** Gets the AxisSizeMultiplier, modified by project/editor level settings, etc., or the given Modifier if specified. */
+	UE_API float GetModifiedAxisSizeMultiplier(const TOptional<TUniqueFunction<float(const float)>>& InModifierOverride = { }) const;
+#pragma endregion Shared
+};
+
+USTRUCT(MinimalAPI)
+struct FTransformGizmoInteraction
+{
+	GENERATED_BODY()
+
+	/** Whether to allow handle selection to persist and reference for indirect operations. */ 
+	UPROPERTY()
+	bool bPersistHandleSelection = false;
+
+	/** Stores additional coordinate systems enabled for this gizmo. By default, all are available. */
+	UPROPERTY(EditAnywhere, Category = "Shared", meta = (Bitmask, BitmaskEnum = "/Script/EditorFramework.ECoordSystem"))
+	int32 AdditionalCoordinateSystems = INDEX_NONE;
+
+	UPROPERTY(EditAnywhere, Category = "Translate", meta = (ShowOnlyInnerProperties, DisplayPriority = 10))
+	FGizmoElementTranslateInteraction TranslateInteraction;
+
+	UPROPERTY(EditAnywhere, Category = "Rotate", meta = (ShowOnlyInnerProperties, DisplayPriority = 11))
+	FGizmoElementRotateInteraction RotateInteraction;
+
+	UPROPERTY(EditAnywhere, Category = "Scale", meta = (ShowOnlyInnerProperties, DisplayPriority = 12))
+	FGizmoElementScaleInteraction ScaleInteraction;
+
+	/** Returns true if the given coordinate system is enabled for this gizmo. */
+	UE_API bool IsCoordinateSystemEnabled(const ECoordSystem InCoordinateSystem) const;
+
+	/** Sets whether the given coordinate system is enabled for this gizmo. */
+	UE_API void SetCoordinateSystemEnabled(const ECoordSystem InCoordinateSystem, const bool bInIsEnabled);
 };
 
 USTRUCT()
@@ -192,8 +344,14 @@ struct FGizmosParameters
 	UPROPERTY(EditAnywhere, Category=GizmoInteraction)
 	bool bScreenSpaceNudge = true;
 
-	UPROPERTY(EditAnywhere, Category = "Transform Gizmo|Development", meta=(ShowOnlyInnerProperties))
-	FGizmoDevParameters DevParameters;
+	UPROPERTY(Config, EditAnywhere, Category = "Appearance")
+	FTransformGizmoStyle Style;
+
+	UPROPERTY(Config, EditAnywhere, Category = "Interaction")
+	FTransformGizmoInteraction Interaction;
+
+	UPROPERTY(Config, EditAnywhere, Category = "Debug", meta=(HideInDetailPanel))
+	FGizmoDebugSettings Debug;
 };
 
 /**
diff --git a/Engine/Source/Editor/Experimental/EditorInteractiveToolsFramework/Public/EditorGizmos/TransformGizmoInterfaces.h b/Engine/Source/Editor/Experimental/EditorInteractiveToolsFramework/Public/EditorGizmos/TransformGizmoInterfaces.h
--- a/Engine/Source/Editor/Experimental/EditorInteractiveToolsFramework/Public/EditorGizmos/TransformGizmoInterfaces.h
+++ b/Engine/Source/Editor/Experimental/EditorInteractiveToolsFramework/Public/EditorGizmos/TransformGizmoInterfaces.h
@@ -23,6 +23,17 @@ enum class EGizmoTransformMode : uint8
 	Max
 };
 
+UENUM()
+namespace EAxisRotateMode
+{
+	enum Type : uint8
+	{
+		Pull = 0,
+		Arc = 1,
+		ScreenArc = 2, // Like Arc, but in screen space, rather than aligned to the associated axis
+	};
+}
+
 UENUM()
 enum class EGizmoTransformScaleType : uint8
 {
diff --git a/Engine/Source/Editor/Experimental/EditorInteractiveToolsFramework/Public/EditorInteractiveGizmoManager.h b/Engine/Source/Editor/Experimental/EditorInteractiveToolsFramework/Public/EditorInteractiveGizmoManager.h
--- a/Engine/Source/Editor/Experimental/EditorInteractiveToolsFramework/Public/EditorInteractiveGizmoManager.h
+++ b/Engine/Source/Editor/Experimental/EditorInteractiveToolsFramework/Public/EditorInteractiveGizmoManager.h
@@ -162,7 +162,19 @@ public:
 	static UE_API void SetUsesNewTRSGizmos(const bool bUseNewTRSGizmos);
 
 	/**
-	 * Delegate to notify from bUseNewTRSGizmos changes  
+	 * Returns the current gizmo version. 0 = legacy, 1 = TransformGizmo, 2 = EditorTRSGizmo.
+	 * Reads from existing settings (bUseExperimentalGizmo, bUseEditorTRSGizmo).
+	 */
+	static UE_API int32 GetTransformGizmoVersion();
+
+	/**
+	 * Sets the gizmo version by writing to the existing settings.
+	 * Writes to existing settings (bUseExperimentalGizmo, bUseEditorTRSGizmo).
+	 */
+	static UE_API void SetTransformGizmoVersion(const int32 InVersion);
+
+	/**
+	 * Delegate to notify from bUseNewTRSGizmos changes
 	 */
 	DECLARE_MULTICAST_DELEGATE_OneParam(FOnUsesNewTRSGizmosChanged, const bool bUseNewTRSGizmos);
 	static UE_API FOnUsesNewTRSGizmosChanged& OnUsesNewTRSGizmosChangedDelegate();
diff --git a/Engine/Source/Editor/Experimental/EditorInteractiveToolsFramework/Public/EditorInteractiveToolsFrameworkModule.h b/Engine/Source/Editor/Experimental/EditorInteractiveToolsFramework/Public/EditorInteractiveToolsFrameworkModule.h
--- a/Engine/Source/Editor/Experimental/EditorInteractiveToolsFramework/Public/EditorInteractiveToolsFrameworkModule.h
+++ b/Engine/Source/Editor/Experimental/EditorInteractiveToolsFramework/Public/EditorInteractiveToolsFrameworkModule.h
@@ -9,6 +9,7 @@
 
 #define UE_API EDITORINTERACTIVETOOLSFRAMEWORK_API
 
+class SDockTab;
 
 class FEditorInteractiveToolsFrameworkGlobals
 {
@@ -25,6 +26,34 @@ public:
 	/** IModuleInterface implementation */
 	virtual void StartupModule() override;
 	virtual void ShutdownModule() override;
+
+private:
+
+	// --- TRS Gizmo CVar management (from EditorTRSGizmoModule) ---
+
+	static void OnUsesNewTRSGizmosChanged(bool bInUsesNewTRSGizmos);
+	void OnPostEngineInit();
+
+	FDelegateHandle TRSGizmoChangeDelegate;
+
+	// --- Settings UI (from EditorTRSGizmoSettingsModule) ---
+
+	void RegisterConsoleCommands();
+	void UnregisterConsoleCommands();
+	void RegisterCustomizations();
+	void UnregisterCustomizations();
+	void RegisterTabs();
+	void UnregisterTabs();
+
+	void ExecuteShowSettings(const TArray<FString>& InArgs);
+	void ExecuteShowGizmoTree(const TArray<FString>& InArgs);
+
+	static const FLazyName GizmoSettingsTabId;
+	static const FLazyName GizmoTreeTabId;
+
+	TArray<IConsoleObject*> ConsoleCommands;
+	TWeakPtr<SDockTab> WeakGizmoSettingsTab;
+	TWeakPtr<SDockTab> WeakGizmoTreeTab;
 };
 
 #undef UE_API
diff --git a/Engine/Source/Editor/Experimental/EditorInteractiveToolsFramework/Public/TransformGizmoEditorSettings.h b/Engine/Source/Editor/Experimental/EditorInteractiveToolsFramework/Public/TransformGizmoEditorSettings.h
--- a/Engine/Source/Editor/Experimental/EditorInteractiveToolsFramework/Public/TransformGizmoEditorSettings.h
+++ b/Engine/Source/Editor/Experimental/EditorInteractiveToolsFramework/Public/TransformGizmoEditorSettings.h
@@ -2,6 +2,8 @@
 
 #pragma once
 
+#include "EditorGizmos/EditorGizmoElementShared.h"
+#include "EditorGizmos/EditorTRSGizmo.h"
 #include "Engine/DeveloperSettings.h"
 #include "EditorGizmos/TransformGizmo.h"
 
@@ -18,6 +20,7 @@ public:
 	UE_API UTransformGizmoEditorSettings(const FObjectInitializer& ObjectInitializer);
 
 #if WITH_EDITOR
+	UE_API virtual void PostInitProperties() override;
 	UE_API virtual void PostEditChangeChainProperty(FPropertyChangedChainEvent& InPropertyChangedEvent) override;
 	UE_API virtual void PostEditChangeProperty(FPropertyChangedEvent& InPropertyChangedEvent) override;
 	UE_API virtual void PostReloadConfig(class FProperty* PropertyThatWasLoaded) override;
@@ -56,16 +59,22 @@ public:
 	//~ End Transform Gizmo Category
 
 	//~ Begin Experimental Gizmo Category
-	/** If true, the new TRS gizmos will be used. */
-	UPROPERTY(Config, EditAnywhere, Category = "Experimental", DisplayName = "Enable Experimental Gizmo")
+	/** If true, the new TRS gizmos will be used. UI controls set both flags together; the intermediate state is CVar-only. */
+	UPROPERTY(Config, EditAnywhere, Category = "Experimental", Getter = "UsesExperimentalGizmo", Setter = "SetUseExperimentalGizmo")
 	bool bUseExperimentalGizmo = true;
 
-	UPROPERTY(Config, EditAnywhere, Category = "Experimental", meta = (ShowOnlyInnerProperties),  meta = (EditCondition = "bUseExperimentalGizmo"))
+	UPROPERTY(Config, EditAnywhere, Category = "Experimental", Getter = "UsesNewTRSGizmo", Setter = "SetUseNewTRSGizmo")
+	bool bUseEditorTRSGizmo = true;
+
+	UPROPERTY(Config, EditAnywhere, Category = "Experimental", meta = (ShowOnlyInnerProperties, EditCondition = "bUseExperimentalGizmo"))
 	FGizmosParameters GizmosParameters;
 	//~ End Experimental Gizmo Category
 
 	UE_API void SetUseExperimentalGizmo(bool bInUseExperimentalGizmo);
+	UE_API void SetUseNewTRSGizmo(bool bInUseNewTRSGizmo);
+
 	UE_API bool UsesLegacyGizmo() const;
+	UE_API bool UsesExperimentalGizmo() const;
 	UE_API bool UsesNewTRSGizmo() const;
 
 	UE_API void SetGizmosParameters(const FGizmosParameters& InGizmosParameters);
