De-duplicated Log declarations

#cis 1273512

#Synced-CL 52057193

diff --git a/Engine/Source/Editor/Experimental/EditorInteractiveToolsFramework/Private/EditorGizmos/EditorTRSGizmo.cpp b/Engine/Source/Editor/Experimental/EditorInteractiveToolsFramework/Private/EditorGizmos/EditorTRSGizmo.cpp
--- a/Engine/Source/Editor/Experimental/EditorInteractiveToolsFramework/Private/EditorGizmos/EditorTRSGizmo.cpp
+++ b/Engine/Source/Editor/Experimental/EditorInteractiveToolsFramework/Private/EditorGizmos/EditorTRSGizmo.cpp
@@ -65,7 +65,7 @@
 
 #define LOCTEXT_NAMESPACE "UEditorTRSGizmo"
 
-DEFINE_LOG_CATEGORY_STATIC(LogTransformGizmo, Log, All);
+DEFINE_LOG_CATEGORY_STATIC(LogEditorTRSGizmo, Log, All);
 
 ENUM_RANGE_BY_COUNT(ETransformGizmoPartIdentifier, ETransformGizmoPartIdentifier::Max);
 
@@ -298,7 +298,7 @@ namespace UE::Editor::InteractiveToolsFramework
 			{
 				if constexpr (bLogStateChanges)
 				{
-					UE_LOG(LogTransformGizmo, Verbose, TEXT("[%lli] %s: %s"),
+					UE_LOG(LogEditorTRSGizmo, Verbose, TEXT("[%lli] %s: %s"),
 						GFrameCounter,
 						*InStateName.ToString(),
 						bInStateToggle ? TEXT("On") : TEXT("Off"));
@@ -3674,7 +3674,7 @@ void UEditorTRSGizmo::InitializeInteractionSign()
 		ViewToAlignedSign = UE::Math::TIntVector3<int8>(1);
 	}
 
-	UE_LOG(LogTransformGizmo, Verbose,
+	UE_LOG(LogEditorTRSGizmo, Verbose,
 		TEXT("InteractionToAlignedSign: %s"),
 		*ViewToAlignedSign.ToString());
 }
@@ -4280,7 +4280,7 @@ void UEditorTRSGizmo::UpdateRotateAxisDelta(UGizmoElementRotateAxis* InElement)
 	DeltaParameters.CursorLocation = WorldCursorPos;
 	DeltaParameters.CursorLocation2D = InteractionPlanarCurrPoint2D;
 
-	UE_CLOG(bLogRotation, LogTransformGizmo, Verbose, TEXT("InteractionDelta: %.1f"), FMath::RadiansToDegrees(InteractionDelta.X));
+	UE_CLOG(bLogRotation, LogEditorTRSGizmo, Verbose, TEXT("InteractionDelta: %.1f"), FMath::RadiansToDegrees(InteractionDelta.X));
 
 	InElement->UpdateDelta(DeltaParameters);
 }
@@ -4521,7 +4521,7 @@ void UEditorTRSGizmo::ResetInteractingStates(const EGizmoTransformMode InMode)
 {
 	if constexpr (UE::Editor::InteractiveToolsFramework::Private::TransformGizmoLocals::bLogStateChanges)
 	{
-		UE_LOG(LogTransformGizmo, Verbose, TEXT("ResetInteractingStates"));
+		UE_LOG(LogEditorTRSGizmo, Verbose, TEXT("ResetInteractingStates"));
 	}
 
 	ETransformGizmoPartIdentifier IdBegin = ETransformGizmoPartIdentifier::Default;
@@ -4581,7 +4581,7 @@ void UEditorTRSGizmo::ResetSelectedStates(const EGizmoTransformMode InMode)
 {
 	if constexpr (UE::Editor::InteractiveToolsFramework::Private::TransformGizmoLocals::bLogStateChanges)
 	{
-		UE_LOG(LogTransformGizmo, Verbose, TEXT("ResetSelectedStates"));
+		UE_LOG(LogEditorTRSGizmo, Verbose, TEXT("ResetSelectedStates"));
 	}
 
 	if (GizmoLocals::DoDebugDraw()
@@ -4632,7 +4632,7 @@ void UEditorTRSGizmo::ResetSubdueStates(const EGizmoTransformMode InMode)
 {
 	if constexpr (UE::Editor::InteractiveToolsFramework::Private::TransformGizmoLocals::bLogStateChanges)
 	{
-		UE_LOG(LogTransformGizmo, Verbose, TEXT("ResetSubdueStates"));
+		UE_LOG(LogEditorTRSGizmo, Verbose, TEXT("ResetSubdueStates"));
 	}
 
 	ETransformGizmoPartIdentifier IdBegin = ETransformGizmoPartIdentifier::Default;
@@ -4702,7 +4702,7 @@ void UEditorTRSGizmo::EndDeltas()
 
 void UEditorTRSGizmo::OnClickPress(const FInputDeviceRay& PressPos)
 {
-	UE_LOG(LogTransformGizmo, Verbose, TEXT("[%lli] OnClickPress"), GFrameCounter);
+	UE_LOG(LogEditorTRSGizmo, Verbose, TEXT("[%lli] OnClickPress"), GFrameCounter);
 
 	PendingDragFunction.Reset();
 
@@ -4810,7 +4810,7 @@ FInputRayHit UEditorTRSGizmo::CanBeginSingleClickAndDragSequence(const FInputDev
 
 void UEditorTRSGizmo::OnDragStart(const FInputDeviceRay& InDragPos)
 {
-	UE_LOG(LogTransformGizmo, Verbose, TEXT("[%lli] OnDragStart"), GFrameCounter);
+	UE_LOG(LogEditorTRSGizmo, Verbose, TEXT("[%lli] OnDragStart"), GFrameCounter);
 
 	UpdateCursorRotation(InDragPos);
 
@@ -5012,7 +5012,7 @@ void UEditorTRSGizmo::OnTerminateDragSequence()
 
 void UEditorTRSGizmo::BeginTransformEditSequence()
 {
-	UE_LOG(LogTransformGizmo, Verbose, TEXT("[%lli] BeginTransformEditSequence"), GFrameCounter);
+	UE_LOG(LogEditorTRSGizmo, Verbose, TEXT("[%lli] BeginTransformEditSequence"), GFrameCounter);
 
 	// Update Debug
 	{
@@ -5041,7 +5041,7 @@ void UEditorTRSGizmo::BeginTransformEditSequence()
 	StartTransform = GetActiveTransform();
 	CurrentTransform = StartTransform;
 
-	UE_LOG(LogTransformGizmo, Verbose, TEXT("StartTransform.Scale: %.1f"), StartTransform.GetScale3D().X);
+	UE_LOG(LogEditorTRSGizmo, Verbose, TEXT("StartTransform.Scale: %.1f"), StartTransform.GetScale3D().X);
 
 	// Camera follows starting: cache cursor position
 	if (bCameraFollowsMovement)
@@ -5058,7 +5058,7 @@ void UEditorTRSGizmo::BeginTransformEditSequence()
 
 void UEditorTRSGizmo::EndTransformEditSequence()
 {
-	UE_LOG(LogTransformGizmo, Verbose, TEXT("[%lli] EndTransformEditSequence"), GFrameCounter);
+	UE_LOG(LogEditorTRSGizmo, Verbose, TEXT("[%lli] EndTransformEditSequence"), GFrameCounter);
 
 	// Update Debug
 	{
@@ -6043,8 +6043,8 @@ void UEditorTRSGizmo::OnClickPressScale(const FInputDeviceRay& PressPos)
 			InteractionReferencePointOffsetDistance = InteractionPlanarStartDistance - InteractionDeltaDivisor;
 		}
 
-		UE_CLOG(bLogScale, LogTransformGizmo, Verbose, TEXT("[%lli] BeginCapture: Offset: %.1f"), GFrameCounter, InteractionReferencePointOffsetDistance);
-		UE_CLOG(bLogScale, LogTransformGizmo, Verbose, TEXT("[%lli] BeginCapture: StartRadius: %.1f | DeltaDivisor: %.1f"), GFrameCounter, InteractionPlanarStartDistance, InteractionDeltaDivisor);
+		UE_CLOG(bLogScale, LogEditorTRSGizmo, Verbose, TEXT("[%lli] BeginCapture: Offset: %.1f"), GFrameCounter, InteractionReferencePointOffsetDistance);
+		UE_CLOG(bLogScale, LogEditorTRSGizmo, Verbose, TEXT("[%lli] BeginCapture: StartRadius: %.1f | DeltaDivisor: %.1f"), GFrameCounter, InteractionPlanarStartDistance, InteractionDeltaDivisor);
 	}
 
 	// Set initial screen-space sign
@@ -6121,7 +6121,7 @@ void UEditorTRSGizmo::OnClickDragScale(const FInputDeviceRay& DragPos)
 			DistanceToReferencePoint = FVector::DotProduct(InteractionAxisDirection, (CursorPlaneHitPos - DistanceReferencePoint));
 			DistanceToReferencePoint2D = FVector2D::DotProduct(DirectionAxis2D, (DragPos.ScreenPosition - DistanceReferencePoint2D));
 
-			UE_CLOG(bLogScale, LogTransformGizmo, Verbose, TEXT("DistanceToRef2D: %.1f | DistanceToRef: %.1f"), DistanceToReferencePoint2D, DistanceToReferencePoint);
+			UE_CLOG(bLogScale, LogEditorTRSGizmo, Verbose, TEXT("DistanceToRef2D: %.1f | DistanceToRef: %.1f"), DistanceToReferencePoint2D, DistanceToReferencePoint);
 
 			// Update Debug
 			if (GizmoLocals::DoDebugDraw())
@@ -6193,11 +6193,11 @@ void UEditorTRSGizmo::OnClickDragScale(const FInputDeviceRay& DragPos)
 
 	ModifyScaleDeltaForScaleType(ScaleType, PreviousAppliedDeltaScale, ScaleDelta);
 
-	UE_CLOG(bLogScale, LogTransformGizmo, Verbose, TEXT("PreSnapDelta: %s | Delta: %s"),
+	UE_CLOG(bLogScale, LogEditorTRSGizmo, Verbose, TEXT("PreSnapDelta: %s | Delta: %s"),
 		*PreSnapScaleDelta.ToString(),
 		*ScaleDelta.ToString());
 
-	UE_CLOG(bLogScale, LogTransformGizmo, Verbose, TEXT("ScaleSign: %.3f | StartR: %.3f | CurrentR: %.3f | R: %.3f | InteractionD: %s | Divisor: %.3f | ScaleD: %s"),
+	UE_CLOG(bLogScale, LogEditorTRSGizmo, Verbose, TEXT("ScaleSign: %.3f | StartR: %.3f | CurrentR: %.3f | R: %.3f | InteractionD: %s | Divisor: %.3f | ScaleD: %s"),
 		ScaleSign,
 		OffsetScreenStartDistance,
 		OffsetScreenCurrDistance,
@@ -6222,12 +6222,12 @@ void UEditorTRSGizmo::OnClickDragScale(const FInputDeviceRay& DragPos)
 			bDetectedNonStandardScaling = true;
 		}
 
-		UE_CLOG(bLogScale, LogTransformGizmo, Verbose, TEXT("[%lli] UpdateCapture: ScaleDelta: %s | Pre: %s | Post: %s"), GFrameCounter, *ScaleDelta.ToString(), *PreDeltaScale.ToString(), *PostDeltaScale.ToString());
+		UE_CLOG(bLogScale, LogEditorTRSGizmo, Verbose, TEXT("[%lli] UpdateCapture: ScaleDelta: %s | Pre: %s | Post: %s"), GFrameCounter, *ScaleDelta.ToString(), *PreDeltaScale.ToString(), *PostDeltaScale.ToString());
 	}
 
 	UpdateScaleDelta();
 
-	UE_CLOG(bLogScale, LogTransformGizmo, Verbose, TEXT("[%lli] UpdateCapture: ScreenToWorldDistanceMultiplier: %.1f"), GFrameCounter, DistanceToReferencePoint);
+	UE_CLOG(bLogScale, LogEditorTRSGizmo, Verbose, TEXT("[%lli] UpdateCapture: ScreenToWorldDistanceMultiplier: %.1f"), GFrameCounter, DistanceToReferencePoint);
 
 	// Update Debug
 	if (GizmoLocals::DoDebugDraw())
@@ -6434,7 +6434,7 @@ void UEditorTRSGizmo::ModifyScaleDeltaForScaleType(const EGizmoTransformScaleTyp
 
 		UE_CLOG(
 			UE::Editor::InteractiveToolsFramework::Private::TransformGizmoLocals::bLogScale,
-			LogTransformGizmo, Verbose, TEXT("DeltaFromStart: %s | StartObjScale: %s | CurrentObjScale: %s | TargetScale: %s | ThisDelta: %s | AdjustedSignedThisDelta: %s"),
+			LogEditorTRSGizmo, Verbose, TEXT("DeltaFromStart: %s | StartObjScale: %s | CurrentObjScale: %s | TargetScale: %s | ThisDelta: %s | AdjustedSignedThisDelta: %s"),
 				*DeltaFromStart.ToCompactString(),
 				*StartObjectScale.ToCompactString(),
 				*CurrentObjectScale.ToCompactString(),
@@ -6697,7 +6697,7 @@ bool UEditorTRSGizmo::OnClickPressRotateArc(
 	}
 
 	// Log StartAngle, InteractionStartAngle, AlignedInteractionAngleOffset
-	UE_CLOG(bLogRotation, LogTransformGizmo, Verbose, TEXT("[%lli] StartAngle: %.1f | AlignedStartAngle: %.1f | Delta: %.1f"),
+	UE_CLOG(bLogRotation, LogEditorTRSGizmo, Verbose, TEXT("[%lli] StartAngle: %.1f | AlignedStartAngle: %.1f | Delta: %.1f"),
 		GFrameCounter,
 		FMath::RadiansToDegrees(InteractionStartAngle),
 		FMath::RadiansToDegrees(HitAngle),
