[TRSGizmo] Fix issue where setting the gizmo version via the CVAR caused the new settings to be disabled and the legacy settings enabled.

#rnx
#Synced-CL 52063422

diff --git a/Engine/Source/Editor/Experimental/EditorInteractiveToolsFramework/Private/Customizations/TransformGizmoEditorSettingsCustomization.cpp b/Engine/Source/Editor/Experimental/EditorInteractiveToolsFramework/Private/Customizations/TransformGizmoEditorSettingsCustomization.cpp
--- a/Engine/Source/Editor/Experimental/EditorInteractiveToolsFramework/Private/Customizations/TransformGizmoEditorSettingsCustomization.cpp
+++ b/Engine/Source/Editor/Experimental/EditorInteractiveToolsFramework/Private/Customizations/TransformGizmoEditorSettingsCustomization.cpp
@@ -118,7 +118,7 @@ FTransformGizmoEditorSettingsCustomization::FTransformGizmoEditorSettingsCustomi
 				return false;
 			}
 
-			return Settings->bUseExperimentalGizmo;
+			return !Settings->UsesLegacyGizmo();
 		});
 }
 
@@ -210,6 +210,14 @@ void FTransformGizmoEditorSettingsCustomization::CustomizeDetails(const TSharedP
 	{
 		GizmoInteractionCategory.SetCategoryVisibility(true);
 		GizmoInteractionCategory.SetSortOrder(0);
+		
+		const TSharedPtr<IPropertyHandle> DragDuplicateOnRotationProperty = GizmoParametersPropertyHandle->GetChildHandle(GET_MEMBER_NAME_CHECKED(FGizmosParameters, bDragDuplicateOnRotation));
+		GizmoInteractionCategory.AddProperty(DragDuplicateOnRotationProperty)
+		.EditCondition(IsUsingNewGizmosAttribute, { });
+		
+		const TSharedPtr<IPropertyHandle> ScreenSpaceNudgeProperty = GizmoParametersPropertyHandle->GetChildHandle(GET_MEMBER_NAME_CHECKED(FGizmosParameters, bScreenSpaceNudge));
+		GizmoInteractionCategory.AddProperty(ScreenSpaceNudgeProperty)
+		.EditCondition(IsUsingNewGizmosAttribute, { });
 
 		TSharedPtr<IPropertyHandle> InteractionPropertyHandle = GizmoParametersPropertyHandle->GetChildHandle(GET_MEMBER_NAME_CHECKED(FGizmosParameters, Interaction));
 
diff --git a/Engine/Source/Editor/Experimental/EditorInteractiveToolsFramework/Private/TransformGizmoEditorSettings.cpp b/Engine/Source/Editor/Experimental/EditorInteractiveToolsFramework/Private/TransformGizmoEditorSettings.cpp
--- a/Engine/Source/Editor/Experimental/EditorInteractiveToolsFramework/Private/TransformGizmoEditorSettings.cpp
+++ b/Engine/Source/Editor/Experimental/EditorInteractiveToolsFramework/Private/TransformGizmoEditorSettings.cpp
@@ -237,7 +237,7 @@ void UTransformGizmoEditorSettings::SetTransformGizmoSize(float InTransformGizmo
 
 void UTransformGizmoEditorSettings::BroadcastNewTRSGizmoChange() const
 {
-	UEditorInteractiveGizmoManager::OnUsesNewTRSGizmosChangedDelegate().Broadcast(UsesExperimentalGizmo() || UsesNewTRSGizmo());
+	UEditorInteractiveGizmoManager::OnUsesNewTRSGizmosChangedDelegate().Broadcast(!UsesLegacyGizmo());
 }
 
 void UTransformGizmoEditorSettings::BroadcastGizmosParametersChange() const
diff --git a/Engine/Source/Editor/Experimental/EditorInteractiveToolsFramework/Public/TransformGizmoEditorSettings.h b/Engine/Source/Editor/Experimental/EditorInteractiveToolsFramework/Public/TransformGizmoEditorSettings.h
--- a/Engine/Source/Editor/Experimental/EditorInteractiveToolsFramework/Public/TransformGizmoEditorSettings.h
+++ b/Engine/Source/Editor/Experimental/EditorInteractiveToolsFramework/Public/TransformGizmoEditorSettings.h
@@ -26,35 +26,35 @@ public:
 	UE_API virtual void PostReloadConfig(class FProperty* PropertyThatWasLoaded) override;
 #endif
 
-	UPROPERTY(Config, EditAnywhere, Category = "Transform Gizmo", DisplayName = "Transform Gizmo Size", meta = (ClampMin = "-10.0", ClampMax = "150.0", EditCondition = "!bUseExperimentalGizmo"))
+	UPROPERTY(Config, EditAnywhere, Category = "Transform Gizmo", DisplayName = "Transform Gizmo Size", meta = (ClampMin = "-10.0", ClampMax = "150.0", EditCondition = "!(bUseExperimentalGizmo || bUseEditorTRSGizmo)"))
 	float TransformGizmoSize = 0.0f;
 
 	/**
 	 * Allow arcball rotation with rotate widget
 	 * (updates the setting with the same name found in Level Editor Viewport Settings)
 	 */
-	UPROPERTY(EditAnywhere, Category = "Transform Gizmo", DisplayName = "Enable Arcball Rotate", meta = (EditCondition = "!bUseExperimentalGizmo"))
+	UPROPERTY(EditAnywhere, Category = "Transform Gizmo", DisplayName = "Enable Arcball Rotate", meta = (EditCondition = "!(bUseExperimentalGizmo || bUseEditorTRSGizmo)"))
 	bool bEnableArcballRotate;
 
 	/**
 	 * Allow screen rotation with rotate widget
 	 * (updates the setting with the same name found in Level Editor Viewport Settings)
 	 */
-	UPROPERTY(EditAnywhere, Category = "Transform Gizmo", DisplayName = "Enable Screen Rotate", meta = (EditCondition = "!bUseExperimentalGizmo"))
+	UPROPERTY(EditAnywhere, Category = "Transform Gizmo", DisplayName = "Enable Screen Rotate", meta = (EditCondition = "!(bUseExperimentalGizmo || bUseEditorTRSGizmo)"))
 	bool bEnableScreenRotate;
 
 	/**
 	 * If true, the Edit widget of a transform will display the axis
 	 * (updates the setting with the same name found in Level Editor Viewport Settings)
 	 */
-	UPROPERTY(EditAnywhere, Category = "Transform Gizmo", DisplayName = "Enable Axis drawing for Transform Edit Widget", meta = (EditCondition = "!bUseExperimentalGizmo"))
+	UPROPERTY(EditAnywhere, Category = "Transform Gizmo", DisplayName = "Enable Axis drawing for Transform Edit Widget", meta = (EditCondition = "!(bUseExperimentalGizmo || bUseEditorTRSGizmo)"))
 	bool bEnableAxisDrawing;
 
 	/**
 	 * Allow translate/rotate widget
 	 * (updates the setting with the same name found in Level Editor Viewport Settings)
 	 */
-	UPROPERTY(EditAnywhere, Category = "Transform Gizmo", DisplayName = "Enable Combined Translate/Rotate Widget", meta = (EditCondition = "!bUseExperimentalGizmo"))
+	UPROPERTY(EditAnywhere, Category = "Transform Gizmo", DisplayName = "Enable Combined Translate/Rotate Widget", meta = (EditCondition = "!(bUseExperimentalGizmo || bUseEditorTRSGizmo)"))
 	bool bEnableCombinedTranslateRotate;
 	//~ End Transform Gizmo Category
 
@@ -66,7 +66,7 @@ public:
 	UPROPERTY(Config, EditAnywhere, Category = "Experimental", Getter = "UsesNewTRSGizmo", Setter = "SetUseNewTRSGizmo")
 	bool bUseEditorTRSGizmo = true;
 
-	UPROPERTY(Config, EditAnywhere, Category = "Experimental", meta = (ShowOnlyInnerProperties, EditCondition = "bUseExperimentalGizmo"))
+	UPROPERTY(Config, EditAnywhere, Category = "Experimental", meta = (ShowOnlyInnerProperties, EditCondition = "bUseExperimentalGizmo || bUseEditorTRSGizmo"))
 	FGizmosParameters GizmosParameters;
 	//~ End Experimental Gizmo Category
 
