EditorTRS cis fix

#Synced-CL 52063915

diff --git a/Engine/Source/Editor/Experimental/EditorInteractiveToolsFramework/Private/EditorGizmos/EditorTRSGizmo.cpp b/Engine/Source/Editor/Experimental/EditorInteractiveToolsFramework/Private/EditorGizmos/EditorTRSGizmo.cpp
--- a/Engine/Source/Editor/Experimental/EditorInteractiveToolsFramework/Private/EditorGizmos/EditorTRSGizmo.cpp
+++ b/Engine/Source/Editor/Experimental/EditorInteractiveToolsFramework/Private/EditorGizmos/EditorTRSGizmo.cpp
@@ -67,8 +67,6 @@
 
 DEFINE_LOG_CATEGORY_STATIC(LogEditorTRSGizmo, Log, All);
 
-ENUM_RANGE_BY_COUNT(ETransformGizmoPartIdentifier, ETransformGizmoPartIdentifier::Max);
-
 namespace GizmoLocals
 {
 
