EditorTRS cis fix #2

#Synced-CL 52063915

diff --git a/Engine/Source/Editor/Experimental/EditorInteractiveToolsFramework/Private/EditorGizmos/TransformGizmo.cpp b/Engine/Source/Editor/Experimental/EditorInteractiveToolsFramework/Private/EditorGizmos/TransformGizmo.cpp
--- a/Engine/Source/Editor/Experimental/EditorInteractiveToolsFramework/Private/EditorGizmos/TransformGizmo.cpp
+++ b/Engine/Source/Editor/Experimental/EditorInteractiveToolsFramework/Private/EditorGizmos/TransformGizmo.cpp
@@ -31,8 +31,6 @@
 
 DEFINE_LOG_CATEGORY_STATIC(LogTransformGizmo, Log, All);
 
-ENUM_RANGE_BY_COUNT(ETransformGizmoPartIdentifier, ETransformGizmoPartIdentifier::Max);
-
 namespace GizmoLocals
 {
 	
diff --git a/Engine/Source/Editor/Experimental/EditorInteractiveToolsFramework/Public/EditorGizmos/TransformGizmo.h b/Engine/Source/Editor/Experimental/EditorInteractiveToolsFramework/Public/EditorGizmos/TransformGizmo.h
--- a/Engine/Source/Editor/Experimental/EditorInteractiveToolsFramework/Public/EditorGizmos/TransformGizmo.h
+++ b/Engine/Source/Editor/Experimental/EditorInteractiveToolsFramework/Public/EditorGizmos/TransformGizmo.h
@@ -104,6 +104,8 @@ enum class ETransformGizmoPartIdentifier
 	Max
 };
 
+ENUM_RANGE_BY_COUNT(ETransformGizmoPartIdentifier, ETransformGizmoPartIdentifier::Max);
+
 namespace UE::Editor::InteractiveToolsFramework
 {
 	/** Friended class to access UTransformGizmo internals, used externally. */
