Cleaned up/simplified scale and rotation settings

#Synced-CL 52063915

diff --git a/Engine/Source/Editor/Experimental/EditorInteractiveToolsFramework/Private/Customizations/TransformGizmoEditorSettingsCustomization.cpp b/Engine/Source/Editor/Experimental/EditorInteractiveToolsFramework/Private/Customizations/TransformGizmoEditorSettingsCustomization.cpp
--- a/Engine/Source/Editor/Experimental/EditorInteractiveToolsFramework/Private/Customizations/TransformGizmoEditorSettingsCustomization.cpp
+++ b/Engine/Source/Editor/Experimental/EditorInteractiveToolsFramework/Private/Customizations/TransformGizmoEditorSettingsCustomization.cpp
@@ -112,13 +112,7 @@ FTransformGizmoEditorSettingsCustomization::FTransformGizmoEditorSettingsCustomi
 	IsUsingNewGizmosAttribute = TAttribute<bool>::Create(
 		[]()
 		{
-			const UTransformGizmoEditorSettings* Settings = GetDefault<UTransformGizmoEditorSettings>();
-			if (!ensure(Settings))
-			{
-				return false;
-			}
-
-			return !Settings->UsesLegacyGizmo();
+			return UEditorInteractiveGizmoManager::GetTransformGizmoVersion() >= 1;
 		});
 }
 
@@ -198,7 +192,8 @@ void FTransformGizmoEditorSettingsCustomization::CustomizeDetails(const TSharedP
 			.ValueContent()
 			[
 				GeneratePresetComboBoxContent()
-			];
+			]
+			.EditCondition(IsUsingNewGizmosAttribute, { });
 		}
 	}
 	
@@ -310,18 +305,19 @@ void FTransformGizmoEditorSettingsCustomization::CustomizeDetails(const TSharedP
 
 		const TSharedPtr<IPropertyHandle> DefaultRotateModePropertyHandle = GizmoParametersPropertyHandle->GetChildHandle(GET_MEMBER_NAME_CHECKED(FGizmosParameters, RotateMode));
 		GizmoInteractionCategory.AddProperty(DefaultRotateModePropertyHandle)
-		.EditCondition(IsUsingNewGizmosAttribute, { });
+		.EditCondition(IsUsingNewGizmosAttribute, { })
+		.OverrideResetToDefault(MakeResetToDefault([]() -> EAxisRotateMode::Type*
+		{
+			static EAxisRotateMode::Type DefaultRotateMode = FGizmoElementRotateInteraction::GetDefaultRotateMode();
+			return &DefaultRotateMode;
+		}));
 
-		const TSharedPtr<IPropertyHandle> ScaleTypePropertyHandle =
-			InteractionPropertyHandle
-			->GetChildHandle(GET_MEMBER_NAME_CHECKED(FTransformGizmoInteraction, ScaleInteraction))
-			->GetChildHandle(GET_MEMBER_NAME_CHECKED(FGizmoElementScaleInteraction, ScaleType));
+		const TSharedPtr<IPropertyHandle> ScaleTypePropertyHandle = GizmoParametersPropertyHandle->GetChildHandle(GET_MEMBER_NAME_CHECKED(FGizmosParameters, ScaleType));
 		GizmoInteractionCategory.AddProperty(ScaleTypePropertyHandle)
 		.EditCondition(IsUsingNewGizmosAttribute, { })
-		.OverrideResetToDefault(MakeResetToDefault([]()
+		.OverrideResetToDefault(MakeResetToDefault([]() -> EGizmoTransformScaleType*
 		{
-			static FTransformGizmoInteraction DefaultInteraction;
-			static EGizmoTransformScaleType DefaultScaleType = DefaultInteraction.ScaleInteraction.GetScaleType();
+			static EGizmoTransformScaleType DefaultScaleType = FGizmoElementScaleInteraction::GetDefaultScaleType();
 			return &DefaultScaleType;
 		}));
 	}
@@ -545,9 +541,6 @@ void FTransformGizmoEditorSettingsCustomization::CustomizeDetails(const TSharedP
 		// @note: This inverts the logic of bUseExperimentalGizmo and sets both flags together
 		// so the UI only offers "New" or "Legacy" (the intermediate state is CVar-only).
 		{
-			TSharedPtr<IPropertyHandle> UseExperimentalGizmoPropertyHandle = InDetailBuilder->GetProperty(GET_MEMBER_NAME_CHECKED(UTransformGizmoEditorSettings, bUseExperimentalGizmo));
-			TSharedPtr<IPropertyHandle> UseEditorTRSGizmoPropertyHandle = InDetailBuilder->GetProperty(GET_MEMBER_NAME_CHECKED(UTransformGizmoEditorSettings, bUseEditorTRSGizmo));
-
 			const FText EnableLegacyGizmoDisplayName = LOCTEXT("UseLegacyGizmoDisplayName", "Enable Legacy");
 			LegacyGizmoCategory
 			.AddCustomRow(EnableLegacyGizmoDisplayName)
@@ -566,7 +559,7 @@ void FTransformGizmoEditorSettingsCustomization::CustomizeDetails(const TSharedP
 					const bool bUseLegacyGizmo = UEditorInteractiveGizmoManager::GetTransformGizmoVersion() == 0;
 					return bUseLegacyGizmo ? ECheckBoxState::Checked : ECheckBoxState::Unchecked;
 				})
-				.OnCheckStateChanged_Lambda([UseExperimentalGizmoPropertyHandle, UseEditorTRSGizmoPropertyHandle](ECheckBoxState InCheckState)
+				.OnCheckStateChanged_Lambda([](ECheckBoxState InCheckState)
 				{
 					const bool bEnableLegacyGizmo = (InCheckState == ECheckBoxState::Checked);
 
diff --git a/Engine/Source/Editor/Experimental/EditorInteractiveToolsFramework/Private/EditorGizmos/GizmoElementRotate.cpp b/Engine/Source/Editor/Experimental/EditorInteractiveToolsFramework/Private/EditorGizmos/GizmoElementRotate.cpp
--- a/Engine/Source/Editor/Experimental/EditorInteractiveToolsFramework/Private/EditorGizmos/GizmoElementRotate.cpp
+++ b/Engine/Source/Editor/Experimental/EditorInteractiveToolsFramework/Private/EditorGizmos/GizmoElementRotate.cpp
@@ -2,6 +2,8 @@
 
 #include "EditorGizmos/GizmoElementRotate.h"
 
+#include "TransformGizmoEditorSettings.h"
+
 FGizmoElementRotateAxisStyle::FGizmoElementRotateAxisStyle()
 {
 	PixelHitDistanceThreshold = 10.0f;
@@ -13,3 +15,18 @@ FGizmoElementRotateArcballStyle::FGizmoElementRotateArcballStyle()
 	LineColors.Default = (Colors.Default.GetValue() * 0.25f).CopyWithNewOpacity(1.0f);
 	PixelHitDistanceThreshold = 15.0f;
 }
+
+EAxisRotateMode::Type FGizmoElementRotateInteraction::GetRotateMode() const
+{
+	if (RotateMode.IsSet())
+	{
+		return RotateMode.GetValue();
+	}
+
+	if (const UTransformGizmoEditorSettings* Settings = GetDefault<UTransformGizmoEditorSettings>())
+	{
+		return Settings->GizmosParameters.RotateMode;
+	}
+
+	return GetDefaultRotateMode();
+}
diff --git a/Engine/Source/Editor/Experimental/EditorInteractiveToolsFramework/Private/EditorGizmos/GizmoElementScale.cpp b/Engine/Source/Editor/Experimental/EditorInteractiveToolsFramework/Private/EditorGizmos/GizmoElementScale.cpp
--- a/Engine/Source/Editor/Experimental/EditorInteractiveToolsFramework/Private/EditorGizmos/GizmoElementScale.cpp
+++ b/Engine/Source/Editor/Experimental/EditorInteractiveToolsFramework/Private/EditorGizmos/GizmoElementScale.cpp
@@ -3,6 +3,7 @@
 #include "EditorGizmos/GizmoElementScale.h"
 
 #include "Editor.h"
+#include "TransformGizmoEditorSettings.h"
 
 FGizmoElementScaleAxisStyle::FGizmoElementScaleAxisStyle()
 {
@@ -58,5 +59,21 @@ FGizmoElementScaleUniformStyle::FGizmoElementScaleUniformStyle()
 
 EGizmoTransformScaleType FGizmoElementScaleInteraction::GetScaleType() const
 {
-	return ScaleType.Get((GEditor && GEditor->UsePercentageBasedScaling()) ? EGizmoTransformScaleType::PercentageBased : EGizmoTransformScaleType::Default);
+	EGizmoTransformScaleType ResolvedScaleType = GetDefaultScaleType();
+	if (ScaleType.IsSet())
+	{
+		ResolvedScaleType = ScaleType.GetValue();
+	}
+	else if (const UTransformGizmoEditorSettings* Settings = GetDefault<UTransformGizmoEditorSettings>())
+	{
+		ResolvedScaleType = Settings->GizmosParameters.ScaleType;
+	}
+	
+	if (ResolvedScaleType != EGizmoTransformScaleType::Legacy)
+	{
+		return ResolvedScaleType;
+	}
+	
+	// Legacy uses Editor globals
+	return (GEditor && GEditor->UsePercentageBasedScaling()) ? EGizmoTransformScaleType::PercentageBased : EGizmoTransformScaleType::Default;
 }
diff --git a/Engine/Source/Editor/Experimental/EditorInteractiveToolsFramework/Public/EditorGizmos/GizmoElementRotate.h b/Engine/Source/Editor/Experimental/EditorInteractiveToolsFramework/Public/EditorGizmos/GizmoElementRotate.h
--- a/Engine/Source/Editor/Experimental/EditorInteractiveToolsFramework/Public/EditorGizmos/GizmoElementRotate.h
+++ b/Engine/Source/Editor/Experimental/EditorInteractiveToolsFramework/Public/EditorGizmos/GizmoElementRotate.h
@@ -187,9 +187,12 @@ struct FGizmoElementRotateInteraction
 {
 	GENERATED_BODY()
 
+	/** The default rotation mode to use with the Gizmo. */
+	static constexpr EAxisRotateMode::Type GetDefaultRotateMode() { return EAxisRotateMode::ScreenArc; }
+
 	/** Optionally override the project-level default rotation mode (Arc or Pull). */
 	UPROPERTY(EditAnywhere, Category = "Input")
-	TOptional<TEnumAsByte<EAxisRotateMode::Type>> DefaultRotateMode;
+	TOptional<TEnumAsByte<EAxisRotateMode::Type>> RotateMode;
 
 	/**
 	 * The explicit drag direction corresponding with view-relative rotation. (1.0, 0.0) means left-to-right movement rotates clockwise, right-to-left moves counter-clockwise. Vertical movement wouldn't do anything.
@@ -206,6 +209,9 @@ struct FGizmoElementRotateInteraction
 	/** Sensitivity multiplier for Arc rotation mode. */
 	UPROPERTY(EditAnywhere, Category = "Input")
 	float ArcMultiplier = 1.0f;
+	
+	/** Returns the overridden RotateMode if set, otherwise use the Editor Setting. */
+	UE_API EAxisRotateMode::Type GetRotateMode() const;
 };
 
 #undef UE_API
diff --git a/Engine/Source/Editor/Experimental/EditorInteractiveToolsFramework/Public/EditorGizmos/GizmoElementScale.h b/Engine/Source/Editor/Experimental/EditorInteractiveToolsFramework/Public/EditorGizmos/GizmoElementScale.h
--- a/Engine/Source/Editor/Experimental/EditorInteractiveToolsFramework/Public/EditorGizmos/GizmoElementScale.h
+++ b/Engine/Source/Editor/Experimental/EditorInteractiveToolsFramework/Public/EditorGizmos/GizmoElementScale.h
@@ -242,6 +242,9 @@ struct FGizmoElementScaleInteraction
 {
 	GENERATED_BODY()
 
+	/** The default scale-type to use with the Gizmo. */
+	static constexpr EGizmoTransformScaleType GetDefaultScaleType() { return EGizmoTransformScaleType::PercentageBased; }
+
 	/** The amount is based on the distance to the reference point specified here. */
 	UPROPERTY(EditAnywhere, Category = "Input")
 	EGizmoElementScaleDistanceSource DistanceSource = EGizmoElementScaleDistanceSource::FromStart;
@@ -254,7 +257,7 @@ struct FGizmoElementScaleInteraction
 	UPROPERTY(EditAnywhere, Category = "Input", meta = (DisplayName = "Drag Direction", EditCondition = "DistanceType == EGizmoElementScaleDistanceType::Directional"))
 	FVector2D Direction = FVector2d(1.0f, 1.0f);
 
-	/** If set, is used instead of the global setting (GEditor->UsePercentageBasedScaling()). */
+	/** Optionally override the project-level default scale type (Default/Offset, or PercentageBased). */
 	UPROPERTY(EditAnywhere, Category = "Input")
 	TOptional<EGizmoTransformScaleType> ScaleType = EGizmoTransformScaleType::PercentageBased;
 
@@ -270,7 +273,7 @@ struct FGizmoElementScaleInteraction
 	UPROPERTY(EditAnywhere, Category = "Input", AdvancedDisplay)
 	TOptional<float> MinimumScale = 0.0001f;
 
-	/** Returns the overridden ScaleType if set, otherwise GEditor->UsePercentageBasedScaling(). */
+	/** Returns the overridden ScaleType if set, otherwise use the Editor Setting. */
 	UE_API EGizmoTransformScaleType GetScaleType() const;
 };
 
diff --git a/Engine/Source/Editor/Experimental/EditorInteractiveToolsFramework/Public/EditorGizmos/TransformGizmo.h b/Engine/Source/Editor/Experimental/EditorInteractiveToolsFramework/Public/EditorGizmos/TransformGizmo.h
--- a/Engine/Source/Editor/Experimental/EditorInteractiveToolsFramework/Public/EditorGizmos/TransformGizmo.h
+++ b/Engine/Source/Editor/Experimental/EditorInteractiveToolsFramework/Public/EditorGizmos/TransformGizmo.h
@@ -299,6 +299,10 @@ struct FGizmosParameters
 	/** Determines how dragging the rotate gizmo affects the selected objects */
 	UPROPERTY(EditAnywhere, Category = "Experimental")
 	TEnumAsByte<EAxisRotateMode::Type> RotateMode = EAxisRotateMode::ScreenArc;
+	
+	/** Determines how scale is calculated when applied. */
+	UPROPERTY(EditAnywhere, Category = "Experimental")
+	EGizmoTransformScaleType ScaleType = EGizmoTransformScaleType::PercentageBased;
 
 	/**
 	 * When enabled, Ctrl+clicking multiple mouse buttons allows you to indirectly manipulate multiple axes at once. 
diff --git a/Engine/Source/Editor/Experimental/EditorInteractiveToolsFramework/Public/EditorGizmos/TransformGizmoInterfaces.h b/Engine/Source/Editor/Experimental/EditorInteractiveToolsFramework/Public/EditorGizmos/TransformGizmoInterfaces.h
--- a/Engine/Source/Editor/Experimental/EditorInteractiveToolsFramework/Public/EditorGizmos/TransformGizmoInterfaces.h
+++ b/Engine/Source/Editor/Experimental/EditorInteractiveToolsFramework/Public/EditorGizmos/TransformGizmoInterfaces.h
@@ -37,8 +37,14 @@ namespace EAxisRotateMode
 UENUM()
 enum class EGizmoTransformScaleType : uint8
 {
-	Default,
-	PercentageBased
+	/** Add/Subtract */
+	Default = 0,
+
+	/** Multiply */
+	PercentageBased = 1,
+
+	/** Default/Offset, or PercentageBased if GEditor->UsePercentageBasedScaling() */
+	Legacy = 2
 };
 
 enum class EViewportContext : uint8
