[TRS Gizmo]
- Separate the thicknesses of the screen-space rotation handle and the delta arc lines from the primary rotation gizmo thicknesses.
- Set the primary rotation axes thickness to match the thickness of the scale/translate axes.

#rnx
#Synced-CL 52063422

diff --git a/Engine/Source/Editor/Experimental/EditorInteractiveToolsFramework/Private/EditorGizmos/EditorTRSGizmo.cpp b/Engine/Source/Editor/Experimental/EditorInteractiveToolsFramework/Private/EditorGizmos/EditorTRSGizmo.cpp
--- a/Engine/Source/Editor/Experimental/EditorInteractiveToolsFramework/Private/EditorGizmos/EditorTRSGizmo.cpp
+++ b/Engine/Source/Editor/Experimental/EditorInteractiveToolsFramework/Private/EditorGizmos/EditorTRSGizmo.cpp
@@ -3208,6 +3208,7 @@ UGizmoElementRotateAxis* UEditorTRSGizmo::MakeRotateScreenSpaceHandle()
 			GetSizeCoefficient());
 
 	ElementStyle.Radius = Style.RotateStyle.AxisStyle.Radius + Style.RotateStyle.ScreenSpaceRadiusOffset;
+	ElementStyle.LineThickness = Style.RotateStyle.ScreenSpaceLineThickness;
 	ElementStyle.LineThicknessMultiplier = Style.LineThicknessMultiplier;
 	ElementStyle.DeltaFillHSVModifier = FLinearColor(0.0f, 1.0f, 1.0f);
 
@@ -3258,6 +3259,7 @@ void UEditorTRSGizmo::UpdateRotateScreenSpaceHandle(UGizmoElementRotateAxis* InE
 		GetSizeCoefficient());
 
 	ElementStyle.Radius = Style.RotateStyle.AxisStyle.Radius + Style.RotateStyle.ScreenSpaceRadiusOffset;
+	ElementStyle.LineThickness = Style.RotateStyle.ScreenSpaceLineThickness;
 	ElementStyle.LineThicknessMultiplier = Style.LineThicknessMultiplier;
 	ElementStyle.DeltaFillHSVModifier = FLinearColor(0.0f, 1.0f, 1.0f);
 
diff --git a/Engine/Source/Editor/Experimental/EditorInteractiveToolsFramework/Private/EditorGizmos/GizmoElementRotateAxis.cpp b/Engine/Source/Editor/Experimental/EditorInteractiveToolsFramework/Private/EditorGizmos/GizmoElementRotateAxis.cpp
--- a/Engine/Source/Editor/Experimental/EditorInteractiveToolsFramework/Private/EditorGizmos/GizmoElementRotateAxis.cpp
+++ b/Engine/Source/Editor/Experimental/EditorInteractiveToolsFramework/Private/EditorGizmos/GizmoElementRotateAxis.cpp
@@ -832,7 +832,7 @@ void UGizmoElementRotateAxis::ApplyStyle()
 		DeltaArcElement->SetHoverLineColor(LineColor);
 		DeltaArcElement->SetInteractLineColor(LineColor);
 
-		const float DeltaArcLineThickness = Style.LineThickness * Style.LineThicknessMultiplier * SizeCoefficient;
+		const float DeltaArcLineThickness = Style.DeltaStrokeThickness * Style.LineThicknessMultiplier * SizeCoefficient;
 		DeltaArcElement->SetLineThickness(FMath::Max(Style.MinLineThickness, DeltaArcLineThickness));
 	}
 
diff --git a/Engine/Source/Editor/Experimental/EditorInteractiveToolsFramework/Public/EditorGizmos/GizmoElementRotate.h b/Engine/Source/Editor/Experimental/EditorInteractiveToolsFramework/Public/EditorGizmos/GizmoElementRotate.h
--- a/Engine/Source/Editor/Experimental/EditorInteractiveToolsFramework/Public/EditorGizmos/GizmoElementRotate.h
+++ b/Engine/Source/Editor/Experimental/EditorInteractiveToolsFramework/Public/EditorGizmos/GizmoElementRotate.h
@@ -47,7 +47,7 @@ struct FGizmoElementRotateAxisStyle : public FGizmoStyleBase
 
 	/** Thickness of the rotation ring line in world units. */
 	UPROPERTY(EditAnywhere, Category = "Rotate", meta = (ClampMin = 0.1f))
-	float LineThickness = 0.625f;
+	float LineThickness = 1.5f;
 
 	/** Radius of the origin sphere drawn at the rotation center during interaction, in world units. */
 	UPROPERTY(EditAnywhere, Category = "Rotate", meta = (ClampMin = 0.1f))
@@ -77,6 +77,10 @@ struct FGizmoElementRotateAxisStyle : public FGizmoStyleBase
 	UPROPERTY(EditAnywhere, Category = "Rotate", meta = (ClampMin = 0.0f, ClampMax = 1.0f))
 	float DeltaFillOpacity = 0.4f;
 
+	/** Thickness of the delta arc lines. */
+	UPROPERTY(EditAnywhere, Category = "Rotate", meta = (ClampMin = 0.1f))
+	float DeltaStrokeThickness = .625f;
+
 	/** Stroke/line opacity of the delta arc. */
 	UPROPERTY(EditAnywhere, Category = "Rotate", meta = (ClampMin = 0.0f, ClampMax = 1.0f))
 	float DeltaStrokeOpacity = 1.0f;
@@ -124,6 +128,7 @@ struct FGizmoElementRotateAxisStyle : public FGizmoStyleBase
 			&& InLeft.OriginToCursorLineDashGapSpacing == InRight.OriginToCursorLineDashGapSpacing
 			&& InLeft.OriginToCursorLineColor == InRight.OriginToCursorLineColor
 			&& InLeft.DeltaFillOpacity == InRight.DeltaFillOpacity
+			&& InLeft.DeltaStrokeThickness == InRight.DeltaStrokeThickness
 			&& InLeft.DeltaStrokeOpacity == InRight.DeltaStrokeOpacity
 			&& InLeft.DeltaFillHSVModifier == InRight.DeltaFillHSVModifier
 			&& InLeft.DeltaTextBackgroundColor == InRight.DeltaTextBackgroundColor
@@ -176,6 +181,10 @@ struct FGizmoElementRotateStyle : public FGizmoStyleBase
 	UPROPERTY(EditAnywhere, Category = "Rotate", meta = (ClampMin = 0.01f))
 	float ScreenSpaceRadiusOffset = 12.6f;
 
+	/** The thickness of the screen space rotation handle arc. */
+	UPROPERTY(EditAnywhere, Category = "Rotate", meta = (ClampMin = 0.01f))
+	float ScreenSpaceLineThickness = 0.625f;
+
 	/** Color of the screen-space rotation circle (outer ring). */
 	UPROPERTY(EditAnywhere, Category = "Rotate", meta = (HideAlphaChannel))
 	FLinearColor ScreenSpaceCircleColor = FLinearColor(0.9f, 0.9f, 0.9f);
