Register snapping manager once on init instead of toggling per mode change

#Synced-CL 52103551

diff --git a/Engine/Source/Editor/Experimental/EditorInteractiveToolsFramework/Private/EditorInteractiveGizmoManager.cpp b/Engine/Source/Editor/Experimental/EditorInteractiveToolsFramework/Private/EditorInteractiveGizmoManager.cpp
--- a/Engine/Source/Editor/Experimental/EditorInteractiveToolsFramework/Private/EditorInteractiveGizmoManager.cpp
+++ b/Engine/Source/Editor/Experimental/EditorInteractiveToolsFramework/Private/EditorInteractiveGizmoManager.cpp
@@ -155,15 +155,27 @@ void UEditorInteractiveGizmoManager::InitializeWithEditorModeManager(IToolsConte
 	Super::Initialize(QueriesAPIIn, TransactionsAPIIn, InputRouterIn);
 	EditorModeManager = InEditorModeManager;
 
-	// @note: most of the registration occurs in OnEditorModeChanged to ensure external callers to Deactivate/Activate mode don't invalidate the state 
-	EditorModeManager->OnEditorModeIDChanged().AddUObject(this, &UEditorInteractiveGizmoManager::OnEditorModeChanged);
+	if (UModeManagerInteractiveToolsContext* InteractiveToolsContext = EditorModeManager->GetInteractiveToolsContext())
+	{
+		// @note: this doesn't re-register if a snapping manager already exists
+		UE::Editor::Gizmos::RegisterSceneSnappingManager(InteractiveToolsContext);
+	}
 }
 
 void UEditorInteractiveGizmoManager::Shutdown()
 {
+	// Only the ModeManager-level GizmoManager should deregister the snapping manager.
+	// Per-mode GizmoManagers must not touch it - their Shutdown() runs deferred (after the
+	// next mode has already entered), which would remove the snapping manager out from under it.
 	if (EditorModeManager)
 	{
-		EditorModeManager->OnEditorModeIDChanged().RemoveAll(this);
+		if (UModeManagerInteractiveToolsContext* InteractiveToolsContext = EditorModeManager->GetInteractiveToolsContext())
+		{
+			if (InteractiveToolsContext->GizmoManager == this)
+			{
+				UE::Editor::Gizmos::DeregisterSceneSnappingManager(InteractiveToolsContext);
+			}
+		}
 	}
 
 	Registry->Shutdown();
@@ -204,28 +216,6 @@ UTransformGizmo* UEditorInteractiveGizmoManager::FindDefaultTransformGizmo() con
 	return Cast<UTransformGizmo>( FindGizmoByInstanceIdentifier(TransformInstanceIdentifier()) );
 }
 
-void UEditorInteractiveGizmoManager::OnEditorModeChanged(const FName& InModeID, bool bInIsEnteringMode)
-{
-	// When a mode is switched, shutdown is called after initializing the other mode.
-	// When the system is shared between modes (ie. snapping), this can de-register the snapping AFTER being registered in init
-	// This callback allows us to de-register things before the other mode registers
-
-	if (EditorModeManager)
-	{
-		UModeManagerInteractiveToolsContext* InteractiveToolsContext = EditorModeManager->GetInteractiveToolsContext();
-		check(InteractiveToolsContext);
-
-		if (bInIsEnteringMode)
-		{
-			UE::Editor::Gizmos::RegisterSceneSnappingManager(InteractiveToolsContext);
-		}
-		else // Only execute on exit (not enter)
-		{
-			UE::Editor::Gizmos::DeregisterSceneSnappingManager(InteractiveToolsContext);
-		}
-	}
-}
-
 bool UEditorInteractiveGizmoManager::DestroyEditorGizmo(UInteractiveGizmo* Gizmo)
 {
 	auto Pred = [Gizmo](const FActiveEditorGizmo& ActiveEditorGizmo) {return ActiveEditorGizmo.Gizmo == Gizmo; };
diff --git a/Engine/Source/Editor/Experimental/EditorInteractiveToolsFramework/Private/Snapping/EditorSnappingManager.cpp b/Engine/Source/Editor/Experimental/EditorInteractiveToolsFramework/Private/Snapping/EditorSnappingManager.cpp
--- a/Engine/Source/Editor/Experimental/EditorInteractiveToolsFramework/Private/Snapping/EditorSnappingManager.cpp
+++ b/Engine/Source/Editor/Experimental/EditorInteractiveToolsFramework/Private/Snapping/EditorSnappingManager.cpp
@@ -38,14 +38,14 @@ namespace UE::Editor::Gizmos
 			return true;
 		}
 
-		UEditorSceneSnappingManager* SelectionManager = NewObject<UEditorSceneSnappingManager>(InToolsContext->ToolManager);
-		if (!ensure(SelectionManager))
+		UEditorSceneSnappingManager* SnappingManager = NewObject<UEditorSceneSnappingManager>(InToolsContext->ToolManager);
+		if (!ensure(SnappingManager))
 		{
 			return false;
 		}
 
-		SelectionManager->Initialize(InToolsContext);
-		InToolsContext->ContextObjectStore->AddContextObject(SelectionManager);
+		SnappingManager->Initialize(InToolsContext);
+		InToolsContext->ContextObjectStore->AddContextObject(SnappingManager);
 
 		return true;
 	}
diff --git a/Engine/Source/Editor/Experimental/EditorInteractiveToolsFramework/Public/EditorInteractiveGizmoManager.h b/Engine/Source/Editor/Experimental/EditorInteractiveToolsFramework/Public/EditorInteractiveGizmoManager.h
--- a/Engine/Source/Editor/Experimental/EditorInteractiveToolsFramework/Public/EditorInteractiveGizmoManager.h
+++ b/Engine/Source/Editor/Experimental/EditorInteractiveToolsFramework/Public/EditorInteractiveGizmoManager.h
@@ -250,8 +250,6 @@ private:
 
 	/** Returns the existing default Gizmo instance if any. */
 	UE_API UTransformGizmo* FindDefaultTransformGizmo() const;
-
-	UE_API void OnEditorModeChanged(const FName& InModeID, bool bInIsEnteringMode);
 };
 
 #undef UE_API
