[ITF Gizmos] Gizmo labels offset when piloting camera

ViewCenter is now computed from unconstrained view rect. This prevents issues due to letterboxed viewports (e.g. piloting a camera)

#Synced-CL 52059832

diff --git a/Engine/Source/Editor/Experimental/EditorInteractiveToolsFramework/Private/EditorGizmos/GizmoElementWidget.cpp b/Engine/Source/Editor/Experimental/EditorInteractiveToolsFramework/Private/EditorGizmos/GizmoElementWidget.cpp
--- a/Engine/Source/Editor/Experimental/EditorInteractiveToolsFramework/Private/EditorGizmos/GizmoElementWidget.cpp
+++ b/Engine/Source/Editor/Experimental/EditorInteractiveToolsFramework/Private/EditorGizmos/GizmoElementWidget.cpp
@@ -41,8 +41,10 @@ void UGizmoElementWidget::DrawHUD(FCanvas* Canvas, IToolsContextRenderAPI* Rende
 			if (SceneView->WorldToPixel(CurrentRenderState.LocalToWorldTransform.GetLocation(), ScreenLocation))
 			{
 				const double DPIScale = Canvas->GetDPIScale();
-				const FIntPoint ViewSize(SceneView->UnscaledViewRect.Width(), SceneView->UnscaledViewRect.Height());
-				const FIntPoint ViewCenter(ViewSize.X / 2, ViewSize.Y / 2);
+
+				// Compute ViewCenter using unconstrained view rectangle (prevents offsets when viewport has letterboxing)
+				const FVector2D ViewCenter = FVector2D(SceneView->UnconstrainedViewRect.Width(), SceneView->UnconstrainedViewRect.Height()) / 2.0f;
+
 				FVector2D WidgetLocation = FVector2D(ScreenLocation.X - ViewCenter.X, ScreenLocation.Y - ViewCenter.Y);
 
 				const FVector2D ScaledOffset2D = Offset2D * Offset2DMultiplier;
