[ITF Gizmos] Rotation Precision Mode

Direct rotation interactions now have precision mode:

- Ctrl + Drag: slower rotation (smaller increments)
- Ctrl + Shift + Drag: faster rotation (bigger increments)

#Synced-CL 52059832

diff --git a/Engine/Source/Editor/Experimental/EditorInteractiveToolsFramework/Private/EditorGizmos/EditorTRSGizmo.cpp b/Engine/Source/Editor/Experimental/EditorInteractiveToolsFramework/Private/EditorGizmos/EditorTRSGizmo.cpp
--- a/Engine/Source/Editor/Experimental/EditorInteractiveToolsFramework/Private/EditorGizmos/EditorTRSGizmo.cpp
+++ b/Engine/Source/Editor/Experimental/EditorInteractiveToolsFramework/Private/EditorGizmos/EditorTRSGizmo.cpp
@@ -1413,7 +1413,7 @@ void UEditorTRSGizmo::SetupBehaviors()
 		TAttribute<bool>::CreateLambda(
 			[this]()
 			{
-				return CameraFollowsMovement();
+				return CameraFollowsMovement() || IsRotationPrecisionMode();
 			}
 		)
 	);
@@ -1582,6 +1582,15 @@ void UEditorTRSGizmo::SetupIndirectBehaviors()
 			OnTerminateSingleClickAndDragSequence();
 		}
 	};
+	MiddleClickDragBehavior->SetUsesUnboundedCursor(
+		TAttribute<bool>::CreateLambda(
+			[this]()
+			{
+				return CameraFollowsMovement() || IsRotationPrecisionMode();
+			}
+		)
+	);
+
 	AddInputBehavior(MiddleClickDragBehavior);
 
 	// Add left/right mouse input behavior for indirect manipulation
@@ -1965,7 +1974,7 @@ void UEditorTRSGizmo::UpdateCursor(const EGizmoElementInteractionState InState,
 		}
 		else
 		{
-			if (CameraFollowsMovement())
+			if (CameraFollowsMovement() || IsRotationPrecisionMode())
 			{
 				CursorAPI->SetCursorOverride(EMouseCursor::None);
 			}
@@ -3966,6 +3975,43 @@ bool UEditorTRSGizmo::AllowsDragDuplicate() const
 	return bDuplicateInCurrentMode && !bUpdatingPivot;
 }
 
+bool UEditorTRSGizmo::IsRotationPrecisionMode() const
+{
+	return bCtrlKeyDown && bInInteraction && !bIndirectManipulation && CurrentMode == EGizmoTransformMode::Rotate;
+}
+
+double UEditorTRSGizmo::GetRotationPrecisionModeMultiplier() const
+{
+	double Multiplier = 1.0;
+
+	if (IsRotationPrecisionMode())
+	{
+		if (LastHitPart == ETransformGizmoPartIdentifier::RotateScreenSpace)
+		{
+			Multiplier = bShiftKeyDown ? Interaction.RotateInteraction.ScreenSpaceRotationPrecisionBoost : Interaction.RotateInteraction.ScreenSpaceRotationPrecisionDamping;
+		}
+		else if (LastHitPart == ETransformGizmoPartIdentifier::RotateArcball)
+		{
+			Multiplier = bShiftKeyDown ? Interaction.RotateInteraction.ArcballRotationPrecisionBoost : Interaction.RotateInteraction.ArcballRotationPrecisionDamping;
+
+			if (FMath::IsNearlyZero(Multiplier))
+			{
+				Multiplier = 1.0;
+			}
+			else
+			{
+				Multiplier = 1.0/Multiplier;
+			}
+		}
+		else if (UE::Editor::InteractiveToolsFramework::Internal::IsAxisSingular(GetInteractionAxisList()))
+		{
+			Multiplier = bShiftKeyDown ? Interaction.RotateInteraction.AxisRotationPrecisionBoost : Interaction.RotateInteraction.AxisRotationPrecisionDamping;
+		}
+	}
+
+	return Multiplier;
+}
+
 FEditorViewportClient* UEditorTRSGizmo::GetEditorViewportClient() const
 {
 	if (const UInteractiveGizmoManager* const GizmoManager = GetGizmoManager())
@@ -4073,6 +4119,22 @@ void UEditorTRSGizmo::OnClickDragRotateAxisInternal(
 	using namespace UE::Geometry;
 	using namespace UE::Editor::GizmoMath;
 
+	// Precision mode switches to pull
+	if (IsRotationPrecisionMode() && RotateMode != EAxisRotateMode::Pull)
+	{
+		SetRotateMode(EAxisRotateMode::Pull);
+	}
+
+	if (bRotationPrecisionModeDirty)
+	{
+		if (!IsRotationPrecisionMode() && !FMath::IsNearlyZero(InteractionCurrAngle))
+		{
+			InteractionScreenCurrPos = InDragPos.ScreenPosition;
+		}
+
+		bRotationPrecisionModeDirty = false;
+	}
+
 	// Always perform ray/plane intersection
 	bool bHitPlane = false;
 	FVector HitPoint = FVector::ZeroVector;
@@ -4105,6 +4167,12 @@ void UEditorTRSGizmo::OnClickDragRotateAxisInternal(
 		Angle = PreviousAngle + FMath::FindDeltaAngleRadians(FMath::UnwindRadians(PreviousAngle), Angle);
 
 		DeltaAngle = FMath::FindDeltaAngleRadians(PreviousAngle, Angle);
+
+		if (IsRotationPrecisionMode())
+		{
+			DeltaAngle *= GetRotationPrecisionModeMultiplier();
+		}
+
 		SnapRotateAngleDelta(DeltaAngle, GetInteractionAxisList());
 
 		// The delta as it applies, corrected for the XY vs. Z discrepancy
@@ -4279,6 +4347,7 @@ void UEditorTRSGizmo::UpdateRotateAxisDelta(UGizmoElementRotateAxis* InElement)
 	DeltaParameters.Angle = (InteractionStartAngle + InteractionDelta.X) - AlignedInteractionAngleOffset;
 	DeltaParameters.CursorLocation = WorldCursorPos;
 	DeltaParameters.CursorLocation2D = InteractionPlanarCurrPoint2D;
+	DeltaParameters.RotateMode = RotateMode;
 
 	UE_CLOG(bLogRotation, LogEditorTRSGizmo, Verbose, TEXT("InteractionDelta: %.1f"), FMath::RadiansToDegrees(InteractionDelta.X));
 
@@ -4749,6 +4818,11 @@ void UEditorTRSGizmo::OnClickPress(const FInputDeviceRay& PressPos)
 		// Ensure any new cursor has the correct rotation before displaying
 		UpdateCursorRotation(PressPos);
 	}
+
+	if (bCtrlKeyDown)
+	{
+		bRotationPrecisionModeDirty = true;
+	}
 }
 
 void UEditorTRSGizmo::OnClickDrag(const FInputDeviceRay& DragPos)
@@ -4826,7 +4900,7 @@ void UEditorTRSGizmo::OnDragStart(const FInputDeviceRay& InDragPos)
 void UEditorTRSGizmo::OnClickRelease(const FInputDeviceRay& InReleasePos, bool bInIsDragOperation)
 {
 	// If leaving a drag in which camera used to follow movement, reset mouse cursor position
-	if (CameraFollowsMovement())
+	if (CameraFollowsMovement() || IsRotationPrecisionMode())
 	{
 		RestoreCursorPosition();
 	}
@@ -4841,6 +4915,7 @@ void UEditorTRSGizmo::OnClickRelease(const FInputDeviceRay& InReleasePos, bool b
 
 	EndTransformEditSequence();
 
+	bIndirectManipulation = false;
 	bInInteraction = false;
 	bOnlyUpdatePivot = false;
 
@@ -4943,8 +5018,10 @@ void UEditorTRSGizmo::OnForceEndCapture()
 
 	bCameraFollowsMovement = false;
 	bResetDragTranslate = false;
+	bRotationPrecisionModeDirty = false;
 	bAltKeyDown = false;
 	bCtrlKeyDown = false;
+	bShiftKeyDown = false;
 	NudgeData.Reset();
 }
 
@@ -4952,6 +5029,18 @@ void UEditorTRSGizmo::OnUpdateModifierState(int ModifierID, bool bIsOn)
 {
 	if (ModifierID == UE::Editor::ViewportInteractions::ShiftKeyMod)
 	{
+		if (bShiftKeyDown != bIsOn)
+		{
+			bShiftKeyDown = bIsOn;
+
+			// In precision mode, shift is used for faster movements
+			if (IsRotationPrecisionMode())
+			{
+				PendingDragFunction.Reset();
+				bRotationPrecisionModeDirty = true;
+			}
+		}
+
 		if (bCameraFollowsMovement != bIsOn)
 		{
 			bResetDragTranslate = true;
@@ -4994,7 +5083,32 @@ void UEditorTRSGizmo::OnUpdateModifierState(int ModifierID, bool bIsOn)
 	}
 	else if (ModifierID == UE::Editor::ViewportInteractions::CtrlKeyMod)
 	{
-		bCtrlKeyDown = bIsOn;
+		if (bCtrlKeyDown != bIsOn)
+		{
+			bCtrlKeyDown = bIsOn;
+
+			if (CurrentMode == EGizmoTransformMode::Rotate)
+			{
+				PendingDragFunction.Reset();
+				bRotationPrecisionModeDirty = true;
+
+				if (bInInteraction && !bIndirectManipulation)
+				{
+					if (IsRotationPrecisionMode())
+					{
+						// Entering precision mode: cache cursor position
+						CacheCursorPosition();
+					}
+					else
+					{
+						// Precision mode ends: restore cursor position
+						RestoreCursorPosition();
+					}
+				}
+
+				UpdateCursor(EGizmoElementInteractionState::Interacting, LastHitPart);
+			}
+		}
 	}
 }
 
@@ -5044,7 +5158,7 @@ void UEditorTRSGizmo::BeginTransformEditSequence()
 	UE_LOG(LogEditorTRSGizmo, Verbose, TEXT("StartTransform.Scale: %.1f"), StartTransform.GetScale3D().X);
 
 	// Camera follows starting: cache cursor position
-	if (bCameraFollowsMovement)
+	if (bCameraFollowsMovement || IsRotationPrecisionMode())
 	{
 		CacheCursorPosition();
 	}
@@ -5850,6 +5964,11 @@ void UEditorTRSGizmo::OnClickPressRotateAxis(const FInputDeviceRay& InPressPos)
 
 	bInInteraction = true;
 
+	if (IsRotationPrecisionMode())
+	{
+		bRotationPrecisionModeDirty = true;
+	}
+
 	SetModeLastHitPart(EGizmoTransformMode::Rotate, LastHitPart);
 }
 
@@ -7093,6 +7212,13 @@ void UEditorTRSGizmo::OnClickPressScreenSpaceRotate(const FInputDeviceRay& InPre
 
 	bInInteraction = true;
 
+	if (IsRotationPrecisionMode())
+	{
+		bRotationPrecisionModeDirty = true;
+	}
+
+	ScreenSpacePrecisionVisualOffset = -FMath::FindDeltaAngleRadians(InteractionCurrAngle, HitAngle);
+
 	SetModeLastHitPart(EGizmoTransformMode::Rotate, LastHitPart);
 }
 
@@ -7100,6 +7226,26 @@ void UEditorTRSGizmo::OnClickDragScreenSpaceRotate(const FInputDeviceRay& InDrag
 {
 	Profiler->BeginScoped(UE::Editor::InteractiveToolsFramework::Private::TransformGizmoLocals::DragRotateScreenName);
 
+	// Making sure current interaction angle and visual display are correct when switching precision mode On/Off
+	// This is required since precision mode has the gizmo lose the 1:1 mapping between cursor movements and actual rotation
+	if (bRotationPrecisionModeDirty)
+	{
+		InteractionScreenAxisDirection = GetScreenRotateAxisDir(InDragPos);
+		InteractionScreenCurrPos = InDragPos.ScreenPosition;
+
+		double HitAngle = 0.0;
+		FVector HitPoint = FVector::ZeroVector;
+		if (UE::Editor::GizmoMath::ComputeAngleInPlaneFromScreen(GizmoViewContext, InDragPos.ScreenPosition, InteractionPlane, HitPoint, HitAngle))
+		{
+			const double AngleOffset = FMath::FindDeltaAngleRadians(InteractionCurrAngle, HitAngle);
+
+			InteractionCurrAngle += AngleOffset;
+			ScreenSpacePrecisionVisualOffset -= AngleOffset;
+		}
+
+		bRotationPrecisionModeDirty = false;
+	}
+
 	check(GizmoViewContext);
 
 	const double PreviousAngle = InteractionCurrAngle;
@@ -7118,7 +7264,7 @@ void UEditorTRSGizmo::OnClickDragScreenSpaceRotate(const FInputDeviceRay& InDrag
 		InteractionCurrAngle = HitAngle;
 	}
 
-	if (bIndirectManipulation)
+	if (bIndirectManipulation || IsRotationPrecisionMode())
 	{
 		FVector2D DragDir = InDragPos.ScreenPosition - PreviousScreenPos;
 
@@ -7139,6 +7285,12 @@ void UEditorTRSGizmo::OnClickDragScreenSpaceRotate(const FInputDeviceRay& InDrag
 		DragDir -= DragDirToRemove;
 
 		double DeltaAngle = FMath::DegreesToRadians(FVector2D::DotProduct(InteractionScreenAxisDirection, DragDir) * Interaction.RotateInteraction.PullMultiplier);
+
+		if (IsRotationPrecisionMode())
+		{
+			DeltaAngle *= FMath::Abs(GetRotationPrecisionModeMultiplier());
+		}
+
 		SnapRotateAngleDelta(DeltaAngle, EAxisList::X); // Only snapping a single angle, so axis list doesn't matter
 
 		// Reapply the (potentially) snapped delta angle
@@ -7181,7 +7333,7 @@ void UEditorTRSGizmo::OnClickDragScreenSpaceRotate(const FInputDeviceRay& InDrag
 	if (RotateScreenSpaceElement2)
 	{
 		// Compensate for GizmoElement screen alignment
-		const double AngleForDelta = InteractionCurrAngle * -1;
+		const double AngleForDelta = (InteractionCurrAngle + ScreenSpacePrecisionVisualOffset) * -1;
 
 		UGizmoElementRotateAxis::FDeltaParameters DeltaParameters;
 		DeltaParameters.Transform = GetActiveTransform();
@@ -7281,7 +7433,7 @@ void UEditorTRSGizmo::OnClickPressArcBallRotate(const FInputDeviceRay& PressPos)
 
 	const FVector& RayOrigin = PressPos.WorldRay.Origin;
 	const FVector& RayDir = PressPos.WorldRay.Direction;
-	const double SphereRadius = GetWorldRadius(Style.RotateStyle.ArcballStyle.Radius * Style.GetModifiedAxisSizeMultiplier());
+	const double SphereRadius = GetWorldRadius(Style.RotateStyle.ArcballStyle.Radius * Style.GetModifiedAxisSizeMultiplier()) * GetRotationPrecisionModeMultiplier();
 
 	StartRotation = CurrentRotation = CurrentTransform.GetRotation();
 	InteractionPlanarOrigin = CurrentTransform.GetLocation();
@@ -7341,7 +7493,25 @@ void UEditorTRSGizmo::OnClickDragArcBallRotate(const FInputDeviceRay& DragPos)
 
 	const FVector& RayOrigin = DragPos.WorldRay.Origin;
 	const FVector& RayDir = DragPos.WorldRay.Direction;
-	const float SphereRadius = GetWorldRadius(Style.RotateStyle.ArcballStyle.Radius * Style.GetModifiedAxisSizeMultiplier());
+
+	const double SphereRadius = GetWorldRadius(Style.RotateStyle.ArcballStyle.Radius * Style.GetModifiedAxisSizeMultiplier()) * GetRotationPrecisionModeMultiplier();
+
+	// Precision mode: make sure rotation keeps working when entering/leaving precision mode
+	if (bRotationPrecisionModeDirty)
+	{
+		// Reproject cursor onto the new sphere to prevent a rotation jump
+		FVector TransitionCurrPoint;
+		if (ArcBallLocals2::GetSphereAndHyperbolicProjection(
+				InteractionPlanarOrigin, SphereRadius, RayOrigin, RayDir, *GizmoViewContext, TransitionCurrPoint
+			))
+		{
+			InteractionArcBallStartPoint = TransitionCurrPoint;
+			StartRotation = CurrentRotation;
+		}
+
+		InteractionScreenCurrPos = DragPos.ScreenPosition;
+		bRotationPrecisionModeDirty = false;
+	}
 
 	// compute projection
 	ArcBallLocals2::GetSphereAndHyperbolicProjection(
diff --git a/Engine/Source/Editor/Experimental/EditorInteractiveToolsFramework/Private/EditorGizmos/GizmoElementRotateAxis.cpp b/Engine/Source/Editor/Experimental/EditorInteractiveToolsFramework/Private/EditorGizmos/GizmoElementRotateAxis.cpp
--- a/Engine/Source/Editor/Experimental/EditorInteractiveToolsFramework/Private/EditorGizmos/GizmoElementRotateAxis.cpp
+++ b/Engine/Source/Editor/Experimental/EditorInteractiveToolsFramework/Private/EditorGizmos/GizmoElementRotateAxis.cpp
@@ -922,4 +922,5 @@ void UGizmoElementRotateAxis::FState::Update(const FDeltaParameters& InParameter
 {
 	Transform = InParameters.Transform;
 	Angle = InParameters.Angle;
+	bShowCursor = !InParameters.bIsIndirectManipulation && InParameters.RotateMode != EAxisRotateMode::Pull;
 }
diff --git a/Engine/Source/Editor/Experimental/EditorInteractiveToolsFramework/Public/EditorGizmos/EditorTRSGizmo.h b/Engine/Source/Editor/Experimental/EditorInteractiveToolsFramework/Public/EditorGizmos/EditorTRSGizmo.h
--- a/Engine/Source/Editor/Experimental/EditorInteractiveToolsFramework/Public/EditorGizmos/EditorTRSGizmo.h
+++ b/Engine/Source/Editor/Experimental/EditorInteractiveToolsFramework/Public/EditorGizmos/EditorTRSGizmo.h
@@ -650,6 +650,9 @@ private:
 	void DuplicateSelection(bool bDroppingDuplicate = false) const;
 	bool AllowsDragDuplicate() const;
 
+	bool IsRotationPrecisionMode() const;
+	double GetRotationPrecisionModeMultiplier() const;
+
 	FEditorViewportClient* GetEditorViewportClient() const;
 
 #pragma region Translation
@@ -820,12 +823,16 @@ protected:
 	bool bDetectedNonStandardScaling = false;
 	
 	bool bCameraFollowsMovement = false;
+	bool bResetDragTranslate = false;
 
 	bool bCtrlKeyDown = false;
+	bool bAltKeyDown = false;
+	bool bShiftKeyDown = false;
 
-	bool bResetDragTranslate = false;
+	bool bRotationPrecisionModeDirty = false;
 
-	bool bAltKeyDown = false;
+	/** Screen space precision rotation requires to keep track of rotation amount, to prevent UI inconsistencies */
+	double ScreenSpacePrecisionVisualOffset = 0.0;
 
 	bool bZeroCrossingDeltaScaleSwap = false;
 
diff --git a/Engine/Source/Editor/Experimental/EditorInteractiveToolsFramework/Public/EditorGizmos/GizmoElementRotate.h b/Engine/Source/Editor/Experimental/EditorInteractiveToolsFramework/Public/EditorGizmos/GizmoElementRotate.h
--- a/Engine/Source/Editor/Experimental/EditorInteractiveToolsFramework/Public/EditorGizmos/GizmoElementRotate.h
+++ b/Engine/Source/Editor/Experimental/EditorInteractiveToolsFramework/Public/EditorGizmos/GizmoElementRotate.h
@@ -221,6 +221,30 @@ struct FGizmoElementRotateInteraction
 	
 	/** Returns the overridden RotateMode if set, otherwise use the Editor Setting. */
 	UE_API EAxisRotateMode::Type GetRotateMode() const;
+
+	/** Arc rotation damping factor (precision mode) */
+	UPROPERTY(EditAnywhere, Category = "Precision Mode", meta = (ClampMin = 0.01f))
+	float ArcballRotationPrecisionDamping = 0.05f;
+
+	/** Axis rotation damping factor (precision mode) */
+	UPROPERTY(EditAnywhere, Category = "Precision Mode", meta = (ClampMin = 0.01f))
+	float AxisRotationPrecisionDamping = 0.025f;
+
+	/** Screen Space rotation damping factor (precision mode) */
+	UPROPERTY(EditAnywhere, Category = "Precision Mode", meta = (ClampMin = 0.01f))
+	float ScreenSpaceRotationPrecisionDamping = 0.025f;
+
+	/** Arc rotation boost factor (precision mode) */
+	UPROPERTY(EditAnywhere, Category = "Precision Mode", meta = (ClampMin = 0.01f))
+	float ArcballRotationPrecisionBoost = 2.0f;
+
+	/** Axis rotation boost factor (precision mode) */
+	UPROPERTY(EditAnywhere, Category = "Precision Mode", meta = (ClampMin = 0.01f))
+	float AxisRotationPrecisionBoost = 4.0f;
+
+	/** Screen Space rotation boost factor (precision mode) */
+	UPROPERTY(EditAnywhere, Category = "Precision Mode", meta = (ClampMin = 0.01f))
+	float ScreenSpaceRotationPrecisionBoost = 4.0f;
 };
 
 #undef UE_API
