[ITF Gizmos] Regression: cannot scale elements with zero rotation in orthographic mode

The move from EditorTRSGizmo -> EditorInteractiveToolsFramework (CL 52057727), includes a small change introducing the scaling issue.
Prior to that CL, code also had unreachable branches, which are removed in this CL, further simplifying it.
The removed if/else was essentially checking `if (DotProduct(PlaneNormal, PlaneNormal) < 0.1)`, which is always true, so the else(s) could never be reached.
A todo note was also mentioning the possibility that the if/else branches were not needed.

#Synced-CL 52139128

diff --git a/Engine/Source/Editor/Experimental/EditorInteractiveToolsFramework/Private/EditorGizmos/EditorGizmoElementSharedInternal.cpp b/Engine/Source/Editor/Experimental/EditorInteractiveToolsFramework/Private/EditorGizmos/EditorGizmoElementSharedInternal.cpp
--- a/Engine/Source/Editor/Experimental/EditorInteractiveToolsFramework/Private/EditorGizmos/EditorGizmoElementSharedInternal.cpp
+++ b/Engine/Source/Editor/Experimental/EditorInteractiveToolsFramework/Private/EditorGizmos/EditorGizmoElementSharedInternal.cpp
@@ -181,29 +181,11 @@ namespace UE::Editor::InteractiveToolsFramework::Internal
 		BreakPlaneFrame(InPlane, PlaneNormal, PlaneUp, PlaneSide);
 
 		const FVector ViewDirection = InViewContext->GetViewDirection();
-		const FVector::FReal NormalDot = FMath::Abs(FVector::DotProduct(ViewDirection, PlaneNormal));
 		const FVector::FReal UpDot = FMath::Abs(FVector::DotProduct(ViewDirection, PlaneUp));
 		const FVector::FReal SideDot = FMath::Abs(FVector::DotProduct(ViewDirection, PlaneSide));
 
-		constexpr double DotThreshold = 0.1;
-
-		// @todo: do we need to test this?
-		FVector PrimaryDirection = PlaneNormal; // X Axis is the plane normal
-		if (NormalDot > DotThreshold)
-		{
-			PrimaryDirection = (UpDot > SideDot) ? PlaneUp : PlaneSide;
-			OutDirectionToRemove = (UpDot > SideDot) ? PlaneSide : PlaneUp;
-		}
-		else if (UpDot > DotThreshold)
-		{
-			PrimaryDirection = (NormalDot > SideDot) ? PlaneNormal : PlaneSide;
-			OutDirectionToRemove = (NormalDot > SideDot) ? PlaneSide : PlaneNormal;
-		}
-		else
-		{
-			PrimaryDirection = (NormalDot > UpDot) ? PlaneNormal : PlaneUp;
-			OutDirectionToRemove = (NormalDot > UpDot) ? PlaneUp : PlaneNormal;
-		}
+		const FVector PrimaryDirection = (UpDot > SideDot) ? PlaneUp : PlaneSide;
+		OutDirectionToRemove = (UpDot > SideDot) ? PlaneSide : PlaneUp;
 
 		Geometry::FFrame3d ModifiedPlane = InPlane;
 		ModifiedPlane.AlignAxis(0, PrimaryDirection); // Align the plane normal to the primary direction
