[Viewport ITF]
- Fix an issue where UViewportClickDragBehavior didn't call OnDragStart() on BeginCapture() even when the mouse had moved enough to drag. This caused a single-frame-delay in state propegation.
- Made UViewportClickDragBehavior start the drag immediately on capture, bypassing the movment threshold. Use of the movement threshold can be restored with a new attribute. This allows dragging behaviors to immediately start dragging when there is no need to compete with a click.
- Add hover support to the LevelViewportClickSelection suite and ViewportLegacyClickInteraction. The hover behavior correctly pushes the HHitTarget cursor value to the the cursor override.

#rnx
#Synced-CL 52132857

diff --git a/Engine/Source/Editor/Experimental/EditorInteractiveToolsFramework/Private/Behaviors/ViewportClickDragBehavior.cpp b/Engine/Source/Editor/Experimental/EditorInteractiveToolsFramework/Private/Behaviors/ViewportClickDragBehavior.cpp
--- a/Engine/Source/Editor/Experimental/EditorInteractiveToolsFramework/Private/Behaviors/ViewportClickDragBehavior.cpp
+++ b/Engine/Source/Editor/Experimental/EditorInteractiveToolsFramework/Private/Behaviors/ViewportClickDragBehavior.cpp
@@ -93,6 +93,13 @@ FInputCaptureUpdate UViewportClickDragBehavior::BeginCapture(const FInputDeviceS
 	}
 
 	OnStateUpdatedInternal(InputState);
+	
+	if (!bDragging && (!RequireDistanceThresholdOnCapture.Get(false) || HasMovedEnoughForDrag()))
+	{
+		bDragging = true;
+		OnDragStartInternal(InputState);
+	}
+	
 	return FInputCaptureUpdate::Begin(this, EInputCaptureSide::Any);
 }
 
@@ -191,6 +198,11 @@ void UViewportClickDragBehavior::SetUsesUnboundedCursor(TAttribute<bool> InUsesU
 	UnboundedCursor = InUsesUnboundedCursor;
 }
 
+void UViewportClickDragBehavior::SetRequireDistanceThresholdOnCapture(TAttribute<bool> InRequireDistanceThresholdOnCapture)
+{
+	RequireDistanceThresholdOnCapture = InRequireDistanceThresholdOnCapture; 
+}
+
 void UViewportClickDragBehavior::ResetDrag(const FInputDeviceState& InputState)
 {
 	if (bDragging)
diff --git a/Engine/Source/Editor/Experimental/EditorInteractiveToolsFramework/Private/ViewportInteractions/ViewportClickInteraction.cpp b/Engine/Source/Editor/Experimental/EditorInteractiveToolsFramework/Private/ViewportInteractions/ViewportClickInteraction.cpp
--- a/Engine/Source/Editor/Experimental/EditorInteractiveToolsFramework/Private/ViewportInteractions/ViewportClickInteraction.cpp
+++ b/Engine/Source/Editor/Experimental/EditorInteractiveToolsFramework/Private/ViewportInteractions/ViewportClickInteraction.cpp
@@ -38,10 +38,6 @@ FInputRayHit UViewportClickInteraction::IsHitByClick(const FInputDeviceRay& InCl
 
 void UViewportClickInteraction::OnClickDown(const FInputDeviceRay& InClickPos)
 {
-	if (UViewportInteractionsBehaviorSource* const ViewportInteractionBehaviorSource = GetViewportInteractionsBehaviorSource())
-	{
-		ViewportInteractionBehaviorSource->SetMouseCursorOverride(EMouseCursor::Crosshairs);
-	}
 }
 
 void UViewportClickInteraction::OnStateUpdated(const FInputDeviceState& InInputDeviceState)
@@ -72,6 +68,21 @@ void UViewportClickInteraction::OnStateUpdated(const FInputDeviceState& InInputD
 	CacheChange(InInputDeviceState.Mouse.Left) || CacheChange(InInputDeviceState.Mouse.Right) || CacheChange(InInputDeviceState.Mouse.Middle);
 }
 
+HHitProxy* UViewportClickInteraction::GetHitProxy(const FInputDeviceRay& InClickPos) const
+{
+	if (FEditorViewportClient* Client = GetEditorViewportClient())
+	{
+		if (Client->Viewport)
+		{
+			const int32 CurrentClickPosX = static_cast<int32>(InClickPos.ScreenPosition.X);
+			const int32 CurrentClickPosY = static_cast<int32>(InClickPos.ScreenPosition.Y);
+		
+			return Client->Viewport->GetHitProxy(CurrentClickPosX, CurrentClickPosY);
+		}		
+	}
+	return nullptr;
+}
+
 void UViewportClickInteraction::OnClickUp(const FInputDeviceRay& InClickPos)
 {
 	if (UViewportInteractionsBehaviorSource* const ViewportInteractionBehaviorSource = GetViewportInteractionsBehaviorSource())
diff --git a/Engine/Source/Editor/Experimental/EditorInteractiveToolsFramework/Private/ViewportInteractions/ViewportDragInteraction.cpp b/Engine/Source/Editor/Experimental/EditorInteractiveToolsFramework/Private/ViewportInteractions/ViewportDragInteraction.cpp
--- a/Engine/Source/Editor/Experimental/EditorInteractiveToolsFramework/Private/ViewportInteractions/ViewportDragInteraction.cpp
+++ b/Engine/Source/Editor/Experimental/EditorInteractiveToolsFramework/Private/ViewportInteractions/ViewportDragInteraction.cpp
@@ -30,14 +30,6 @@ FInputRayHit UViewportDragInteraction::CanBeginClickDragSequence(const FInputDev
 	return CanBeActivated() ? FInputRayHit(TNumericLimits<double>::Max()) : FInputRayHit();
 }
 
-void UViewportDragInteraction::OnBeginCapture(const FInputDeviceRay& InClickPressPos)
-{
-	if (UViewportInteractionsBehaviorSource* BehaviorSource = GetViewportInteractionsBehaviorSource())
-	{
-		BehaviorSource->SetMouseCursorOverride(EMouseCursor::Crosshairs);
-	}
-}
-
 void UViewportDragInteraction::OnDragStart(const FInputDeviceRay& InDragStartPos)
 {
 	if (!IsEnabled())
diff --git a/Engine/Source/Editor/Experimental/EditorInteractiveToolsFramework/Private/ViewportInteractions/ViewportInteraction.cpp b/Engine/Source/Editor/Experimental/EditorInteractiveToolsFramework/Private/ViewportInteractions/ViewportInteraction.cpp
--- a/Engine/Source/Editor/Experimental/EditorInteractiveToolsFramework/Private/ViewportInteractions/ViewportInteraction.cpp
+++ b/Engine/Source/Editor/Experimental/EditorInteractiveToolsFramework/Private/ViewportInteractions/ViewportInteraction.cpp
@@ -157,6 +157,22 @@ FEditorViewportClient* UViewportInteraction::GetEditorViewportClient() const
 	return nullptr;
 }
 
+void UViewportInteraction::SetMouseCursorOverride(EMouseCursor::Type InMouseCursor)
+{
+	if (UViewportInteractionsBehaviorSource* Source = GetViewportInteractionsBehaviorSource())
+	{
+		Source->SetMouseCursorOverride(InMouseCursor);
+	}
+}
+
+void UViewportInteraction::ClearMouseCursorOverride()
+{
+	if (UViewportInteractionsBehaviorSource* Source = GetViewportInteractionsBehaviorSource())
+	{
+		Source->ClearMouseCursorOverride();
+	}
+}
+
 void UViewportInteraction::OnUserDefinedChordChanged(const FUICommandInfo& InCommandInfo)
 {
 	for (const TSharedPtr<FUICommandInfo>& Command : GetCommands())
diff --git a/Engine/Source/Editor/Experimental/EditorInteractiveToolsFramework/Private/ViewportInteractions/ViewportLegacyClickInteraction.cpp b/Engine/Source/Editor/Experimental/EditorInteractiveToolsFramework/Private/ViewportInteractions/ViewportLegacyClickInteraction.cpp
--- a/Engine/Source/Editor/Experimental/EditorInteractiveToolsFramework/Private/ViewportInteractions/ViewportLegacyClickInteraction.cpp
+++ b/Engine/Source/Editor/Experimental/EditorInteractiveToolsFramework/Private/ViewportInteractions/ViewportLegacyClickInteraction.cpp
@@ -4,6 +4,7 @@
 
 #include "EditorModeManager.h"
 #include "EditorViewportClient.h"
+#include "BaseBehaviors/MouseHoverBehavior.h"
 
 UViewportLegacyClickInteraction::UViewportLegacyClickInteraction()
 {
@@ -15,10 +16,10 @@ void UViewportLegacyClickInteraction::BuildBehaviors()
 {
 	Super::BuildBehaviors();
 	
+	using namespace UE::Editor::ViewportInteractions;
+	
 	if (UViewportClickBehavior* ClickBehavior = ClickBehaviorWeak.Get())
 	{
-		using namespace UE::Editor::ViewportInteractions;
-	
 		// Lower priority by one step so that custom UViewportClickInteractions do not have to compete.
 		ClickBehavior->SetDefaultPriority(FInputCapturePriority(CLICK_PRIORITY).MakeLower());
 		
@@ -33,6 +34,49 @@ void UViewportLegacyClickInteraction::BuildBehaviors()
 			FButtonBinding(EKeys::LeftShift).Required(false)
 		});
 	}
+	
+	UMouseHoverBehavior* HoverBehavior = NewObject<UMouseHoverBehavior>();
+	HoverBehavior->Initialize(this);
+	HoverBehavior->SetDefaultPriority(CLICK_PRIORITY);
+	RegisterInputBehavior(HoverBehavior);
+}
+
+void UViewportLegacyClickInteraction::OnClickDown(const FInputDeviceRay& InClickPos)
+{
+	if (!OnUpdateHover(InClickPos))
+	{
+		SetMouseCursorOverride(EMouseCursor::Default);
+	}
+}
+
+FInputRayHit UViewportLegacyClickInteraction::BeginHoverSequenceHitTest(const FInputDeviceRay& PressPos)
+{
+	if (GetHitProxy(PressPos))
+	{
+		return FInputRayHit(0.0);
+	}
+	return FInputRayHit();
+}
+
+void UViewportLegacyClickInteraction::OnBeginHover(const FInputDeviceRay& DevicePos)
+{
+	OnUpdateHover(DevicePos);
+}
+
+bool UViewportLegacyClickInteraction::OnUpdateHover(const FInputDeviceRay& DevicePos)
+{
+	if (HHitProxy* HitProxy = GetHitProxy(DevicePos))
+	{
+		SetMouseCursorOverride(HitProxy->GetMouseCursor());
+		return true;
+	}
+	
+	return false;
+}
+
+void UViewportLegacyClickInteraction::OnEndHover()
+{
+	ClearMouseCursorOverride();
 }
 
 void UViewportLegacyClickInteraction::ProcessClick_Internal(const FViewportClick& InViewportClick, HHitProxy* InHitProxy, FSceneView& View)
diff --git a/Engine/Source/Editor/Experimental/EditorInteractiveToolsFramework/Public/Behaviors/ViewportClickDragBehavior.h b/Engine/Source/Editor/Experimental/EditorInteractiveToolsFramework/Public/Behaviors/ViewportClickDragBehavior.h
--- a/Engine/Source/Editor/Experimental/EditorInteractiveToolsFramework/Public/Behaviors/ViewportClickDragBehavior.h
+++ b/Engine/Source/Editor/Experimental/EditorInteractiveToolsFramework/Public/Behaviors/ViewportClickDragBehavior.h
@@ -113,6 +113,12 @@ public:
 	 */
 	UE_API void SetUsesUnboundedCursor(TAttribute<bool> InUsesUnboundedCursor);
 
+	/**
+	 * Call to set whether drag starts should be contingent on the distance threshold when this behavior has capture.
+	 * Default false. When false, this behavior will immediately start the drag operation when it wins capture.
+	 */
+	UE_API void SetRequireDistanceThresholdOnCapture(TAttribute<bool> InRequireDistanceThresholdOnCapture);
+
 	/**
 	 * Reset drag state and tracked drag distance. Calls OnDragEnd() but does not exit the behavior.
 	 */
@@ -120,9 +126,9 @@ public:
 
 	bool IsDragging() const { return bDragging; }
 	FVector2D GetMouseDownPosition() const { return MouseDownPosition; }
-
+	
 protected:
-
+	
 	UE_API virtual void OnDragStartInternal(const FInputDeviceState& InputState);
 	UE_API virtual void OnClickDragInternal(
 		const FInputDeviceState& InputState, const FInputCaptureData& Data
@@ -168,6 +174,7 @@ protected:
 	bool bDragging = false;
 
 	TAttribute<bool> UnboundedCursor;
+	TAttribute<bool> RequireDistanceThresholdOnCapture;
 };
 
 /**
diff --git a/Engine/Source/Editor/Experimental/EditorInteractiveToolsFramework/Public/ViewportInteractions/ViewportClickInteraction.h b/Engine/Source/Editor/Experimental/EditorInteractiveToolsFramework/Public/ViewportInteractions/ViewportClickInteraction.h
--- a/Engine/Source/Editor/Experimental/EditorInteractiveToolsFramework/Public/ViewportInteractions/ViewportClickInteraction.h
+++ b/Engine/Source/Editor/Experimental/EditorInteractiveToolsFramework/Public/ViewportInteractions/ViewportClickInteraction.h
@@ -47,9 +47,11 @@ public:
 protected:
 	UE_API virtual void BuildBehaviors();
 
+	UE_API HHitProxy* GetHitProxy(const FInputDeviceRay& InClickPos) const;
+
 	/** Override this function to provide custom functionality */
 	virtual void ProcessClick_Internal(const FViewportClick& InViewportClick, HHitProxy* InHitProxy, FSceneView& View) {}
-
+	
 	TWeakObjectPtr<UViewportClickBehavior> ClickBehaviorWeak;
 	
 	FKey LastChangedMouseButton = FKey();
diff --git a/Engine/Source/Editor/Experimental/EditorInteractiveToolsFramework/Public/ViewportInteractions/ViewportDragInteraction.h b/Engine/Source/Editor/Experimental/EditorInteractiveToolsFramework/Public/ViewportInteractions/ViewportDragInteraction.h
--- a/Engine/Source/Editor/Experimental/EditorInteractiveToolsFramework/Public/ViewportInteractions/ViewportDragInteraction.h
+++ b/Engine/Source/Editor/Experimental/EditorInteractiveToolsFramework/Public/ViewportInteractions/ViewportDragInteraction.h
@@ -32,7 +32,6 @@ public:
 
 	//~ Begin IViewportClickDragBehaviorTarget
 	UE_API virtual FInputRayHit CanBeginClickDragSequence(const FInputDeviceRay& InPressPos) override;
-	UE_API virtual void OnBeginCapture(const FInputDeviceRay& InClickPressPos) override;
 	UE_API virtual void OnDragStart(const FInputDeviceRay& InDragStartPos) override;
 	UE_API virtual void OnDrag(const FInputDeviceRay& InDragPos) override;
 	UE_API virtual void OnDragEnd(const FInputDeviceRay& InDragEndPos) override;
diff --git a/Engine/Source/Editor/Experimental/EditorInteractiveToolsFramework/Public/ViewportInteractions/ViewportInteraction.h b/Engine/Source/Editor/Experimental/EditorInteractiveToolsFramework/Public/ViewportInteractions/ViewportInteraction.h
--- a/Engine/Source/Editor/Experimental/EditorInteractiveToolsFramework/Public/ViewportInteractions/ViewportInteraction.h
+++ b/Engine/Source/Editor/Experimental/EditorInteractiveToolsFramework/Public/ViewportInteractions/ViewportInteraction.h
@@ -165,6 +165,9 @@ public:
 protected:
 	UE_API FEditorModeTools* GetModeTools() const;
 	UE_API FEditorViewportClient* GetEditorViewportClient() const;
+	
+	UE_API void SetMouseCursorOverride(EMouseCursor::Type InMouseCursor);
+	UE_API void ClearMouseCursorOverride();
 
 	/**
 	 * Define what should happen to UInputBehavior(s) when an "observed" command chord changes
diff --git a/Engine/Source/Editor/Experimental/EditorInteractiveToolsFramework/Public/ViewportInteractions/ViewportInteractionsBehaviorSource.h b/Engine/Source/Editor/Experimental/EditorInteractiveToolsFramework/Public/ViewportInteractions/ViewportInteractionsBehaviorSource.h
--- a/Engine/Source/Editor/Experimental/EditorInteractiveToolsFramework/Public/ViewportInteractions/ViewportInteractionsBehaviorSource.h
+++ b/Engine/Source/Editor/Experimental/EditorInteractiveToolsFramework/Public/ViewportInteractions/ViewportInteractionsBehaviorSource.h
@@ -22,6 +22,8 @@ class UTransformGizmo;
 class UTransformProxy;
 class UViewportInteraction;
 
+#define UE_API EDITORINTERACTIVETOOLSFRAMEWORK_API
+
 namespace UE::Editor::ViewportInteractions
 {
 
@@ -30,27 +32,27 @@ DECLARE_MULTICAST_DELEGATE(FOnEditorViewportInteractionsToggleDelegate)
 /**
  * Returns true if ITF-based input tools should be used.
  */
-EDITORINTERACTIVETOOLSFRAMEWORK_API bool UseEditorViewportInteractions();
+UE_API bool UseEditorViewportInteractions();
 
 /**
  * Helps hiding unwanted logs. For testing.
  */
-EDITORINTERACTIVETOOLSFRAMEWORK_API bool IsVerbose();
+UE_API bool IsVerbose();
 
 /**
  * Toggles ITF Viewport Interactions On/Off
  */
-EDITORINTERACTIVETOOLSFRAMEWORK_API void ToggleEditorViewportInteractions(bool bInEnable);
+UE_API void ToggleEditorViewportInteractions(bool bInEnable);
 
-EDITORINTERACTIVETOOLSFRAMEWORK_API FOnEditorViewportInteractionsToggleDelegate& OnEditorViewportInteractionsActivated();
-EDITORINTERACTIVETOOLSFRAMEWORK_API FOnEditorViewportInteractionsToggleDelegate& OnEditorViewportInteractionsDeactivated();
+UE_API FOnEditorViewportInteractionsToggleDelegate& OnEditorViewportInteractionsActivated();
+UE_API FOnEditorViewportInteractionsToggleDelegate& OnEditorViewportInteractionsDeactivated();
 
 /**
  * Checks whether the specified Command should be triggered by the specified Key
  */
 bool CommandMatchesKey(const TSharedPtr<FUICommandInfo>& InCommandInfo, const FKey& InKeyID);
 
-EDITORINTERACTIVETOOLSFRAMEWORK_API void LOG(const TCHAR* InMessage);
+UE_API void LOG(const TCHAR* InMessage);
 
 } // namespace UE::Editor::ViewportInteractions
 
@@ -99,68 +101,68 @@ public:
 	/**
 	 * Will initialize the required behaviors and the Behavior Set
 	 */
-	EDITORINTERACTIVETOOLSFRAMEWORK_API void Initialize(const UEditorInteractiveToolsContext* InInteractiveToolsContext);
+	UE_API void Initialize(const UEditorInteractiveToolsContext* InInteractiveToolsContext);
 
 	/**
 	 * Register this Input Behavior Source to the InputRouter
 	 */
-	EDITORINTERACTIVETOOLSFRAMEWORK_API void RegisterBehaviorSources();
+	UE_API void RegisterBehaviorSources();
 
 	/**
 	 * Deregister this Input Behavior Source from the InputRouter
 	 */
-	EDITORINTERACTIVETOOLSFRAMEWORK_API void DeregisterBehaviorSources();
+	UE_API void DeregisterBehaviorSources();
 
 	/**
 	 * Call Tick on those Viewport Interactions which require it (e.g. Camera Mover)
 	 */
-	EDITORINTERACTIVETOOLSFRAMEWORK_API void Tick(float InDeltaTime) const;
+	UE_API void Tick(float InDeltaTime) const;
 
 	/**
 	 * Renders the active tool on the specified View/Canvas
 	 */
-	EDITORINTERACTIVETOOLSFRAMEWORK_API void RenderTools(FViewportClient* InViewportClient, const FSceneView* InSceneView, FCanvas* InCanvas) const;
+	UE_API void RenderTools(FViewportClient* InViewportClient, const FSceneView* InSceneView, FCanvas* InCanvas) const;
 
 	/**
 	 * Add the specified interactions to this Behavior Source
 	 */
-	EDITORINTERACTIVETOOLSFRAMEWORK_API void AddInteractions(const TArray<const UClass*>& InInteractions, bool bInReregister = false);
+	UE_API void AddInteractions(const TArray<const UClass*>& InInteractions, bool bInReregister = false);
 
 	/**
 	 * Add the specified interaction to this Behavior Source
 	 */
-	EDITORINTERACTIVETOOLSFRAMEWORK_API UViewportInteraction* AddInteraction(const UClass* InInteractionClass, bool bInReregister = false);
+	UE_API UViewportInteraction* AddInteraction(const UClass* InInteractionClass, bool bInReregister = false);
 
 	/**
 	 * Finds an interaction with a given class if one exists
 	 */
-	EDITORINTERACTIVETOOLSFRAMEWORK_API UViewportInteraction* FindInteraction(const UClass* InInteractionClass) const;
+	UE_API UViewportInteraction* FindInteraction(const UClass* InInteractionClass) const;
 
 	/**
 	 * Finds an interaction with a given name if one exists
 	 */
-	EDITORINTERACTIVETOOLSFRAMEWORK_API UViewportInteraction* FindInteraction(const FName& InName) const;
+	UE_API UViewportInteraction* FindInteraction(const FName& InName) const;
 
 	/**
 	 * Returns true if any interaction in the given group is enabled. An enabled interaction can potentially take action.
 	 */
-	EDITORINTERACTIVETOOLSFRAMEWORK_API bool IsAnyInteractionInGroupEnabled(const FName& InName) const;
+	UE_API bool IsAnyInteractionInGroupEnabled(const FName& InName) const;
 	
 	/**
 	 * Returns true if any interaction in the given group is active. An active interaction is currently taking action.
 	 */
-	EDITORINTERACTIVETOOLSFRAMEWORK_API bool IsAnyInteractionInGroupActive(const FName& InName) const;
+	UE_API bool IsAnyInteractionInGroupActive(const FName& InName) const;
 
 	/**
 	 * Removes the specified interaction from the active ones.
 	 */
-	EDITORINTERACTIVETOOLSFRAMEWORK_API void RemoveInteraction(const FName InInteractionName, bool bInReregister = false);
+	UE_API void RemoveInteraction(const FName InInteractionName, bool bInReregister = false);
 
 	/**
 	 * De-registers behaviors, removes viewport interactions.
 	 * Behaviors added by Initialize will be kept around.
 	 */
-	EDITORINTERACTIVETOOLSFRAMEWORK_API void Reset();
+	UE_API void Reset();
 
 	/**
 	 * Returns the first Viewport Interaction of type T
@@ -208,9 +210,9 @@ public:
 	/**
 	 * Sets current mouse cursor override
 	 */
-	EDITORINTERACTIVETOOLSFRAMEWORK_API void SetMouseCursorOverride(EMouseCursor::Type InMouseCursor);
+	UE_API void SetMouseCursorOverride(EMouseCursor::Type InMouseCursor);
 
-	EDITORINTERACTIVETOOLSFRAMEWORK_API void ClearMouseCursorOverride();
+	UE_API void ClearMouseCursorOverride();
 
 	/**
 	 * Use the templated version in case there is some class specific initialization required by the newly created interaction,
@@ -254,7 +256,7 @@ public:
 
 	bool IsMouseLooking() const;
 
-	EDITORINTERACTIVETOOLSFRAMEWORK_API void SetIsMouseLooking(bool bInIsLooking);
+	UE_API void SetIsMouseLooking(bool bInIsLooking);
 
 	bool HasCameraMoved() const
 	{
@@ -277,9 +279,9 @@ public:
 
 	FEditorModeTools* GetModeTools() const;
 
-	FEditorViewportClient* GetEditorViewportClient() const;
+	UE_API FEditorViewportClient* GetEditorViewportClient() const;
 
-	FEditorViewportClient* GetHoveredEditorViewportClient() const;
+	UE_API FEditorViewportClient* GetHoveredEditorViewportClient() const;
 
 	FOnMouseLookingChanged& OnMouseLookingStateChanged()
 	{
@@ -301,7 +303,7 @@ public:
 		return OnShiftKeyStateChangedDelegate;
 	}
 
-	EDITORINTERACTIVETOOLSFRAMEWORK_API void SetUnsupportedViewportInteractions(const TArray<FName>& InUnsupportedInteractions);
+	UE_API void SetUnsupportedViewportInteractions(const TArray<FName>& InUnsupportedInteractions);
 
 	const TArray<FName>& GetUnsupportedViewportInteractions() const
 	{
@@ -375,3 +377,5 @@ private:
 
 	TArray<FName> UnsupportedInteractions;
 };
+
+#undef UE_API
diff --git a/Engine/Source/Editor/Experimental/EditorInteractiveToolsFramework/Public/ViewportInteractions/ViewportLegacyClickInteraction.h b/Engine/Source/Editor/Experimental/EditorInteractiveToolsFramework/Public/ViewportInteractions/ViewportLegacyClickInteraction.h
--- a/Engine/Source/Editor/Experimental/EditorInteractiveToolsFramework/Public/ViewportInteractions/ViewportLegacyClickInteraction.h
+++ b/Engine/Source/Editor/Experimental/EditorInteractiveToolsFramework/Public/ViewportInteractions/ViewportLegacyClickInteraction.h
@@ -13,12 +13,19 @@
  * An interaction that passes through ITF clicks through to the legacy viewport click handlers.
  */
 UCLASS(MinimalAPI)
-class UViewportLegacyClickInteraction : public UViewportClickInteraction
+class UViewportLegacyClickInteraction : public UViewportClickInteraction, public IHoverBehaviorTarget
 {
 	GENERATED_BODY()
 	
 public:
-	UViewportLegacyClickInteraction();	
+	UViewportLegacyClickInteraction();
+	
+	UE_API virtual void OnClickDown(const FInputDeviceRay& InClickPos) override;
+
+	UE_API virtual FInputRayHit BeginHoverSequenceHitTest(const FInputDeviceRay& PressPos) override;
+	UE_API virtual void OnBeginHover(const FInputDeviceRay& DevicePos) override;
+	UE_API virtual bool OnUpdateHover(const FInputDeviceRay& DevicePos) override;
+	UE_API virtual void OnEndHover() override;
 
 protected:
 
diff --git a/Engine/Source/Editor/LevelEditor/Private/ViewportInteractions/LevelViewportClickSelection.cpp b/Engine/Source/Editor/LevelEditor/Private/ViewportInteractions/LevelViewportClickSelection.cpp
--- a/Engine/Source/Editor/LevelEditor/Private/ViewportInteractions/LevelViewportClickSelection.cpp
+++ b/Engine/Source/Editor/LevelEditor/Private/ViewportInteractions/LevelViewportClickSelection.cpp
@@ -114,6 +114,49 @@ void ULevelViewportClickSelection::BuildBehaviors()
 			FButtonBinding(EKeys::LeftShift).Required(false)
 		});
 	}
+	
+	UMouseHoverBehavior* HoverBehavior = NewObject<UMouseHoverBehavior>();
+	HoverBehavior->Initialize(this);
+	HoverBehavior->SetDefaultPriority(CLICK_PRIORITY);
+	RegisterInputBehavior(HoverBehavior);
+}
+
+void ULevelViewportClickSelection::OnClickDown(const FInputDeviceRay& InClickPos)
+{
+	if (!OnUpdateHover(InClickPos))
+	{
+		SetMouseCursorOverride(EMouseCursor::Default);
+	}
+}
+
+FInputRayHit ULevelViewportClickSelection::BeginHoverSequenceHitTest(const FInputDeviceRay& PressPos)
+{
+	if (GetHitProxy(PressPos))
+	{
+		return FInputRayHit(0.0);
+	}
+	return FInputRayHit();
+}
+
+void ULevelViewportClickSelection::OnBeginHover(const FInputDeviceRay& DevicePos)
+{
+	OnUpdateHover(DevicePos);
+}
+
+bool ULevelViewportClickSelection::OnUpdateHover(const FInputDeviceRay& DevicePos)
+{
+	if (HHitProxy* HitProxy = GetHitProxy(DevicePos))
+	{
+		SetMouseCursorOverride(HitProxy->GetMouseCursor());
+		return true;
+	}
+
+	return false;
+}
+
+void ULevelViewportClickSelection::OnEndHover()
+{
+	ClearMouseCursorOverride();
 }
 
 void ULevelViewportClickSelection::ProcessClick_Internal(const FViewportClick& InViewportClick, HHitProxy* InHitProxy, FSceneView& View)
diff --git a/Engine/Source/Editor/LevelEditor/Private/ViewportInteractions/LevelViewportClickSelection.h b/Engine/Source/Editor/LevelEditor/Private/ViewportInteractions/LevelViewportClickSelection.h
--- a/Engine/Source/Editor/LevelEditor/Private/ViewportInteractions/LevelViewportClickSelection.h
+++ b/Engine/Source/Editor/LevelEditor/Private/ViewportInteractions/LevelViewportClickSelection.h
@@ -3,6 +3,7 @@
 #pragma once
 
 #include "EditorViewportClient.h"
+#include "BaseBehaviors/MouseHoverBehavior.h"
 #include "ViewportInteractions/ViewportClickInteraction.h"
 #include "LevelViewportClickSelection.generated.h"
 
@@ -12,13 +13,20 @@
  * Left-mouse click in level editor viewport, for elements selection
  */
 UCLASS(MinimalAPI)
-class ULevelViewportClickSelection : public UViewportClickInteraction
+class ULevelViewportClickSelection : public UViewportClickInteraction, public IHoverBehaviorTarget
 {
 	GENERATED_BODY()
 
 public:
 	UE_API ULevelViewportClickSelection();
 
+	UE_API virtual void OnClickDown(const FInputDeviceRay& InClickPos) override;
+
+	UE_API virtual FInputRayHit BeginHoverSequenceHitTest(const FInputDeviceRay& PressPos) override;
+	UE_API virtual void OnBeginHover(const FInputDeviceRay& DevicePos) override;
+	UE_API virtual bool OnUpdateHover(const FInputDeviceRay& DevicePos) override;
+	UE_API virtual void OnEndHover() override;
+
 protected:
 	UE_API virtual void BuildBehaviors() override;
 	UE_API virtual void ProcessClick_Internal(const FViewportClick& InViewportClick, HHitProxy* InHitProxy, FSceneView& View) override;
