Make drag termination contingent on meaningful state change in IndirectClickDragBehavior->OnStateUpdatedFunc

#Synced-CL 52103551

diff --git a/Engine/Source/Editor/Experimental/EditorInteractiveToolsFramework/Private/EditorGizmos/EditorTRSGizmo.cpp b/Engine/Source/Editor/Experimental/EditorInteractiveToolsFramework/Private/EditorGizmos/EditorTRSGizmo.cpp
--- a/Engine/Source/Editor/Experimental/EditorInteractiveToolsFramework/Private/EditorGizmos/EditorTRSGizmo.cpp
+++ b/Engine/Source/Editor/Experimental/EditorInteractiveToolsFramework/Private/EditorGizmos/EditorTRSGizmo.cpp
@@ -1655,12 +1655,6 @@ void UEditorTRSGizmo::SetupIndirectBehaviors()
 
 	IndirectClickDragBehavior->OnStateUpdatedFunc = [this](const FInputDeviceState& Input)
 	{
-		if (bIndirectManipulation)
-		{
-			// Since state changed, make sure to properly end ongoing interaction (e.g. transactions)
-			EndTransformEditSequence();
-		}
-
 		struct FSimpleButtonState
 		{
 			bool bDown = false;
@@ -1772,6 +1766,9 @@ void UEditorTRSGizmo::SetupIndirectBehaviors()
 		{
 			if (bInInteraction)
 			{
+				// Since state changed, make sure to properly end ongoing interaction (e.g. transactions)
+				EndTransformEditSequence();
+
 				if (!IndirectClickDragBehavior->IsDragging())
 				{
 					OnClickRelease(Input, true);
