[Editor TRS] Gizmo preset tooltips - replace confirmation dialog with rich comparison tooltips

#Synced-CL 52258331

diff --git a/Engine/Source/Editor/Experimental/EditorInteractiveToolsFramework/Private/Customizations/TransformGizmoEditorSettingsCustomization.cpp b/Engine/Source/Editor/Experimental/EditorInteractiveToolsFramework/Private/Customizations/TransformGizmoEditorSettingsCustomization.cpp
--- a/Engine/Source/Editor/Experimental/EditorInteractiveToolsFramework/Private/Customizations/TransformGizmoEditorSettingsCustomization.cpp
+++ b/Engine/Source/Editor/Experimental/EditorInteractiveToolsFramework/Private/Customizations/TransformGizmoEditorSettingsCustomization.cpp
@@ -5,7 +5,9 @@
 #include "DetailCategoryBuilder.h"
 #include "DetailLayoutBuilder.h"
 #include "DetailWidgetRow.h"
-#include "Dialog/SCustomDialog.h"
+#include "Widgets/SToolTip.h"
+#include "Widgets/Views/SListView.h"
+#include "Widgets/Views/STableRow.h"
 #include "EditorInteractiveGizmoManager.h"
 #include "EditorTRSCustomizationUtils.h"
 #include "EditorTRSGizmoPresets.h"
@@ -18,6 +20,7 @@
 #include "Settings/EditorProjectSettings.h"
 #include "Settings/EditorStyleSettings.h"
 #include "Styling/CoreStyle.h"
+#include "Styling/StyleColors.h"
 #include "TransformGizmoEditorSettings.h"
 #include "UnrealWidgetFwd.h"
 #include "Widgets/Images/SImage.h"
@@ -32,6 +35,73 @@
 
 namespace TransformGizmoEditorSettingsCustomizationLocals
 {
+	static constexpr FLazyName ColumnId_Setting("Setting");
+	static constexpr FLazyName ColumnId_Current("Current");
+	static constexpr FLazyName ColumnId_Preset("Preset");
+
+	using namespace UE::Editor::InteractiveToolsFramework;
+
+	/** Table row for the preset comparison tooltip. */
+	class SPresetComparisonRow : public SMultiColumnTableRow<TSharedPtr<FPresetSettingComparison>>
+	{
+	public:
+		SLATE_BEGIN_ARGS(SPresetComparisonRow) {}
+			SLATE_ARGUMENT(TSharedPtr<FPresetSettingComparison>, Item)
+		SLATE_END_ARGS()
+
+		void Construct(const FArguments& InArgs, const TSharedRef<STableViewBase>& InOwnerTable)
+		{
+			Item = InArgs._Item;
+			SMultiColumnTableRow<TSharedPtr<FPresetSettingComparison>>::Construct(
+				FSuperRowType::FArguments()
+				.Style(&FAppStyle::GetWidgetStyle<FTableRowStyle>("TableView.AlternatingRow")), 
+				InOwnerTable);
+		}
+
+		virtual TSharedRef<SWidget> GenerateWidgetForColumn(const FName& InColumnName) override
+		{
+			if (!Item.IsValid())
+			{
+				return SNullWidget::NullWidget;
+			}
+
+			const FSlateFontInfo RowFont = FCoreStyle::GetDefaultFontStyle("Regular", 9);
+
+			if (InColumnName == ColumnId_Setting)
+			{
+				return SNew(STextBlock)
+					.Text(Item->SettingName)
+					.Font(RowFont)
+					.IsEnabled(Item->bCurrentValueDiffersToPresetValue)
+					.Margin(FMargin(4.0f, 2.0f));
+			}
+
+			if (Item->bIsBool)
+			{
+				const bool bValue = (InColumnName == ColumnId_Current) ? Item->bCurrentValue : Item->bPresetValue;
+				return SNew(SBox)
+					.HAlign(HAlign_Center)
+					.VAlign(VAlign_Center)
+					.Padding(FMargin(4.0f, 1.0f))
+					[
+						SNew(SCheckBox)
+						.IsChecked(bValue ? ECheckBoxState::Checked : ECheckBoxState::Unchecked)
+						.IsEnabled(Item->bCurrentValueDiffersToPresetValue)
+					];
+			}
+
+			const FText& ValueText = (InColumnName == ColumnId_Current) ? Item->CurrentValue : Item->PresetValue;
+			return SNew(STextBlock)
+				.Text(ValueText)
+				.Font(RowFont)
+				.IsEnabled(Item->bCurrentValueDiffersToPresetValue)
+				.Margin(FMargin(4.0f, 2.0f));
+		}
+
+	private:
+		TSharedPtr<FPresetSettingComparison> Item;
+	};
+
 	/** Recursively searches for a nested property handle given an initial parent handle and a list of property names representing the path. */
 	TSharedPtr<IPropertyHandle> GetNestedPropertyHandle(const TSharedPtr<IPropertyHandle>& InParentHandle, const std::initializer_list<FName> InPropertyPath)
 	{
@@ -626,43 +696,76 @@ void FTransformGizmoEditorSettingsCustomization::SaveConfigs()
 	GetMutableDefault<UTransformGizmoEditorSettings>()->SaveConfig();
 }
 
-void FTransformGizmoEditorSettingsCustomization::RequestApplyPreset(FName InPresetName)
+TSharedRef<IToolTip> FTransformGizmoEditorSettingsCustomization::GeneratePresetTooltip(const FTransformGizmoPreset& InPreset)
 {
-	const FTransformGizmoPreset* Preset = FTransformGizmoPreset::FindPresetByName(InPresetName);
-	if (!Preset)
+	using namespace UE::Editor::InteractiveToolsFramework;
+	using namespace TransformGizmoEditorSettingsCustomizationLocals;
+
+	constexpr FTransformGizmoPresetAccessor PresetAccessor;
+	TArray<FPresetSettingComparison> RawComparisons = PresetAccessor.BuildComparisonWithPreset(InPreset);
+
+	TSharedRef<TArray<TSharedPtr<FPresetSettingComparison>>> Items = MakeShared<TArray<TSharedPtr<FPresetSettingComparison>>>();
+	Items->Reserve(RawComparisons.Num());
+	for (FPresetSettingComparison& Comparison : RawComparisons)
 	{
-		return;
+		Items->Add(MakeShared<FPresetSettingComparison>(MoveTemp(Comparison)));
 	}
 
-	const FText Title = FText::Format(
-		LOCTEXT("ApplyPresetDialogTitle", "Apply \"{0}\" Preset"),
-		Preset->DisplayName);
-	const FText Description = Preset->BuildDescription();
+	const FSlateFontInfo HeaderFont = FCoreStyle::GetDefaultFontStyle("Bold", 9);
+
+	const TSharedRef<SHeaderRow> HeaderRow = SNew(SHeaderRow)
+		+ SHeaderRow::Column(ColumnId_Setting)
+		.DefaultLabel(LOCTEXT("TooltipHeaderSetting", "Setting"))
+		.FillWidth(0.5f)
+		.HAlignHeader(HAlign_Left)
+		.HAlignCell(HAlign_Left)
+		+ SHeaderRow::Column(ColumnId_Current)
+		.DefaultLabel(LOCTEXT("TooltipHeaderCurrent", "Current"))
+		.FillWidth(0.25f)
+		.HAlignHeader(HAlign_Center)
+		.HAlignCell(HAlign_Center)
+		+ SHeaderRow::Column(ColumnId_Preset)
+		.DefaultLabel(LOCTEXT("TooltipHeaderPreset", "Preset"))
+		.FillWidth(0.25f)
+		.HAlignHeader(HAlign_Center)
+		.HAlignCell(HAlign_Center);
+
+	const TSharedRef<SListView<TSharedPtr<FPresetSettingComparison>>> ListView =
+		SNew(SListView<TSharedPtr<FPresetSettingComparison>>)
+		.ListViewStyle(&FAppStyle::Get().GetWidgetStyle<FTableViewStyle>("ListView"))
+		.ListItemsSource(&Items.Get())
+		.HeaderRow(HeaderRow)
+		.SelectionMode(ESelectionMode::None)
+		.OnGenerateRow_Lambda([Items](TSharedPtr<FPresetSettingComparison> InItem, const TSharedRef<STableViewBase>& InOwnerTable) -> TSharedRef<ITableRow>
+		{
+			return SNew(SPresetComparisonRow, InOwnerTable)
+				.Item(InItem);
+		});
 
-	TSharedRef<SCustomDialog> Dialog = SNew(SCustomDialog)
-		.Title(Title)
-		.Content()
+	return SNew(SToolTip)
+	[
+		SNew(SVerticalBox)
+		+ SVerticalBox::Slot()
+		.AutoHeight()
+		.Padding(4.0f)
 		[
-			SNew(SVerticalBox)
-			+ SVerticalBox::Slot()
-			.AutoHeight()
-			.Padding(16.0f)
+			SNew(STextBlock)
+			.Text(InPreset.ToolTipText)
+			.Font(FCoreStyle::GetDefaultFontStyle("Italic", 9)) 
+			.ColorAndOpacity(FStyleColors::ForegroundHeader)
+		]
+		+ SVerticalBox::Slot()
+		.AutoHeight()
+		.MaxHeight(300.0f)
+		.Padding(0.0f, 4.0f, 0.0f, 0.0f)
+		[
+			SNew(SBox)
+			.MinDesiredWidth(350.0f)
 			[
-				SNew(STextBlock)
-				.Text(Description)
-				.AutoWrapText(true)
+				ListView
 			]
 		]
-		.Buttons({
-			SCustomDialog::FButton(LOCTEXT("ApplyPresetOK", "Apply"), FSimpleDelegate(), SCustomDialog::EButtonRole::Confirm),
-			SCustomDialog::FButton(LOCTEXT("ApplyPresetCancel", "Cancel"), FSimpleDelegate(), SCustomDialog::EButtonRole::Cancel)
-		});
-
-	const int32 ButtonPressed = Dialog->ShowModal();
-	if (ButtonPressed == 0) // Apply
-	{
-		ApplyPreset(*Preset);
-	}
+	];
 }
 
 TSharedRef<SWidget> FTransformGizmoEditorSettingsCustomization::GeneratePresetComboBoxContent()
@@ -677,12 +780,18 @@ TSharedRef<SWidget> FTransformGizmoEditorSettingsCustomization::GeneratePresetCo
 			{
 				for (const FTransformGizmoPreset& Preset : FTransformGizmoPreset::GetBuiltInPresets())
 				{
-					FName PresetName = Preset.Name;
+					const TSharedRef<SWidget> EntryWidget =
+						SNew(STextBlock)
+						.Text(Preset.DisplayName)
+						.Font(IDetailLayoutBuilder::GetDetailFont())
+						.ToolTip(StrongThis->GeneratePresetTooltip(Preset));
+
 					MenuBuilder.AddMenuEntry(
-						Preset.DisplayName,
-						FText::GetEmpty(),
-						FSlateIcon(),
-						FUIAction(FExecuteAction::CreateSP(StrongThis.ToSharedRef(), &FTransformGizmoEditorSettingsCustomization::RequestApplyPreset, PresetName)));
+						FUIAction(FExecuteAction::CreateLambda([PresetCopy = Preset]()
+						{
+							FTransformGizmoEditorSettingsCustomization::ApplyPreset(PresetCopy);
+						})),
+						EntryWidget);
 				}
 			}
 
diff --git a/Engine/Source/Editor/Experimental/EditorInteractiveToolsFramework/Private/Customizations/TransformGizmoEditorSettingsCustomization.h b/Engine/Source/Editor/Experimental/EditorInteractiveToolsFramework/Private/Customizations/TransformGizmoEditorSettingsCustomization.h
--- a/Engine/Source/Editor/Experimental/EditorInteractiveToolsFramework/Private/Customizations/TransformGizmoEditorSettingsCustomization.h
+++ b/Engine/Source/Editor/Experimental/EditorInteractiveToolsFramework/Private/Customizations/TransformGizmoEditorSettingsCustomization.h
@@ -7,6 +7,7 @@
 #include "IDetailCustomization.h"
 #include "Misc/Attribute.h"
 
+class IToolTip;
 class IPropertyHandle;
 
 class FTransformGizmoEditorSettingsCustomization
@@ -30,8 +31,8 @@ private:
 	/** Returns the current display name for the active preset (or "Custom"). */
 	FText GetActivePresetDisplayName() const;
 
-	/** Applies a preset by name after user confirmation. */
-	void RequestApplyPreset(FName InPresetName);
+	/** Generates a rich tooltip widget showing a comparison table between current settings and a preset. */
+	TSharedRef<IToolTip> GeneratePresetTooltip(const FTransformGizmoPreset& InPreset);
 
 	/** Apply the given preset values to the settings objects. */
 	static void ApplyPreset(const FTransformGizmoPreset& InPreset);
diff --git a/Engine/Source/Editor/Experimental/EditorInteractiveToolsFramework/Private/EditorTRSGizmoPresets.cpp b/Engine/Source/Editor/Experimental/EditorInteractiveToolsFramework/Private/EditorTRSGizmoPresets.cpp
--- a/Engine/Source/Editor/Experimental/EditorInteractiveToolsFramework/Private/EditorTRSGizmoPresets.cpp
+++ b/Engine/Source/Editor/Experimental/EditorInteractiveToolsFramework/Private/EditorTRSGizmoPresets.cpp
@@ -77,35 +77,6 @@ const FTransformGizmoPreset* FTransformGizmoPreset::FindPresetByName(FName InNam
 	return Presets.FindByPredicate([InName](const FTransformGizmoPreset& P) { return P.Name == InName; });
 }
 
-FText FTransformGizmoPreset::BuildDescription() const
-{
-	const FText CheckedText = FText::FromString(TEXT("✔"));
-	const FText UncheckedText = FText::FromString(TEXT("✘"));
-
-	return FText::Format(
-		LOCTEXT("PresetDescriptionFormat",
-			"Rotate Mode: {0}\n"
-			"Additive Indirect Manipulation: {1}\n"
-			"Sequential Indirect Axes Buttons: {2}\n"
-			"Indirect Scale is always Uniform: {3}\n"
-			"Persist Handle Selection: {4}\n"
-			"Include Explicit Rotation: {5}\n"
-			"Drag Duplicate on Rotation: {6}\n"
-			"Scale Type: {7}"),
-		RotateMode == EAxisRotateMode::Pull 
-		? LOCTEXT("RotateModePull", "Pull") 
-		: RotateMode == EAxisRotateMode::Arc 
-		? LOCTEXT("RotateModeArc", "Arc")
-		: LOCTEXT("RotateModeScreenArc", "Screen Arc"),
-		bAdditiveIndirectAxes ? CheckedText : UncheckedText,
-		bSequentialIndirectAxesButtons ? CheckedText : UncheckedText,
-		bUniformIndirectScale ? CheckedText : UncheckedText,
-		bPersistHandleSelection ? CheckedText : UncheckedText,
-		bIncludeExplicitRotation ? CheckedText : UncheckedText,
-		bDragDuplicateOnRotation ? CheckedText : UncheckedText,
-		ScaleType == EGizmoTransformScaleType::PercentageBased ? LOCTEXT("ScaleTypePercentageBased", "Percentage Based") : LOCTEXT("ScaleTypeDefault", "Default"));
-}
-
 namespace UE::Editor::InteractiveToolsFramework
 {
 	bool FTransformGizmoPresetAccessor::CurrentSettingsMatchPreset(const FTransformGizmoPreset& InPreset) const
@@ -120,6 +91,84 @@ namespace UE::Editor::InteractiveToolsFramework
 			&& GetScaleType() == InPreset.ScaleType;
 	}
 
+	TArray<FPresetSettingComparison> FTransformGizmoPresetAccessor::BuildComparisonWithPreset(const FTransformGizmoPreset& InPreset) const
+	{
+		TArray<FPresetSettingComparison> Comparisons;
+		Comparisons.Reserve(8);
+
+		auto AddBoolComparison = [&](const FText& InSettingName, const bool bInCurrentValue, const bool bInPresetValue)
+		{
+			FPresetSettingComparison Comparison;
+			Comparison.SettingName = InSettingName;
+			Comparison.bCurrentValueDiffersToPresetValue = bInCurrentValue != bInPresetValue;
+			Comparison.bIsBool = true;
+			Comparison.bCurrentValue = bInCurrentValue;
+			Comparison.bPresetValue = bInPresetValue;
+			Comparisons.Add(MoveTemp(Comparison));
+		};
+
+		AddBoolComparison(LOCTEXT("CompareRotateDuplicate", "Drag Duplicate on Rotation"), GetDragDuplicateOnRotation(), InPreset.bDragDuplicateOnRotation);
+		AddBoolComparison(LOCTEXT("CompareExplicitRotation", "Include Explicit Rotation"), GetIncludeExplicitRotation(), InPreset.bIncludeExplicitRotation);
+		AddBoolComparison(LOCTEXT("CompareUniformScale", "Uniform Indirect Scale"), GetUniformIndirectScale(), InPreset.bUniformIndirectScale);
+		AddBoolComparison(LOCTEXT("ComparePersistHandle", "Persist Handle Selection"), GetPersistHandleSelection(), InPreset.bPersistHandleSelection);
+		AddBoolComparison(LOCTEXT("CompareSequentialAxes", "Sequential Indirect Axes"), GetSequentialIndirectAxesButtons(), InPreset.bSequentialIndirectAxesButtons);
+		AddBoolComparison(LOCTEXT("CompareAdditiveAxes", "Additive Indirect Axes"), GetAdditiveIndirectAxes(), InPreset.bAdditiveIndirectAxes);
+
+		{
+			auto RotateModeToText = [](const EAxisRotateMode::Type InMode) -> FText
+			{
+				switch (InMode)
+				{
+				case EAxisRotateMode::Pull: 
+					return LOCTEXT("CompareRotateModePull", "Pull");
+
+				case EAxisRotateMode::Arc: 
+					return LOCTEXT("CompareRotateModeArc", "Arc");
+
+				case EAxisRotateMode::ScreenArc:
+				default:
+					return LOCTEXT("CompareRotateModeScreenArc", "Screen Arc");
+				}
+			};
+			
+			const EAxisRotateMode::Type CurrentRotateMode = GetRotateMode();
+			Comparisons.Add({
+				LOCTEXT("CompareRotateMode", "Rotate Mode"),
+				RotateModeToText(CurrentRotateMode),
+				RotateModeToText(InPreset.RotateMode),
+				CurrentRotateMode != InPreset.RotateMode
+			});
+		}
+
+		{
+			auto ScaleTypeToText = [](const EGizmoTransformScaleType InType) -> FText
+			{
+				switch (InType)
+				{
+				case EGizmoTransformScaleType::Default: 
+					return LOCTEXT("CompareScaleTypeDefault", "Default");
+
+				case EGizmoTransformScaleType::Legacy: 
+					return LOCTEXT("CompareScaleTypeLegacy", "Legacy");
+
+				case EGizmoTransformScaleType::PercentageBased:
+				default:
+					return LOCTEXT("CompareScaleTypePercentage", "Percentage Based");
+				}
+			};
+
+			const EGizmoTransformScaleType CurrentScaleType = GetScaleType();
+			Comparisons.Add({
+				LOCTEXT("CompareScaleType", "Scale Type"),
+				ScaleTypeToText(CurrentScaleType),
+				ScaleTypeToText(InPreset.ScaleType),
+				CurrentScaleType != InPreset.ScaleType
+			});
+		}
+
+		return Comparisons;
+	}
+
 	void FTransformGizmoPresetAccessor::ApplyPreset(const FTransformGizmoPreset& InPreset) const
 	{
 		FScopedTransaction Transaction(LOCTEXT("ApplyGizmoPreset", "Apply Transform Gizmo Preset"));
@@ -235,13 +284,13 @@ namespace UE::Editor::InteractiveToolsFramework
 	EGizmoTransformScaleType FTransformGizmoPresetAccessor::GetScaleType() const
 	{
 		const UTransformGizmoEditorSettings* Settings = GetDefault<UTransformGizmoEditorSettings>();
-		return Settings->GizmosParameters.Interaction.ScaleInteraction.ScaleType.Get(EGizmoTransformScaleType::Default);
+		return Settings->GizmosParameters.ScaleType;
 	}
 
 	void FTransformGizmoPresetAccessor::SetScaleType(EGizmoTransformScaleType InScaleType) const
 	{
 		UTransformGizmoEditorSettings* Settings = GetMutableDefault<UTransformGizmoEditorSettings>();
-		Settings->GizmosParameters.Interaction.ScaleInteraction.ScaleType = InScaleType;
+		Settings->GizmosParameters.ScaleType = InScaleType;
 
 		BroadcastPropertyChanged<UTransformGizmoEditorSettings>(Settings, "GizmosParameters");
 	}
diff --git a/Engine/Source/Editor/Experimental/EditorInteractiveToolsFramework/Private/EditorTRSGizmoPresets.h b/Engine/Source/Editor/Experimental/EditorInteractiveToolsFramework/Private/EditorTRSGizmoPresets.h
--- a/Engine/Source/Editor/Experimental/EditorInteractiveToolsFramework/Private/EditorTRSGizmoPresets.h
+++ b/Engine/Source/Editor/Experimental/EditorInteractiveToolsFramework/Private/EditorTRSGizmoPresets.h
@@ -78,13 +78,37 @@ struct FTransformGizmoPreset
 
 	/** Finds a preset by name. Returns nullptr if not found. */
 	 static const FTransformGizmoPreset* FindPresetByName(FName InName);
-
-	/** Builds a text description of this preset's settings. */
-	FText BuildDescription() const;
 };
 
 namespace UE::Editor::InteractiveToolsFramework
 {
+	/**
+	 * A single row of comparison data between current settings and a preset.
+	 */
+	struct FPresetSettingComparison
+	{
+		/** The display name of the setting (e.g. "Rotate Mode"). */
+		FText SettingName;
+
+		/** The current value of this setting, formatted for display (used for non-bool settings). */
+		FText CurrentValue;
+
+		/** The value this setting will be changed to if the preset is applied (used for non-bool settings). */
+		FText PresetValue;
+
+		/** True if the current value differs from the preset value. */
+		bool bCurrentValueDiffersToPresetValue = false;
+
+		/** True if this setting is a boolean and should be rendered as checkboxes. */
+		bool bIsBool = false;
+
+		/** The current boolean state (only valid when bIsBool is true). */
+		bool bCurrentValue = false;
+
+		/** The preset boolean state (only valid when bIsBool is true). */
+		bool bPresetValue = false;
+	};
+
 	/** Provides getters/setters to get current settings, and apply them. */
 	struct FTransformGizmoPresetAccessor
 	{
@@ -94,6 +118,9 @@ namespace UE::Editor::InteractiveToolsFramework
 		/** Applies the given preset to the current settings. */
 		void ApplyPreset(const FTransformGizmoPreset& InPreset) const;
 
+		/** Builds a per-setting comparison between the current settings and the given preset. */
+		TArray<FPresetSettingComparison> BuildComparisonWithPreset(const FTransformGizmoPreset& InPreset) const;
+
 		/** Returns the current value of bDragDuplicateOnRotation. */
 		bool GetDragDuplicateOnRotation() const;
 
