[ITF Viewport Interactions] Fix for Drag tools not drawing regression

#Synced-CL 52295285

diff --git a/Engine/Source/Editor/Experimental/EditorInteractiveToolsFramework/Private/EditorDragTools/DragToolInteraction.cpp b/Engine/Source/Editor/Experimental/EditorInteractiveToolsFramework/Private/EditorDragTools/DragToolInteraction.cpp
--- a/Engine/Source/Editor/Experimental/EditorInteractiveToolsFramework/Private/EditorDragTools/DragToolInteraction.cpp
+++ b/Engine/Source/Editor/Experimental/EditorInteractiveToolsFramework/Private/EditorDragTools/DragToolInteraction.cpp
@@ -126,6 +126,11 @@ void UDragToolInteraction::OnStateUpdated(const FInputDeviceState& InInputDevice
 	InputState = InInputDeviceState;
 }
 
+bool UDragToolInteraction::ShouldDraw(const EViewInteractionState InInteractionState)
+{
+	return IsDragging() && EnumHasAllFlags(InInteractionState, EViewInteractionState::Focused);
+}
+
 void UDragToolInteraction::OnDragEnd(const FInputDeviceRay& InReleasePos)
 {
 	if (UViewportInteractionsBehaviorSource* BehaviorSource = GetViewportInteractionsBehaviorSource())
diff --git a/Engine/Source/Editor/Experimental/EditorInteractiveToolsFramework/Private/EditorDragTools/MarqueeSelectInteraction.cpp b/Engine/Source/Editor/Experimental/EditorInteractiveToolsFramework/Private/EditorDragTools/MarqueeSelectInteraction.cpp
--- a/Engine/Source/Editor/Experimental/EditorInteractiveToolsFramework/Private/EditorDragTools/MarqueeSelectInteraction.cpp
+++ b/Engine/Source/Editor/Experimental/EditorInteractiveToolsFramework/Private/EditorDragTools/MarqueeSelectInteraction.cpp
@@ -25,7 +25,7 @@ void UMarqueeSelectInteraction::Draw(FCanvas* InCanvas, IToolsContextRenderAPI*
 		return;
 	}
 
-	if (!IsDragging() || InRenderAPI->GetViewInteractionState() != EViewInteractionState::Focused)
+	if (!ShouldDraw(InRenderAPI->GetViewInteractionState()))
 	{
 		return;
 	}
diff --git a/Engine/Source/Editor/Experimental/EditorInteractiveToolsFramework/Private/EditorDragTools/MeasureToolInteraction.cpp b/Engine/Source/Editor/Experimental/EditorInteractiveToolsFramework/Private/EditorDragTools/MeasureToolInteraction.cpp
--- a/Engine/Source/Editor/Experimental/EditorInteractiveToolsFramework/Private/EditorDragTools/MeasureToolInteraction.cpp
+++ b/Engine/Source/Editor/Experimental/EditorInteractiveToolsFramework/Private/EditorDragTools/MeasureToolInteraction.cpp
@@ -39,12 +39,7 @@ void UMeasureToolInteraction::Draw(FCanvas* InCanvas, IToolsContextRenderAPI* In
 		return;
 	}
 
-	if (!IsDragging())
-	{
-		return;
-	}
-
-	if (InRenderAPI->GetViewInteractionState() != EViewInteractionState::Focused)
+	if (!ShouldDraw(InRenderAPI->GetViewInteractionState()))
 	{
 		return;
 	}
diff --git a/Engine/Source/Editor/Experimental/EditorInteractiveToolsFramework/Private/EditorDragTools/ViewportChangeInteraction.cpp b/Engine/Source/Editor/Experimental/EditorInteractiveToolsFramework/Private/EditorDragTools/ViewportChangeInteraction.cpp
--- a/Engine/Source/Editor/Experimental/EditorInteractiveToolsFramework/Private/EditorDragTools/ViewportChangeInteraction.cpp
+++ b/Engine/Source/Editor/Experimental/EditorInteractiveToolsFramework/Private/EditorDragTools/ViewportChangeInteraction.cpp
@@ -35,13 +35,8 @@ void UViewportChangeInteraction::Draw(FCanvas* InCanvas, IToolsContextRenderAPI*
 	{
 		return;
 	}
-	
-	if (!IsDragging())
-	{
-		return;
-	}
 
-	if (InRenderAPI->GetViewInteractionState() != EViewInteractionState::Focused)
+	if (!ShouldDraw(InRenderAPI->GetViewInteractionState()))
 	{
 		return;
 	}
diff --git a/Engine/Source/Editor/Experimental/EditorInteractiveToolsFramework/Public/EditorDragTools/DragToolInteraction.h b/Engine/Source/Editor/Experimental/EditorInteractiveToolsFramework/Public/EditorDragTools/DragToolInteraction.h
--- a/Engine/Source/Editor/Experimental/EditorInteractiveToolsFramework/Public/EditorDragTools/DragToolInteraction.h
+++ b/Engine/Source/Editor/Experimental/EditorInteractiveToolsFramework/Public/EditorDragTools/DragToolInteraction.h
@@ -90,6 +90,8 @@ public:
 	}
 
 protected:
+	UE_API virtual bool ShouldDraw(const EViewInteractionState InInteractionState);
+
 	/** The start/end location of the current drag. */
 	FVector Start, End;
 
