[Viewport Interactions] Force terminate viewport interaction behaviors before deregistering them.

#rnx
#Synced-CL 52381418

diff --git a/Engine/Source/Editor/Experimental/EditorInteractiveToolsFramework/Private/ViewportInteractions/ViewportInteractionsBehaviorSource.cpp b/Engine/Source/Editor/Experimental/EditorInteractiveToolsFramework/Private/ViewportInteractions/ViewportInteractionsBehaviorSource.cpp
--- a/Engine/Source/Editor/Experimental/EditorInteractiveToolsFramework/Private/ViewportInteractions/ViewportInteractionsBehaviorSource.cpp
+++ b/Engine/Source/Editor/Experimental/EditorInteractiveToolsFramework/Private/ViewportInteractions/ViewportInteractionsBehaviorSource.cpp
@@ -292,6 +292,7 @@ void UViewportInteractionsBehaviorSource::DeregisterBehaviorSources()
 	{
 		if (UInputRouter* InputRouter = InteractiveToolsContext->InputRouter)
 		{
+			InputRouter->ForceTerminateSource(this);
 			InputRouter->DeregisterSource(this);
 		}
 
