[Viewport Interactions] Defer resetting the cursor due to cursor overrides to Tick().

#rnx
#Synced-CL 52419101

diff --git a/Engine/Source/Editor/Experimental/EditorInteractiveToolsFramework/Private/ViewportInteractions/ViewportInteractionsBehaviorSource.cpp b/Engine/Source/Editor/Experimental/EditorInteractiveToolsFramework/Private/ViewportInteractions/ViewportInteractionsBehaviorSource.cpp
--- a/Engine/Source/Editor/Experimental/EditorInteractiveToolsFramework/Private/ViewportInteractions/ViewportInteractionsBehaviorSource.cpp
+++ b/Engine/Source/Editor/Experimental/EditorInteractiveToolsFramework/Private/ViewportInteractions/ViewportInteractionsBehaviorSource.cpp
@@ -178,7 +178,7 @@ void UViewportInteractionsBehaviorSource::OnForceEndCapture()
 	SetIsMouseLooking(false);
 }
 
-void UViewportInteractionsBehaviorSource::Initialize(const UEditorInteractiveToolsContext* InInteractiveToolsContext)
+void UViewportInteractionsBehaviorSource::Initialize(UEditorInteractiveToolsContext* InInteractiveToolsContext)
 {
 	if (!InInteractiveToolsContext)
 	{
@@ -300,7 +300,7 @@ void UViewportInteractionsBehaviorSource::DeregisterBehaviorSources()
 	}
 }
 
-void UViewportInteractionsBehaviorSource::Tick(float InDeltaTime) const
+void UViewportInteractionsBehaviorSource::Tick(float InDeltaTime)
 {
 	for (const TPair<FName, TObjectPtr<UViewportInteraction>>& Pair : ViewportInteractions)
 	{
@@ -309,6 +309,20 @@ void UViewportInteractionsBehaviorSource::Tick(float InDeltaTime) const
 			Interaction->Tick(InDeltaTime);
 		}
 	}
+	
+	if (CursorOverride != LastCursorOverride)
+	{
+		if (LastCursorOverride.Get(EMouseCursor::Default) == EMouseCursor::None)
+		{
+			RestoreCursorPosition();
+		}
+		else if (CursorOverride.Get(EMouseCursor::Default) == EMouseCursor::None)
+		{
+			CacheCursorPosition();
+		}
+		
+		LastCursorOverride = CursorOverride;
+	}
 }
 
 void UViewportInteractionsBehaviorSource::RenderTools(IToolsContextRenderAPI* InRenderAPI) const
@@ -508,62 +522,94 @@ void UViewportInteractionsBehaviorSource::Reset()
 
 void UViewportInteractionsBehaviorSource::SetMouseCursorOverride(EMouseCursor::Type InMouseCursor)
 {
-	CursorOverride = InMouseCursor;
-
-	if (FEditorViewportClient* const EditorViewportClient = GetEditorViewportClient())
+	if (!CursorOverride.IsSet() || CursorOverride.GetValue() != InMouseCursor)
 	{
-		if (UToolsContextCursorAPI* ToolsContextCursorAPI =  GetCursorAPI(EditorViewportClient))
-		{
-			ToolsContextCursorAPI->SetCursorOverride(CursorOverride);
-		}
-
-		// In case we hide the cursor, we show it back where it was when clearing the override
-		if (InMouseCursor == EMouseCursor::None)
+		if (UToolsContextCursorAPI* CursorAPI = GetCursorAPI())
 		{
-			EditorViewportClient->Viewport->GetMousePos(HiddenCursorPosition);
+			CursorAPI->SetCursorOverride(InMouseCursor);
+			CursorOverride = InMouseCursor;
 		}
 	}
 }
 
 void UViewportInteractionsBehaviorSource::ClearMouseCursorOverride()
 {
-	const bool bResetMousePos = CursorOverride == EMouseCursor::None;
-	CursorOverride = EMouseCursor::Default;
+	if (UToolsContextCursorAPI* CursorAPI = GetCursorAPI())
+	{
+		CursorAPI->ClearCursorOverride();
+		CursorOverride.Reset();
+	}
+}
 
-	if (FEditorViewportClient* const EditorViewportClient = GetEditorViewportClient())
+UToolsContextCursorAPI* UViewportInteractionsBehaviorSource::GetCursorAPI()
+{
+	if (UToolsContextCursorAPI* CachedValue = WeakCursorAPI.Get())
 	{
-		if (UToolsContextCursorAPI* ToolsContextCursorAPI =  GetCursorAPI(EditorViewportClient))
+		return CachedValue;
+	}
+	
+	if (const UEditorInteractiveToolsContext* Context = EditorInteractiveToolsContextWeak.Get())
+	{
+		if (const UContextObjectStore* const ContextObjectStore = Context->ContextObjectStore)
 		{
-			ToolsContextCursorAPI->ClearCursorOverride();
+			if (UToolsContextCursorAPI* FoundApi = ContextObjectStore->FindContext<UToolsContextCursorAPI>())
+			{
+				WeakCursorAPI = FoundApi;
+				return FoundApi;
+			}
 		}
+	}
+	
+	return nullptr;
+}
+
+void UViewportInteractionsBehaviorSource::CacheCursorPosition()
+{
+	if (const FEditorViewportClient* EditorViewportClient = GetEditorViewportClient())
+	{
+		HiddenCursorPosition = FIntPoint(
+			EditorViewportClient->Viewport->GetMouseX(),
+			EditorViewportClient->Viewport->GetMouseY()
+		);
+	}
+}
 
-		if (bResetMousePos)
+void UViewportInteractionsBehaviorSource::RestoreCursorPosition()
+{
+	if (!HiddenCursorPosition.IsSet())
+	{
+		return;
+	}
+
+	if (FEditorViewportClient* EditorViewportClient = GetEditorViewportClient())
+	{
+		const FIntPoint CursorPosition = HiddenCursorPosition.GetValue();
+		
+		switch (UE::Editor::ViewportInteractions::CVarMouseRestoration.GetValueOnAnyThread())
 		{
-			const int32 Mode = UE::Editor::ViewportInteractions::CVarMouseRestoration.GetValueOnAnyThread();
-			switch (Mode)
-			{
-			case 1: // Reset the mouse back to its original position (parity with legacy)
-				EditorViewportClient->Viewport->SetMouse(HiddenCursorPosition.X, HiddenCursorPosition.Y);
-				break;
-			case 2: // Reset the mouse to original position only when the mouse is no longer within the viewport
+		case 1: // Reset the mouse back to its original position (parity with legacy)
+			EditorViewportClient->Viewport->SetMouse(CursorPosition.X, CursorPosition.Y);
+			break;
+		case 2: // Reset the mouse to original position only when the mouse is no longer within the viewport
 			{
 				FIntPoint CurrentMousePosition;
 				EditorViewportClient->Viewport->GetMousePos(CurrentMousePosition);
 		
-				FIntPoint Size = EditorViewportClient->Viewport->GetSizeXY();
+				const FIntPoint Size = EditorViewportClient->Viewport->GetSizeXY();
 		
 				if (CurrentMousePosition.X < 0 || CurrentMousePosition.X > Size.X
 					|| CurrentMousePosition.Y < 0 || CurrentMousePosition.Y > Size.Y)
 				{
-					EditorViewportClient->Viewport->SetMouse(HiddenCursorPosition.X, HiddenCursorPosition.Y);				
+					EditorViewportClient->Viewport->SetMouse(CursorPosition.X, CursorPosition.Y);				
 				}
 				break;
 			}
-			default:
-				break;
-			}
+		default:
+			break;
 		}
 	}
+	
+	HiddenCursorPosition.Reset();
 }
 
 bool UViewportInteractionsBehaviorSource::IsMouseLooking() const
@@ -587,7 +633,10 @@ void UViewportInteractionsBehaviorSource::SetIsMouseLooking(bool bInIsLooking)
 	{
 		if (bInIsLooking)
 		{
-			SetMouseCursorOverride(CursorOverride);
+			if (CursorOverride.IsSet())
+			{
+				SetMouseCursorOverride(CursorOverride.GetValue());
+			}
 		}
 		else
 		{
diff --git a/Engine/Source/Editor/Experimental/EditorInteractiveToolsFramework/Public/ViewportInteractions/ViewportInteractionsBehaviorSource.h b/Engine/Source/Editor/Experimental/EditorInteractiveToolsFramework/Public/ViewportInteractions/ViewportInteractionsBehaviorSource.h
--- a/Engine/Source/Editor/Experimental/EditorInteractiveToolsFramework/Public/ViewportInteractions/ViewportInteractionsBehaviorSource.h
+++ b/Engine/Source/Editor/Experimental/EditorInteractiveToolsFramework/Public/ViewportInteractions/ViewportInteractionsBehaviorSource.h
@@ -101,7 +101,7 @@ public:
 	/**
 	 * Will initialize the required behaviors and the Behavior Set
 	 */
-	UE_API void Initialize(const UEditorInteractiveToolsContext* InInteractiveToolsContext);
+	UE_API void Initialize(UEditorInteractiveToolsContext* InInteractiveToolsContext);
 
 	/**
 	 * Register this Input Behavior Source to the InputRouter
@@ -116,7 +116,7 @@ public:
 	/**
 	 * Call Tick on those Viewport Interactions which require it (e.g. Camera Mover)
 	 */
-	UE_API void Tick(float InDeltaTime) const;
+	UE_API void Tick(float InDeltaTime);
 
 	/**
 	 * Renders the active tool on the specified View/World
@@ -349,7 +349,12 @@ private:
 	void OnGizmoCreatedDelegate(UTransformGizmo* InTransformGizmo);
 	void OnGizmoMovementBegin(UTransformProxy* InTransformProxy);
 	void OnGizmoMovementEnd(UTransformProxy* InTransformProxy);
-
+	
+	UToolsContextCursorAPI* GetCursorAPI();
+	
+	void CacheCursorPosition();
+	void RestoreCursorPosition();
+	
 	bool bIsShiftDown;
 	bool bIsCtrlDown;
 	bool bIsAltDown;
@@ -358,15 +363,19 @@ private:
 	bool bIsMiddleMouseButtonDown;
 	bool bIsRightMouseButtonDown;
 
-	EMouseCursor::Type CursorOverride = EMouseCursor::Default;
+	TOptional<EMouseCursor::Type> CursorOverride;
+	TOptional<EMouseCursor::Type> LastCursorOverride;
 
 	/**
 	 * Caches cursor position whenever we hide it, so we can place it back where it was as soon as we show it again
 	 */
-	FIntPoint HiddenCursorPosition;
+	TOptional<FIntPoint> HiddenCursorPosition;
 
 	UPROPERTY(Transient)
-	TWeakObjectPtr<const UEditorInteractiveToolsContext> EditorInteractiveToolsContextWeak;
+	TWeakObjectPtr<UEditorInteractiveToolsContext> EditorInteractiveToolsContextWeak;
+	
+	UPROPERTY(Transient)
+	TWeakObjectPtr<UToolsContextCursorAPI> WeakCursorAPI = nullptr;
 
 	FOnMouseLookingChanged OnMouseLookingChangedDelegate;
 
