Restored Debug tab in Gizmos.ShowSettings

#Synced-CL 52437958

diff --git a/Engine/Source/Editor/Experimental/EditorInteractiveToolsFramework/Private/Customizations/GizmoCustomization.cpp b/Engine/Source/Editor/Experimental/EditorInteractiveToolsFramework/Private/Customizations/GizmoCustomization.cpp
--- a/Engine/Source/Editor/Experimental/EditorInteractiveToolsFramework/Private/Customizations/GizmoCustomization.cpp
+++ b/Engine/Source/Editor/Experimental/EditorInteractiveToolsFramework/Private/Customizations/GizmoCustomization.cpp
@@ -153,7 +153,34 @@ void FTransformGizmoEditorSettingsCustomizationBase::HideAllProperties(IDetailLa
 
 void FTransformGizmoEditorSettingsCustomizationBase::IsolateCategory(IDetailLayoutBuilder& InDetailBuilder, const TSharedRef<IPropertyHandle>& InContainerProperty, const FName InCategoryName)
 {
-	InDetailBuilder.HideCategory(InCategoryName);
+	// Only hide the category if children will be re-added to different categories.
+	// If all children share the same category as InCategoryName, HideCategory would
+	// permanently suppress them since they'd be re-added to the hidden category.
+	bool bAllChildrenInSameCategory = true;
+	{
+		uint32 NumChildren;
+		if (InContainerProperty->GetNumChildren(NumChildren) == FPropertyAccess::Success)
+		{
+			for (uint32 ChildIndex = 0; ChildIndex < NumChildren; ++ChildIndex)
+			{
+				TSharedPtr<IPropertyHandle> ChildProperty = InContainerProperty->GetChildHandle(ChildIndex);
+				if (ChildProperty.IsValid())
+				{
+					const FName ChildCategory = ChildProperty->GetDefaultCategoryName();
+					if (!ChildCategory.IsEqual(InCategoryName))
+					{
+						bAllChildrenInSameCategory = false;
+						break;
+					}
+				}
+			}
+		}
+	}
+
+	if (!bAllChildrenInSameCategory)
+	{
+		InDetailBuilder.HideCategory(InCategoryName);
+	}
 
 	const auto ResetToDefault = FResetToDefaultOverride::Create(
 		FIsResetToDefaultVisible::CreateSP(this, &FTransformGizmoEditorSettingsCustomizationBase::IsResetToDefaultVisible),
