[ITF Gizmos] Temporary Surface and TRS Snapping

Adding a way to temporarily turn Snapping ON when it is OFF. Works for Surface and TRS snapping.

While interacting with the gizmo hold:

- `Z` --> turns on surface snapping
- `X` --> turns on TRS snapping, based on current mode

While NOT interacting a quick press of the following will
- `Z` --> toggle surface snapping on/off
- `X` --> toggle TRS snapping, based on current mode

Gizmo responds to snapping external changes too, in case there is a key binding for a specific snapping mode.

There is a CVar which allows to also temporarily turn OFF snapping when it's ON.

`Editor.Gizmo.TemporarySnapDisableAvailable 1`

#Synced-CL 52417822

diff --git a/Engine/Source/Editor/Experimental/EditorInteractiveToolsFramework/Private/Behaviors/ViewportQuickToggleInputBehavior.cpp b/Engine/Source/Editor/Experimental/EditorInteractiveToolsFramework/Private/Behaviors/ViewportQuickToggleInputBehavior.cpp
new file mode 100644
--- /dev/null
+++ b/Engine/Source/Editor/Experimental/EditorInteractiveToolsFramework/Private/Behaviors/ViewportQuickToggleInputBehavior.cpp
@@ -0,0 +1,80 @@
+// Copyright Epic Games, Inc. All Rights Reserved.
+
+#include "Behaviors/ViewportQuickToggleInputBehavior.h"
+#include "Editor.h"
+
+void UViewportQuickToggleInputBehavior::Initialize(IQuickToggleBehaviorTarget* InTarget, const FKey& InKey)
+{
+	Super::Initialize(InTarget, InKey);
+
+	QuickToggleTarget = InTarget;
+	KeyPressedTime.Init(0.0, 1);
+}
+
+void UViewportQuickToggleInputBehavior::Initialize(IQuickToggleBehaviorTarget* InTarget, const TArray<FKey>& InKeys)
+{
+	Super::Initialize(InTarget, InKeys);
+
+	QuickToggleTarget = InTarget;
+	KeyPressedTime.Init(0.0, InKeys.Num());
+}
+
+FInputCaptureRequest UViewportQuickToggleInputBehavior::WantsCapture(const FInputDeviceState& InputState)
+{
+	return Super::WantsCapture(InputState);
+}
+
+bool UViewportQuickToggleInputBehavior::WantsForceEndCapture()
+{
+	return true;
+}
+
+void UViewportQuickToggleInputBehavior::ForceEndCapture(const FInputCaptureData& CaptureData)
+{
+	// Do NOT call Super — UKeyInputBehavior::ForceEndCapture would call Target->OnForceEndCapture()
+	// again via the base-class path, resulting in a double call for bRequireAllKeys=false.
+
+	if (QuickToggleTarget)
+	{
+		QuickToggleTarget->OnForceEndCapture();
+	}
+
+	KeyPressedTime.Init(0.0, KeyPressedTime.Num());
+}
+
+void UViewportQuickToggleInputBehavior::OnKeyPressedInternal(const FInputDeviceState& Input)
+{
+	Super::OnKeyPressedInternal(Input);
+
+	const int32 KeyIndex = TargetKeys.Find(Input.Keyboard.ActiveKey.Button);
+	if (KeyIndex > -1)
+	{
+		KeyPressedTime[KeyIndex] = FPlatformTime::Seconds();
+	}
+}
+
+void UViewportQuickToggleInputBehavior::OnKeyReleasedInternal(const FInputDeviceState& Input)
+{
+	if (QuickToggleTarget && CanQuickToggle(Input))
+	{
+
+		FKey Key = Input.Keyboard.ActiveKey.Button;
+		const int32 KeyIndex = TargetKeys.Find(Key);
+		if (KeyIndex > -1)
+		{
+			const double TimeSinceLastPress = FPlatformTime::Seconds() - KeyPressedTime[KeyIndex];
+			if (TimeSinceLastPress <= QuickToggleDelay)
+			{
+				QuickToggleTarget->OnQuickToggle(Key);
+			}
+		}
+	}
+
+	Super::OnKeyReleasedInternal(Input);
+}
+
+bool UViewportQuickToggleInputBehavior::CanQuickToggle(const FInputDeviceState& Input) const
+{
+	return CanQuickToggleFunc == nullptr || CanQuickToggleFunc(Input);
+}
+
diff --git a/Engine/Source/Editor/Experimental/EditorInteractiveToolsFramework/Private/EditorGizmos/EditorTRSGizmo.cpp b/Engine/Source/Editor/Experimental/EditorInteractiveToolsFramework/Private/EditorGizmos/EditorTRSGizmo.cpp
--- a/Engine/Source/Editor/Experimental/EditorInteractiveToolsFramework/Private/EditorGizmos/EditorTRSGizmo.cpp
+++ b/Engine/Source/Editor/Experimental/EditorInteractiveToolsFramework/Private/EditorGizmos/EditorTRSGizmo.cpp
@@ -4,7 +4,6 @@
 
 #include "AnimationCoreLibrary.h"
 #include "BaseBehaviors/ClickDragBehavior.h"
-#include "BaseBehaviors/KeyAsModifierInputBehavior.h"
 #include "BaseBehaviors/MouseHoverBehavior.h"
 #include "BaseGizmos/AxisSources.h"
 #include "BaseGizmos/GizmoElementGroup.h"
@@ -12,7 +11,6 @@
 #include "BaseGizmos/GizmoMath.h"
 #include "BaseGizmos/GizmoRenderingUtil.h"
 #include "BaseGizmos/GizmoUtil.h"
-#include "BaseGizmos/StateTargets.h"
 #include "BaseGizmos/ViewBasedTransformAdjusters.h"
 #include "Behaviors/SingleClickAndDragBehavior.h"
 #include "Behaviors/ViewportClickDragBehavior.h"
@@ -34,7 +32,6 @@
 #include "EditorGizmos/TransformGizmo.h"
 #include "EditorInteractiveToolsFrameworkStyle.h"
 #include "EditorModeManager.h"
-#include "EditorToolDataVisualizer.h"
 #include "TransformGizmoEditorSettings.h"
 #include "EditorViewportClient.h"
 #include "GizmoEdModeInterface.h"
@@ -44,7 +41,6 @@
 #include "GizmoElementScaleAxis.h"
 #include "GizmoElementScaleGroup.h"
 #include "GizmoElementTranslateGroup.h"
-#include "GlobalRenderResources.h"
 #include "Intersection/IntersectionUtil.h"
 #include "Materials/Material.h"
 #include "Math/UnitConversion.h"
@@ -52,11 +48,8 @@
 #include "SceneQueries/SceneSnappingManager.h"
 #include "Selection.h"
 #include "Widgets/SGizmoCursor.h"
-#include "ScopedTransaction.h"
-#include "ToolDataVisualizer.h"
 #include "BaseBehaviors/KeyInputBehavior.h"
 #include "Tools/AssetEditorContextInterface.h"
-#include "TransformGizmoEditorSettings.h"
 #include "UnrealEdGlobals.h"
 #include "Settings/LevelEditorViewportSettings.h"
 #include "Toolkits/IToolkitHost.h"
@@ -1346,6 +1339,8 @@ void UEditorTRSGizmo::Setup()
 		return;
 	}
 
+	CacheSnappingStates();
+
 	UInteractiveGizmo::Setup();
 
 	using namespace UE::Editor::InteractiveToolsFramework;
@@ -2306,6 +2301,8 @@ void UEditorTRSGizmo::UpdateMode()
 		CurrentMode = NewMode;
 		CurrentAxisToDraw = NewAxisToDraw;
 
+		bTRSSnapDirty = true;
+
 		// Reset various profilers so that we can see the cost of the new mode
 		Profiler->Reset(UE::Editor::InteractiveToolsFramework::Private::TransformGizmoLocals::RenderName);
 		Profiler->Reset(UE::Editor::InteractiveToolsFramework::Private::TransformGizmoLocals::DrawHUDName);
@@ -4163,6 +4160,8 @@ void UEditorTRSGizmo::OnClickDragRotateAxisInternal(
 		const FQuat Delta = std::forward<MakePullQuatFunc>(InMakePullQuat)(AppliedDeltaAngle, AngleSign, PreviousScreenPos, InDragPos.ScreenPosition);
 		ApplyRotateDelta(Delta);
 
+		CumulativeRotationDelta = Delta * CumulativeRotationDelta;
+
 		InteractionDelta.X += (DeltaAngle * AngleSign);
 		InteractionCurrAngle = Angle;
 		InteractionScreenCurrPos = FMath::IsNearlyZero(DeltaAngle) ? PreviousScreenPos : InDragPos.ScreenPosition;
@@ -4211,6 +4210,8 @@ void UEditorTRSGizmo::OnClickDragRotateAxisInternal(
 			const FQuat Delta = std::forward<MakeArcQuatFunc>(InMakeArcQuat)(HitAngleDelta, AngleSign);
 			ApplyRotateDelta(Delta);
 
+			CumulativeRotationDelta = CumulativeRotationDelta * Delta;
+
 			InteractionDelta.X += (HitAngleDelta * AngleSign);
 			InteractionCurrAngle = HitAngle;
 
@@ -4802,6 +4803,10 @@ void UEditorTRSGizmo::OnClickPress(const FInputDeviceRay& PressPos)
 	{
 		bRotationPrecisionModeDirty = true;
 	}
+
+	CacheSnappingStates();
+
+	CumulativeRotationDelta = FQuat::Identity;
 }
 
 void UEditorTRSGizmo::OnClickDrag(const FInputDeviceRay& DragPos)
@@ -4812,6 +4817,8 @@ void UEditorTRSGizmo::OnClickDrag(const FInputDeviceRay& DragPos)
 		return;
 	}
 
+	HandleSnappingChanges(DragPos);
+
 	// Update Debug
 	{
 		DebugData.InteractionRay = DragPos.WorldRay;
@@ -5260,6 +5267,8 @@ void UEditorTRSGizmo::OnClickPressAxis(const FInputDeviceRay& PressPos)
 
 	bInInteraction = true;
 	SetModeLastHitPart(EGizmoTransformMode::Translate, LastHitPart);
+
+	CumulativeDragDelta = FVector::ZeroVector;
 }
 
 void UEditorTRSGizmo::OnDragStartTranslateAxis(const FInputDeviceRay& PressPos)
@@ -5512,6 +5521,181 @@ void UEditorTRSGizmo::RestoreCursorPosition() const
 	}
 }
 
+void UEditorTRSGizmo::ApplySnappingUpdateDeltas(const FInputDeviceRay& DragPos)
+{
+	switch (CurrentMode)
+	{
+	case EGizmoTransformMode::Translate:
+	{
+		// Rounding to closest snapped value
+		FVector RoundedDelta = CumulativeDragDelta;
+		SnapTranslate(RoundedDelta, GetInteractionAxisList(), DragPos.WorldRay);
+		ApplyTranslateDelta(RoundedDelta - CumulativeDragDelta);
+		CumulativeDragDelta = FVector::ZeroVector;
+
+		break;
+	}
+	case EGizmoTransformMode::Rotate:
+	{
+		if (bCachedRotationSnap) // Snapping Off --> On
+		{
+			FQuat SnappedCumulativeRotationDelta = CumulativeRotationDelta;
+			SnapRotateDelta(SnappedCumulativeRotationDelta, GetInteractionAxisList());
+
+			// Compute and apply the delta to reach that rotation, so that our object snaps back to closes snapping value
+			const FQuat DeltaRotationToSnap = SnappedCumulativeRotationDelta * CumulativeRotationDelta.Inverse();
+			ApplyRotateDelta(DeltaRotationToSnap);
+
+			// Update interaction widgets where needed
+			if (UE::Editor::InteractiveToolsFramework::Internal::IsAxisSingular(GetInteractionAxisList()))
+			{
+				const double CumulativeAngle = CumulativeRotationDelta.GetAngle();
+				const double SnappedCumulativeAngle = SnappedCumulativeRotationDelta.GetAngle();
+				double CorrectionAngle = SnappedCumulativeAngle - CumulativeAngle;
+
+				if (GetInteractionAxisList() == EAxisList::Screen)
+				{
+					if (InteractionCurrAngle - InteractionStartAngle < 0.0f)
+					{
+						CorrectionAngle *= -1.0f;
+					}
+
+					InteractionCurrAngle += CorrectionAngle;
+				}
+				else
+				{
+					if (InteractionDelta.X < 0.0f)
+					{
+						CorrectionAngle *= -1.0f;
+					}
+
+					InteractionDelta.X += CorrectionAngle;
+				}
+			}
+
+			CumulativeRotationDelta = FQuat::Identity;
+			CurrentRotation = StartRotation = CurrentTransform.GetRotation();
+			InteractionArcBallStartPoint = InteractionArcBallCurrPoint;
+		}
+		else // Snapping On --> Off
+		{
+			CumulativeRotationDelta = FQuat::Identity;
+		}
+
+		break;
+	}
+	case EGizmoTransformMode::Scale:
+	{
+		// Rounding to closest snapped value
+		FVector RoundedDelta = CumulativeDragDelta;
+		SnapScaleDelta(RoundedDelta, GetInteractionAxisList(), DragPos.WorldRay);
+		ApplyScaleDelta(RoundedDelta - CumulativeDragDelta);
+		CumulativeDragDelta = FVector::ZeroVector;
+
+		break;
+	}
+	default:;
+	}
+}
+
+void UEditorTRSGizmo::CacheSnappingStates()
+{
+	if (const ULevelEditorViewportSettings* Settings = GetDefault<ULevelEditorViewportSettings>())
+	{
+		bCachedTranslationSnap = Settings->GridEnabled;
+		bCachedRotationSnap = Settings->RotGridEnabled;
+		bCachedScaleSnap = Settings->SnapScaleEnabled;
+
+		bCachedSurfaceSnap = Settings->SnapToSurface.bEnabled;
+	}
+}
+
+bool UEditorTRSGizmo::HasCurrentTRSSnappingChanged() const
+{
+	bool bChanged = false;
+	if (const ULevelEditorViewportSettings* Settings = GetDefault<ULevelEditorViewportSettings>())
+	{
+		switch (CurrentMode)
+		{
+		case EGizmoTransformMode::Translate:
+			bChanged = bCachedTranslationSnap != Settings->GridEnabled;
+			break;
+		case EGizmoTransformMode::Rotate:
+			bChanged = bCachedRotationSnap != Settings->RotGridEnabled;
+			break;
+		case EGizmoTransformMode::Scale:
+			bChanged = bCachedScaleSnap != Settings->SnapScaleEnabled;
+			break;
+		default:;
+		}
+	}
+
+	return bChanged;
+}
+
+bool UEditorTRSGizmo::HasSurfaceSnappingChanged() const
+{
+	if (const ULevelEditorViewportSettings* Settings = GetDefault<ULevelEditorViewportSettings>())
+	{
+		return bCachedSurfaceSnap != Settings->SnapToSurface.bEnabled;
+	}
+
+	return false;
+}
+
+
+void UEditorTRSGizmo::HandleSnappingChanges(const FInputDeviceRay& DragPos)
+{
+	if (HasCurrentTRSSnappingChanged())
+	{
+		CacheSnappingStates();
+		ApplySnappingUpdateDeltas(DragPos);
+	}
+
+	if (HasSurfaceSnappingChanged())
+	{
+		CacheSnappingStates();
+		ApplySurfaceSnappingUpdateDeltas(DragPos);
+	}
+
+	bTRSSnapDirty = false;
+	bSurfaceSnapDirty = false;
+}
+
+void UEditorTRSGizmo::ApplySurfaceSnappingUpdateDeltas(const FInputDeviceRay& DragPos)
+{
+	const bool bSurfaceSnapEnabled = GetDefault<ULevelEditorViewportSettings>()->SnapToSurface.bEnabled;
+	if (bSurfaceSnapEnabled)
+	{
+		GUnrealEd->UpdatePivotLocationForSelection();
+	}
+	else
+	{
+		GUnrealEd->UpdatePivotLocationForSelection();
+
+		CurrentTransform = GetActiveTransform();
+
+		InteractionPlanarOrigin = CurrentTransform.GetLocation();
+
+		if (CameraFollowsMovement() || bIndirectManipulation)
+		{
+			InteractionScreenCurrPos = DragPos.ScreenPosition;
+		}
+		else
+		{
+			FVector::FReal HitDepth;
+			if (GetRayParamIntersectionWithInteractionPlane(DragPos, HitDepth))
+			{
+				InteractionPlanarCurrPoint = DragPos.WorldRay.Origin + DragPos.WorldRay.Direction * HitDepth;
+			}
+		}
+
+		FGizmoState State;
+		State.TransformMode = EGizmoTransformMode::Translate;
+		GetEditorViewportClient()->BeginTransform(State);
+	}
+}
+
 bool UEditorTRSGizmo::CanRender(IToolsContextRenderAPI* InRenderAPI) const
 {
 	if (!InRenderAPI || !GizmoElementRoot || !CanInteract())
@@ -5601,6 +5785,7 @@ void UEditorTRSGizmo::OnClickDragTranslateAxis(const FInputDeviceRay& DragPos)
 		SnapTranslate(Delta, GetInteractionAxisList(), DragPos.WorldRay);
 
 		ApplyTranslateDelta(Delta);
+		CumulativeDragDelta += Delta;
 
 		InteractionScreenCurrPos = Delta.IsNearlyZero() ? PreviousScreenPos : DragPos.ScreenPosition;
 
@@ -5648,6 +5833,7 @@ void UEditorTRSGizmo::OnClickDragTranslateAxis(const FInputDeviceRay& DragPos)
 			
 			SnapTranslate(Delta, GetInteractionAxisList(), DragPos.WorldRay);
 			ApplyTranslateDelta(Delta);
+			CumulativeDragDelta += Delta;
 			
 			if (!Delta.IsNearlyZero())
 			{
@@ -5706,6 +5892,8 @@ void UEditorTRSGizmo::OnClickPressTranslatePlanar(const FInputDeviceRay& PressPo
 	{
 		BeginTranslateDelta();
 	}
+
+	CumulativeDragDelta = FVector::ZeroVector;
 }
 
 void UEditorTRSGizmo::OnClickDragTranslatePlanar(const FInputDeviceRay& DragPos)
@@ -5758,6 +5946,7 @@ void UEditorTRSGizmo::OnClickDragTranslatePlanar(const FInputDeviceRay& DragPos)
 			
 		SnapTranslate(Delta, GetInteractionAxisList(), Ray);
 		ApplyTranslateDelta(Delta);
+		CumulativeDragDelta += Delta;
 
 		// Update cached transform
 		CurrentTransform = GetActiveTransform();
@@ -5833,6 +6022,8 @@ void UEditorTRSGizmo::OnClickPressScreenSpaceTranslate(const FInputDeviceRay& Pr
 	{
 		OnClickPressPlanar(PressPos);
 	}
+
+	CumulativeDragDelta = FVector::ZeroVector;
 }
 
 void UEditorTRSGizmo::OnClickDragScreenSpaceTranslate(const FInputDeviceRay& DragPos)
@@ -6157,6 +6348,8 @@ void UEditorTRSGizmo::OnClickPressScale(const FInputDeviceRay& PressPos)
 	BeginScaleDelta();
 
 	SetModeLastHitPart(EGizmoTransformMode::Scale, LastHitPart);
+
+	CumulativeDragDelta = FVector::ZeroVector;
 }
 
 void UEditorTRSGizmo::OnClickDragScale(const FInputDeviceRay& DragPos)
@@ -6311,6 +6504,8 @@ void UEditorTRSGizmo::OnClickDragScale(const FInputDeviceRay& DragPos)
 
 		ApplyScaleDelta(ScaleDelta);
 
+		CumulativeDragDelta += ScaleDelta;
+
 		FVector PostDeltaScale = GetActiveTransform().GetScale3D();
 		
 		if (FMath::IsNearlyEqual((PreDeltaScale - PostDeltaScale).SquaredLength(), 0.0))
@@ -7278,6 +7473,8 @@ void UEditorTRSGizmo::OnClickDragScreenSpaceRotate(const FInputDeviceRay& InDrag
 		const FQuat Delta(InteractionPlanarNormal, DeltaAngle);
 		ApplyRotateDelta(Delta);
 
+		CumulativeRotationDelta = CumulativeRotationDelta * Delta;
+
 		InteractionCurrAngle = HitAngle;
 		InteractionScreenCurrPos = FMath::IsNearlyZero(DeltaAngle) ? PreviousScreenPos : InDragPos.ScreenPosition;
 	}
@@ -7296,6 +7493,8 @@ void UEditorTRSGizmo::OnClickDragScreenSpaceRotate(const FInputDeviceRay& InDrag
 			const FQuat Delta = FQuat(InteractionPlanarNormal, HitAngleDelta);
 			ApplyRotateDelta(Delta);
 
+			CumulativeRotationDelta = CumulativeRotationDelta * Delta;
+
 			InteractionCurrAngle = HitAngle;
 
 			// Update Debug
@@ -7517,6 +7716,7 @@ void UEditorTRSGizmo::OnClickDragArcBallRotate(const FInputDeviceRay& DragPos)
 	{
 		SnapRotateDelta(DeltaRot, GetInteractionAxisList());
 		ApplyRotateDelta(DeltaRot);
+		CumulativeRotationDelta = DeltaRot * CumulativeRotationDelta;
 		CurrentRotation = DeltaRot * CurrentRotation;
 	}
 
diff --git a/Engine/Source/Editor/Experimental/EditorInteractiveToolsFramework/Private/ViewportInteractions/ViewportCameraRotateInteraction.cpp b/Engine/Source/Editor/Experimental/EditorInteractiveToolsFramework/Private/ViewportInteractions/ViewportCameraRotateInteraction.cpp
--- a/Engine/Source/Editor/Experimental/EditorInteractiveToolsFramework/Private/ViewportInteractions/ViewportCameraRotateInteraction.cpp
+++ b/Engine/Source/Editor/Experimental/EditorInteractiveToolsFramework/Private/ViewportInteractions/ViewportCameraRotateInteraction.cpp
@@ -39,7 +39,13 @@ UViewportCameraRotateInteraction::UViewportCameraRotateInteraction()
 				|| CommandMatchesKey(FViewportNavigationCommands::Get().RotateLeft, ActiveKey)
 				|| CommandMatchesKey(FViewportNavigationCommands::Get().RotateRight, ActiveKey))
 			{
-				bResult = IsAnyMouseButtonDown();
+				bool bIsGizmoDragging = false;
+				if (UViewportInteractionsBehaviorSource* BehaviorSource = GetViewportInteractionsBehaviorSource())
+				{
+					bIsGizmoDragging = BehaviorSource->IsGizmoDragging();
+				}
+
+				bResult = IsAnyMouseButtonDown() && !bIsGizmoDragging;
 				// If a mouse button is down, but mouse hasn't been dragged yet, we need to tell the behavior source that
 				// User is performing mouse look or at least camera movement. This also prevents regular clicks from performing their action on button release
 				// e.g. we don't want a RMB release to summon a context menu after mouse looking
diff --git a/Engine/Source/Editor/Experimental/EditorInteractiveToolsFramework/Private/ViewportInteractions/ViewportCameraTranslateInteraction.cpp b/Engine/Source/Editor/Experimental/EditorInteractiveToolsFramework/Private/ViewportInteractions/ViewportCameraTranslateInteraction.cpp
--- a/Engine/Source/Editor/Experimental/EditorInteractiveToolsFramework/Private/ViewportInteractions/ViewportCameraTranslateInteraction.cpp
+++ b/Engine/Source/Editor/Experimental/EditorInteractiveToolsFramework/Private/ViewportInteractions/ViewportCameraTranslateInteraction.cpp
@@ -69,7 +69,13 @@ UViewportCameraTranslateInteraction::UViewportCameraTranslateInteraction()
 				CommandMatchesKey(FViewportNavigationCommands::Get().LocalUp, ActiveKey) ||
 				CommandMatchesKey(FViewportNavigationCommands::Get().LocalDown, ActiveKey))
 			{
-				bResult = IsAnyMouseButtonDown();
+				bool bIsGizmoDragging = false;
+				if (UViewportInteractionsBehaviorSource* BehaviorSource = GetViewportInteractionsBehaviorSource())
+				{
+					bIsGizmoDragging = BehaviorSource->IsGizmoDragging();
+				}
+
+				bResult = IsAnyMouseButtonDown() && !bIsGizmoDragging;
 				// If a mouse button is down, but mouse hasn't been dragged yet, we need to tell the behavior source that
 				// User is performing mouse look or at least camera movement. This also prevents regular clicks from performing their action on button release
 				// e.g. we don't want a RMB release to summon a context menu after mouse looking
diff --git a/Engine/Source/Editor/Experimental/EditorInteractiveToolsFramework/Private/ViewportInteractions/ViewportFOVInteraction.cpp b/Engine/Source/Editor/Experimental/EditorInteractiveToolsFramework/Private/ViewportInteractions/ViewportFOVInteraction.cpp
--- a/Engine/Source/Editor/Experimental/EditorInteractiveToolsFramework/Private/ViewportInteractions/ViewportFOVInteraction.cpp
+++ b/Engine/Source/Editor/Experimental/EditorInteractiveToolsFramework/Private/ViewportInteractions/ViewportFOVInteraction.cpp
@@ -39,7 +39,13 @@ UViewportFOVInteraction::UViewportFOVInteraction()
 			if (CommandMatchesKey(FViewportNavigationCommands::Get().FovZoomOut, ActiveKey) ||
 				CommandMatchesKey(FViewportNavigationCommands::Get().FovZoomIn, ActiveKey))
 			{
-				bResult = IsAnyMouseButtonDown();
+				bool bIsGizmoDragging = false;
+				if (UViewportInteractionsBehaviorSource* BehaviorSource = GetViewportInteractionsBehaviorSource())
+				{
+					bIsGizmoDragging = BehaviorSource->IsGizmoDragging();
+				}
+
+				bResult = IsAnyMouseButtonDown() && !bIsGizmoDragging;
 				// If a mouse button is down, but mouse hasn't been dragged yet, we need to tell the behavior source that
 				// User is performing mouse look or at least camera movement. This also prevents regular clicks from performing their action on button release
 				// e.g. we don't want a RMB release to summon a context menu after mouse looking
diff --git a/Engine/Source/Editor/Experimental/EditorInteractiveToolsFramework/Private/ViewportInteractions/ViewportInteractionsUtilities.cpp b/Engine/Source/Editor/Experimental/EditorInteractiveToolsFramework/Private/ViewportInteractions/ViewportInteractionsUtilities.cpp
--- a/Engine/Source/Editor/Experimental/EditorInteractiveToolsFramework/Private/ViewportInteractions/ViewportInteractionsUtilities.cpp
+++ b/Engine/Source/Editor/Experimental/EditorInteractiveToolsFramework/Private/ViewportInteractions/ViewportInteractionsUtilities.cpp
@@ -17,6 +17,7 @@
 #include "ViewportInteractions/ViewportOrbitInteraction.h"
 #include "ViewportInteractions/ViewportOrthoPanInteraction.h"
 #include "ViewportInteractions/ViewportPanInteraction.h"
+#include "ViewportInteractions/ViewportSnapToggleInteraction.h"
 #include "ViewportInteractions/ViewportViewAngleInteraction.h"
 #include "ViewportInteractions/ViewportZoomInteraction.h"
 
@@ -42,7 +43,8 @@ void AddDefaultCameraMovementInteractions(
 				UViewportZoomInteraction::StaticClass(),
 				UViewportCameraTranslateInteraction::StaticClass(),
 				UViewportCameraRotateInteraction::StaticClass(),
-				UViewportFOVInteraction::StaticClass() }
+				UViewportFOVInteraction::StaticClass(),
+				UViewportSnapToggleInteraction::StaticClass()}
 		);
 	}
 }
diff --git a/Engine/Source/Editor/Experimental/EditorInteractiveToolsFramework/Private/ViewportInteractions/ViewportSnapToggleInteraction.cpp b/Engine/Source/Editor/Experimental/EditorInteractiveToolsFramework/Private/ViewportInteractions/ViewportSnapToggleInteraction.cpp
new file mode 100644
--- /dev/null
+++ b/Engine/Source/Editor/Experimental/EditorInteractiveToolsFramework/Private/ViewportInteractions/ViewportSnapToggleInteraction.cpp
@@ -0,0 +1,308 @@
+// Copyright Epic Games, Inc. All Rights Reserved.
+
+#include "ViewportInteractions/ViewportSnapToggleInteraction.h"
+
+#include "Editor/UnrealEdEngine.h"
+#include "EditorModeManager.h"
+#include "EditorViewportCommands.h"
+#include "HAL/IConsoleManager.h"
+#include "Settings/LevelEditorViewportSettings.h"
+#include "UnrealEdGlobals.h"
+#include "ViewportInteractions/ViewportInteractionsBehaviorSource.h"
+
+namespace UE::ViewportSnapToggle
+{
+	TAutoConsoleVariable<bool> CVarTemporarySnapDisableAvailable(
+		TEXT("Editor.Gizmo.TemporarySnapDisableAvailable"),
+		false,
+		TEXT("When enabled, temporary snapping held keys can also temporarily disable snapping when it is currently enabled.")
+	);
+} // namespace UE::ViewportSnapToggle
+
+UViewportSnapToggleInteraction::UViewportSnapToggleInteraction()
+{
+	CurrentWidgetMode = UE::Widget::EWidgetMode::WM_None;
+
+	const FName UniqueName = MakeUniqueObjectName(GetTransientPackageAsObject(), UViewportSnapToggleInteraction::StaticClass(), *FString::Printf(TEXT("%s_Behavior"), *GetClass()->GetName()));
+	UViewportQuickToggleInputBehavior* QuickToggleBehavior = NewObject<UViewportQuickToggleInputBehavior>(GetTransientPackageAsObject(), UniqueName);
+
+	QuickToggleBehavior->bRequireAllKeys = false;
+	QuickToggleBehavior->Initialize(this, GetKeys());
+	QuickToggleBehavior->SetDefaultPriority(UE::Editor::ViewportInteractions::VIEWPORT_INTERACTIONS_DEFAULT_PRIORITY);
+
+	QuickToggleBehavior->ModifierCheckFunc = [](const FInputDeviceState& InputDeviceState) -> bool
+	{
+		return !InputDeviceState.bCtrlKeyDown && !InputDeviceState.bShiftKeyDown && !InputDeviceState.bAltKeyDown && !InputDeviceState.bCmdKeyDown;
+	};
+
+	QuickToggleBehaviorWeak = QuickToggleBehavior;
+
+	RegisterInputBehavior(QuickToggleBehavior);
+}
+
+void UViewportSnapToggleInteraction::Initialize(UViewportInteractionsBehaviorSource* InViewportInteractionsBehaviorSource)
+{
+	Super::Initialize(InViewportInteractionsBehaviorSource);
+
+	if (FEditorModeTools* ModeTools = GetModeTools())
+	{
+		ModeTools->OnWidgetModeChanged().AddUObject(this, &UViewportSnapToggleInteraction::OnWidgetModeChanged);
+		CurrentWidgetMode = ModeTools->GetWidgetMode();
+	}
+}
+
+void UViewportSnapToggleInteraction::Shutdown()
+{
+	if (FEditorModeTools* const ModeTools = GetModeTools())
+	{
+		ModeTools->OnWidgetModeChanged().RemoveAll(this);
+	}
+
+	if (bTRSKeyHeld)
+	{
+		SetTRSSnapping(bCachedTRSSnapping);
+		bTRSKeyHeld = false;
+	}
+
+	if (bSurfaceKeyHeld)
+	{
+		SetSurfaceSnapping(bCachedSurfaceSnapping);
+		bSurfaceKeyHeld = false;
+	}
+
+	Super::Shutdown();
+}
+
+void UViewportSnapToggleInteraction::OnQuickToggle(const FKey& InKey)
+{
+	if (!FEditorViewportCommands::IsRegistered())
+	{
+		return;
+	}
+
+	using namespace UE::Editor::ViewportInteractions;
+	const FEditorViewportCommands& ViewportCommands = FEditorViewportCommands::Get();
+
+	if (CommandMatchesKey(ViewportCommands.CurrentTRSSnapToggle, InKey))
+	{
+		// bTRSKeyHeld is false: this means we are quick toggling ON --> OFF
+		// while temporary disable is not available
+		if (!bTRSKeyHeld)
+		{
+			SetTRSSnapping(false);
+		}
+
+		bTRSKeyHeld = false;
+	}
+	else if (CommandMatchesKey(ViewportCommands.SurfaceSnapping, InKey))
+	{
+		// bSurfaceKeyHeld is false: this means we are quick toggling ON --> OFF
+		// while temporary disable is not available
+		if (!bSurfaceKeyHeld)
+		{
+			SetSurfaceSnapping(false);
+		}
+		bSurfaceKeyHeld = false;
+	}
+}
+
+void UViewportSnapToggleInteraction::OnKeyPressed(const FKey& InKeyID)
+{
+	// CVarTemporarySnapDisableAvailable: 0 (default)
+	// - OFF-->ON : we can set snapping ON as soon as key is pressed. This works both for Quick Toggle and Temp Snapping
+	// - ON -->OFF: we cannot know if this is going to be a quick toggle or a Temp Snapping. Skip setting snapping, and do that on Quick Toggle if needed.
+
+	// CVarTemporarySnapDisableAvailable: 1
+	// - We can always set snapping right away. Snapping will not be set on Quick Toggle.
+
+	using namespace UE::Editor::ViewportInteractions;
+
+	if (!FEditorViewportCommands::IsRegistered())
+	{
+		return;
+	}
+
+	const FEditorViewportCommands& ViewportCommands = FEditorViewportCommands::Get();
+
+	if (!bTRSKeyHeld && CommandMatchesKey(ViewportCommands.CurrentTRSSnapToggle, InKeyID))
+	{
+		CacheTRSSnapping();
+
+		// When temporary disable is not available, do not set snapping on key pressed
+		// ON --> OFF should only work for Quick Toggle. Leave bTRSKeyHeld set to false.
+		if (bCachedTRSSnapping && !CanTemporaryDisableSnapping())
+		{
+			return;
+		}
+
+		SetTRSSnapping(!bCachedTRSSnapping);
+		bTRSKeyHeld = true;
+	}
+	else if (!bSurfaceKeyHeld && CommandMatchesKey(ViewportCommands.SurfaceSnapping, InKeyID))
+	{
+		CacheSurfaceSnapping();
+
+		if (bCachedSurfaceSnapping && !CanTemporaryDisableSnapping())
+		{
+			return;
+		}
+
+		SetSurfaceSnapping(!bCachedSurfaceSnapping);
+		bSurfaceKeyHeld = true;
+	}
+}
+
+void UViewportSnapToggleInteraction::OnKeyReleased(const FKey& InKeyID)
+{
+	using namespace UE::Editor::ViewportInteractions;
+
+	if (!FEditorViewportCommands::IsRegistered())
+	{
+		return;
+	}
+
+	const FEditorViewportCommands& ViewportCommands = FEditorViewportCommands::Get();
+
+	// We check key(s) hold state, since OnQuickToggle could have already set the flag to false.
+	// In that case, do nothing.
+
+	if (bTRSKeyHeld && CommandMatchesKey(ViewportCommands.CurrentTRSSnapToggle, InKeyID))
+	{
+		SetTRSSnapping(bCachedTRSSnapping);
+		bTRSKeyHeld = false;
+	}
+	else if (bSurfaceKeyHeld && CommandMatchesKey(ViewportCommands.SurfaceSnapping, InKeyID))
+	{
+		SetSurfaceSnapping(bCachedSurfaceSnapping);
+		bSurfaceKeyHeld = false;
+	}
+}
+
+void UViewportSnapToggleInteraction::OnForceEndCapture()
+{
+	if (bTRSKeyHeld)
+	{
+		SetTRSSnapping(bCachedTRSSnapping);
+		bTRSKeyHeld = false;
+	}
+
+	if (bSurfaceKeyHeld)
+	{
+		SetSurfaceSnapping(bCachedSurfaceSnapping);
+		bSurfaceKeyHeld = false;
+	}
+}
+
+void UViewportSnapToggleInteraction::OnCommandChordChanged()
+{
+	if (UViewportQuickToggleInputBehavior* QuickToggleBehavior = QuickToggleBehaviorWeak.Get())
+	{
+		QuickToggleBehavior->Initialize(this, GetKeys());
+	}
+}
+
+TArray<TSharedPtr<FUICommandInfo>> UViewportSnapToggleInteraction::GetCommands() const
+{
+	if (FEditorViewportCommands::IsRegistered())
+	{
+		return { FEditorViewportCommands::Get().CurrentTRSSnapToggle, FEditorViewportCommands::Get().SurfaceSnapping };
+	}
+
+	return {};
+}
+
+TArray<FKey> UViewportSnapToggleInteraction::GetKeys() const
+{
+	TArray<FKey> Keys;
+
+	for (int32 ChordIndex = 0; ChordIndex < static_cast<uint8>(EMultipleKeyBindingIndex::NumChords); ++ChordIndex)
+	{
+		const EMultipleKeyBindingIndex BindingIndex = static_cast<EMultipleKeyBindingIndex>(ChordIndex);
+
+		for (const TSharedPtr<FUICommandInfo>& Command : GetCommands())
+		{
+			Keys.Add(Command->GetActiveChord(BindingIndex)->Key);
+		}
+	}
+
+	return Keys;
+}
+
+void UViewportSnapToggleInteraction::SetTRSSnapping(bool bEnable) const
+{
+	using namespace UE::Widget;
+
+	if (const ULevelEditorViewportSettings* Settings = GetDefault<ULevelEditorViewportSettings>())
+	{
+		switch (CurrentWidgetMode)
+		{
+		case WM_Translate:
+		case WM_2D:
+			if (Settings->GridEnabled != bEnable)
+			{
+				GUnrealEd->Exec(GEditor->GetEditorWorldContext().World(), *FString::Printf(TEXT("MODE GRID=%d"), bEnable));
+			}
+			break;
+
+		case WM_Rotate:
+			if (Settings->RotGridEnabled != bEnable)
+			{
+				GUnrealEd->Exec(GEditor->GetEditorWorldContext().World(), *FString::Printf(TEXT("MODE ROTGRID=%d"), bEnable));
+			}
+			break;
+
+		case WM_Scale:
+			if (Settings->SnapScaleEnabled != bEnable)
+			{
+				GUnrealEd->Exec(GEditor->GetEditorWorldContext().World(), *FString::Printf(TEXT("MODE SCALEGRID=%d"), bEnable));
+			}
+			break;
+
+		default:;
+		}
+	}
+}
+
+void UViewportSnapToggleInteraction::SetSurfaceSnapping(bool bEnable)
+{
+	ULevelEditorViewportSettings* Settings = GetMutableDefault<ULevelEditorViewportSettings>();
+
+	if (Settings->SnapToSurface.bEnabled != bEnable)
+	{
+		Settings->SnapToSurface.bEnabled = bEnable;
+	}
+}
+
+bool UViewportSnapToggleInteraction::CanTemporaryDisableSnapping()
+{
+	return UE::ViewportSnapToggle::CVarTemporarySnapDisableAvailable.GetValueOnAnyThread();
+}
+
+void UViewportSnapToggleInteraction::OnWidgetModeChanged(UE::Widget::EWidgetMode InWidgetMode)
+{
+	CurrentWidgetMode = InWidgetMode;
+}
+
+void UViewportSnapToggleInteraction::CacheTRSSnapping()
+{
+	const ULevelEditorViewportSettings* LevelViewportSettings = GetDefault<ULevelEditorViewportSettings>();
+
+	switch (CurrentWidgetMode)
+	{
+	case UE::Widget::EWidgetMode::WM_Rotate:
+		bCachedTRSSnapping = LevelViewportSettings->RotGridEnabled;
+		break;
+	case UE::Widget::EWidgetMode::WM_Scale:
+		bCachedTRSSnapping = LevelViewportSettings->SnapScaleEnabled;
+		break;
+	default:
+		bCachedTRSSnapping = LevelViewportSettings->GridEnabled;
+		break;
+	}
+}
+
+void UViewportSnapToggleInteraction::CacheSurfaceSnapping()
+{
+	bCachedSurfaceSnapping = GetDefault<ULevelEditorViewportSettings>()->SnapToSurface.bEnabled;
+}
+
diff --git a/Engine/Source/Editor/Experimental/EditorInteractiveToolsFramework/Public/Behaviors/ViewportQuickToggleInputBehavior.h b/Engine/Source/Editor/Experimental/EditorInteractiveToolsFramework/Public/Behaviors/ViewportQuickToggleInputBehavior.h
new file mode 100644
--- /dev/null
+++ b/Engine/Source/Editor/Experimental/EditorInteractiveToolsFramework/Public/Behaviors/ViewportQuickToggleInputBehavior.h
@@ -0,0 +1,54 @@
+// Copyright Epic Games, Inc. All Rights Reserved.
+
+#pragma once
+
+#include "BaseBehaviors/KeyAsModifierInputBehavior.h"
+#include "BaseBehaviors/KeyInputBehavior.h"
+
+#include "ViewportQuickToggleInputBehavior.generated.h"
+
+class IQuickToggleBehaviorTarget : public IKeyInputBehaviorTarget
+{
+public:
+	virtual ~IQuickToggleBehaviorTarget() override = default;
+
+	virtual void OnQuickToggle(const FKey& InKey) {}
+};
+
+UCLASS(MinimalAPI)
+class UViewportQuickToggleInputBehavior : public UKeyInputBehavior
+{
+	GENERATED_BODY()
+
+public:
+	EDITORINTERACTIVETOOLSFRAMEWORK_API UViewportQuickToggleInputBehavior() {}
+
+	virtual EInputDevices GetSupportedDevices() override
+	{
+		return EInputDevices::Keyboard;
+	}
+
+	EDITORINTERACTIVETOOLSFRAMEWORK_API void Initialize(IQuickToggleBehaviorTarget* InTarget, const FKey& InKey);
+	EDITORINTERACTIVETOOLSFRAMEWORK_API void Initialize(IQuickToggleBehaviorTarget* InTarget, const TArray<FKey>& InKeys);
+
+	virtual FInputCaptureRequest WantsCapture(const FInputDeviceState& InputState) override;
+	virtual bool WantsForceEndCapture() override;
+	virtual void ForceEndCapture(const FInputCaptureData& CaptureData) override;
+
+	virtual void OnKeyPressedInternal(const FInputDeviceState& Input) override;
+	virtual void OnKeyReleasedInternal(const FInputDeviceState& Input) override;
+
+	/** Can be used to specify if this behavior target should receive OnQuickToggle events or not */
+	TFunction<bool(const FInputDeviceState&)> CanQuickToggleFunc = nullptr;
+
+private:
+	virtual void Initialize(IKeyInputBehaviorTarget* InTarget, const FKey& InKey) override {}
+	virtual void Initialize(IKeyInputBehaviorTarget* InTarget, const TArray<FKey>& InKeys) override {}
+
+	bool CanQuickToggle(const FInputDeviceState& Input) const;
+
+	IQuickToggleBehaviorTarget* QuickToggleTarget;
+
+	TArray<double> KeyPressedTime;
+	double QuickToggleDelay = 0.2f;
+};
diff --git a/Engine/Source/Editor/Experimental/EditorInteractiveToolsFramework/Public/EditorGizmos/EditorTRSGizmo.h b/Engine/Source/Editor/Experimental/EditorInteractiveToolsFramework/Public/EditorGizmos/EditorTRSGizmo.h
--- a/Engine/Source/Editor/Experimental/EditorInteractiveToolsFramework/Public/EditorGizmos/EditorTRSGizmo.h
+++ b/Engine/Source/Editor/Experimental/EditorInteractiveToolsFramework/Public/EditorGizmos/EditorTRSGizmo.h
@@ -14,13 +14,9 @@
 #include "EditorGizmos/TransformGizmoInterfaces.h"
 #include "FrameTypes.h"
 #include "Framework/Commands/InputChord.h"
-#include "GizmoElementRotate.h"
-#include "GizmoElementScale.h"
-#include "GizmoElementTranslate.h"
 #include "InputState.h"
 #include "Materials/MaterialInstanceDynamic.h"
 #include "Math/Axis.h"
-#include "Math/Color.h"
 #include "Math/Quat.h"
 #include "Math/Ray.h"
 #include "Math/Transform.h"
@@ -584,7 +580,19 @@ protected:
 
 	void CacheCursorPosition();
 	void RestoreCursorPosition() const;
-	
+
+	/** Used when TRS Snapping flag has changed, and gizmo needs to adjust accordingly */
+	void ApplySnappingUpdateDeltas(const FInputDeviceRay& DragPos);
+
+	/** Used when Surface Snapping flag has changed, and gizmo needs to adjust accordingly */
+	void ApplySurfaceSnappingUpdateDeltas(const FInputDeviceRay& DragPos);
+
+	void CacheSnappingStates();
+
+	bool HasCurrentTRSSnappingChanged() const;
+	bool HasSurfaceSnappingChanged() const;
+	void HandleSnappingChanges(const FInputDeviceRay& DragPos);
+
 	/** Returns true if the gizmo should be rendered. */
 	UE_API virtual bool CanRender(IToolsContextRenderAPI* InRenderAPI) const;
 
@@ -841,6 +849,18 @@ protected:
 	
 	FVector CachedViewLocation;
 
+	bool bTRSSnapDirty = false;
+	bool bSurfaceSnapDirty = false;
+
+	bool bCachedTranslationSnap = false;
+	bool bCachedRotationSnap = false;
+	bool bCachedScaleSnap = false;
+	bool bCachedSurfaceSnap = false;
+
+	/** Used when temporary rotation snapping is turned on, to reset rotation */
+	FQuat CumulativeRotationDelta;
+	FVector CumulativeDragDelta;
+
 private:
 	TWeakObjectPtr<USingleClickAndDragBehavior> ClickAndDragBehaviorWeak = nullptr;
 
diff --git a/Engine/Source/Editor/Experimental/EditorInteractiveToolsFramework/Public/ViewportInteractions/ViewportSnapToggleInteraction.h b/Engine/Source/Editor/Experimental/EditorInteractiveToolsFramework/Public/ViewportInteractions/ViewportSnapToggleInteraction.h
new file mode 100644
--- /dev/null
+++ b/Engine/Source/Editor/Experimental/EditorInteractiveToolsFramework/Public/ViewportInteractions/ViewportSnapToggleInteraction.h
@@ -0,0 +1,71 @@
+// Copyright Epic Games, Inc. All Rights Reserved.
+
+#pragma once
+
+#include "Behaviors/ViewportQuickToggleInputBehavior.h"
+#include "UnrealWidgetFwd.h"
+#include "ViewportInteraction.h"
+
+#include "ViewportSnapToggleInteraction.generated.h"
+
+#define UE_API EDITORINTERACTIVETOOLSFRAMEWORK_API
+
+/**
+ * A viewport interaction that handles temporary and permanent snapping toggles
+ *
+ * Hold: toggles the relevant snap setting while the key is held, restoring on release.
+ * Quick-press: permanently toggles the snap setting
+ */
+UCLASS(MinimalAPI, Transient)
+class UViewportSnapToggleInteraction : public UViewportInteraction, public IQuickToggleBehaviorTarget
+{
+	GENERATED_BODY()
+
+public:
+	UE_API UViewportSnapToggleInteraction();
+
+	//~ Begin UViewportInteraction
+	UE_API virtual void Initialize(UViewportInteractionsBehaviorSource* InViewportInteractionsBehaviorSource) override;
+	UE_API virtual void Shutdown() override;
+	//~ End UViewportInteraction
+
+	//~ Begin IQuickToggleBehaviorTarget
+	UE_API virtual void OnQuickToggle(const FKey& InKey) override;
+	//~ End IQuickToggleBehaviorTarget
+
+	//~ Begin IKeyInputBehaviorTarget
+	UE_API virtual void OnKeyPressed(const FKey& InKeyID) override;
+	UE_API virtual void OnKeyReleased(const FKey& InKeyID) override;
+	UE_API virtual void OnForceEndCapture() override;
+	//~ End IKeyInputBehaviorTarget
+
+protected:
+	//~ Begin UViewportInteractionBase
+	UE_API virtual void OnCommandChordChanged() override;
+	UE_API virtual TArray<TSharedPtr<FUICommandInfo>> GetCommands() const override;
+	//~ End UViewportInteractionBase
+
+	UE_API TArray<FKey> GetKeys() const;
+
+private:
+	static void SetSurfaceSnapping(bool bEnable);
+	static bool CanTemporaryDisableSnapping();
+
+	void SetTRSSnapping(bool bEnable) const;
+	void OnWidgetModeChanged(UE::Widget::EWidgetMode InWidgetMode);
+	void CacheTRSSnapping();
+	void CacheSurfaceSnapping();
+
+
+	TWeakObjectPtr<UViewportQuickToggleInputBehavior> QuickToggleBehaviorWeak;
+
+	UE::Widget::EWidgetMode CurrentWidgetMode;
+
+	bool bCachedTRSSnapping = false;
+	bool bCachedSurfaceSnapping = false;
+
+	bool bTRSKeyHeld = false;
+	bool bSurfaceKeyHeld = false;
+};
+
+#undef UE_API
diff --git a/Engine/Source/Editor/UnrealEd/Private/EditorViewportCommands.cpp b/Engine/Source/Editor/UnrealEd/Private/EditorViewportCommands.cpp
--- a/Engine/Source/Editor/UnrealEd/Private/EditorViewportCommands.cpp
+++ b/Engine/Source/Editor/UnrealEd/Private/EditorViewportCommands.cpp
@@ -211,7 +211,8 @@ void FEditorViewportCommands::RegisterCommands()
 	UI_COMMAND( RotationGridSnap, "Rotation Snap", "Toggle whether objects snap to preset angles when you rotate them.\nUse the submenu to choose the angle increment.", EUserInterfaceActionType::ToggleButton, FInputChord() );
 	UI_COMMAND( Layer2DSnap, "Layer2D Snap", "Enables or disables snapping objects to a 2D layer", EUserInterfaceActionType::ToggleButton, FInputChord());
 	UI_COMMAND( ScaleGridSnap, "Scale Snap", "Toggle whether objects resize by a preset ratio when you scale them.\nUse the submenu to choose the scaling increment.", EUserInterfaceActionType::ToggleButton, FInputChord() );
-	UI_COMMAND( SurfaceSnapping, "Surface Snapping", "Toggle whether objects snap to surfaces when you drag them.\nUse the submenu to set up the snapping behavior.", EUserInterfaceActionType::ToggleButton, FInputChord() );
+	UI_COMMAND( CurrentTRSSnapToggle, "Toggle Current TRS Snap", "Temporarily enables or disables snapping for the current transform mode while held; quick-press to toggle permanently.", EUserInterfaceActionType::Button, FInputChord( EKeys::X ) );
+	UI_COMMAND( SurfaceSnapping, "Surface Snapping", "Toggle whether objects snap to surfaces when you drag them.\nUse the submenu to set up the snapping behavior.", EUserInterfaceActionType::ToggleButton, FInputChord( EKeys::Z ) );
 	UI_COMMAND( RotateToSurfaceNormal, "Rotate to Surface Normal", "If enabled, actors rotation will match that of their snapping surface", EUserInterfaceActionType::ToggleButton, FInputChord() );
 
 	UI_COMMAND( ToggleAutoExposure, "Auto", "If enabled, enables automatic exposure", EUserInterfaceActionType::ToggleButton, FInputChord() );
diff --git a/Engine/Source/Editor/UnrealEd/Public/EditorViewportCommands.h b/Engine/Source/Editor/UnrealEd/Public/EditorViewportCommands.h
--- a/Engine/Source/Editor/UnrealEd/Public/EditorViewportCommands.h
+++ b/Engine/Source/Editor/UnrealEd/Public/EditorViewportCommands.h
@@ -297,6 +297,9 @@ public:
 	/** Enables or disables snapping to a scaling grid while scaling objects */
 	TSharedPtr< FUICommandInfo > ScaleGridSnap;
 
+	/** Temporarily toggles snapping for the current transform mode (translate/rotate/scale) while held; quick-press toggles permanently */
+	TSharedPtr< FUICommandInfo > CurrentTRSSnapToggle;
+
 	/** Enables or disables snapping to surfaces when dragging in the viewport */
 	TSharedPtr< FUICommandInfo > SurfaceSnapping;
 
diff --git a/Engine/Source/Runtime/InteractiveToolsFramework/Private/BaseBehaviors/KeyInputBehavior.cpp b/Engine/Source/Runtime/InteractiveToolsFramework/Private/BaseBehaviors/KeyInputBehavior.cpp
--- a/Engine/Source/Runtime/InteractiveToolsFramework/Private/BaseBehaviors/KeyInputBehavior.cpp
+++ b/Engine/Source/Runtime/InteractiveToolsFramework/Private/BaseBehaviors/KeyInputBehavior.cpp
@@ -52,7 +52,7 @@ FInputCaptureUpdate UKeyInputBehavior::BeginCapture(const FInputDeviceState& Inp
 
 	if ((bRequireAllKeys && AreAllKeysPressed()))
 	{
-		Target->OnKeyPressed(Input.Keyboard.ActiveKey.Button);
+		OnKeyPressedInternal(Input);
 		bAllKeysSeenPressed = true;
 	}
 
@@ -72,14 +72,14 @@ FInputCaptureUpdate UKeyInputBehavior::UpdateCapture(const FInputDeviceState& In
 
 	if (bRequireAllKeys && AreAllKeysPressed())
 	{
-		Target->OnKeyPressed(Input.Keyboard.ActiveKey.Button);
+		OnKeyPressedInternal(Input);
 		bAllKeysSeenPressed = true;
 		return FInputCaptureUpdate::Continue();
 	}
 
 	if (bRequireAllKeys && (!AreAllKeysPressed() && bAllKeysSeenPressed) )
 	{
-		Target->OnKeyReleased(Input.Keyboard.ActiveKey.Button);
+		OnKeyReleasedInternal(Input);
 		bAllKeysSeenPressed = false;
 		return FInputCaptureUpdate::End();
 	}
@@ -115,6 +115,22 @@ bool UKeyInputBehavior::IsKeyPressed(FKey InKey) const
 	return false;
 }
 
+void UKeyInputBehavior::OnKeyPressedInternal(const FInputDeviceState& Input)
+{
+	if (Target)
+	{
+		Target->OnKeyPressed(Input.Keyboard.ActiveKey.Button);
+	}
+}
+
+void UKeyInputBehavior::OnKeyReleasedInternal(const FInputDeviceState& Input)
+{
+	if (Target)
+	{
+		Target->OnKeyReleased(Input.Keyboard.ActiveKey.Button);
+	}
+}
+
 void UKeyInputBehavior::InitializeKeyActivations()
 {
 	KeyActivations.SetNum(TargetKeys.Num());
@@ -141,7 +157,7 @@ bool UKeyInputBehavior::UpdateActivations(const FInputDeviceState& Input, bool b
 			KeyActivations[KeyIndex] = true;
 			if (bEmitOnChange)
 			{
-				Target->OnKeyPressed(Input.Keyboard.ActiveKey.Button);
+				OnKeyPressedInternal(Input);
 			}
 			return true;
 		}
@@ -154,7 +170,7 @@ bool UKeyInputBehavior::UpdateActivations(const FInputDeviceState& Input, bool b
 			KeyActivations[KeyIndex] = false;
 			if (bEmitOnChange)
 			{
-				Target->OnKeyReleased(Input.Keyboard.ActiveKey.Button);
+				OnKeyReleasedInternal(Input);
 			}
 			return true;
 		}
diff --git a/Engine/Source/Runtime/InteractiveToolsFramework/Public/BaseBehaviors/KeyInputBehavior.h b/Engine/Source/Runtime/InteractiveToolsFramework/Public/BaseBehaviors/KeyInputBehavior.h
--- a/Engine/Source/Runtime/InteractiveToolsFramework/Public/BaseBehaviors/KeyInputBehavior.h
+++ b/Engine/Source/Runtime/InteractiveToolsFramework/Public/BaseBehaviors/KeyInputBehavior.h
@@ -101,6 +101,9 @@ public:
 	bool bAlwaysUpdateModifiers = false;
 
 protected:
+	INTERACTIVETOOLSFRAMEWORK_API virtual void OnKeyPressedInternal(const FInputDeviceState& Input);
+	INTERACTIVETOOLSFRAMEWORK_API virtual void OnKeyReleasedInternal(const FInputDeviceState& Input);
+
 	/** Modifier Target object */
 	IKeyInputBehaviorTarget* Target;
 
