[Editor TRS] Gizmo Rendering Issue when changing gizmo scale

#Synced-CL 52474860

diff --git a/Engine/Source/Editor/Experimental/EditorInteractiveToolsFramework/Private/EditorGizmos/EditorTRSGizmo.cpp b/Engine/Source/Editor/Experimental/EditorInteractiveToolsFramework/Private/EditorGizmos/EditorTRSGizmo.cpp
--- a/Engine/Source/Editor/Experimental/EditorInteractiveToolsFramework/Private/EditorGizmos/EditorTRSGizmo.cpp
+++ b/Engine/Source/Editor/Experimental/EditorInteractiveToolsFramework/Private/EditorGizmos/EditorTRSGizmo.cpp
@@ -4319,13 +4319,12 @@ void UEditorTRSGizmo::UpdateRotateAxisDelta(UGizmoElementRotateAxis* InElement)
 	}
 
 	using namespace UE::Editor::InteractiveToolsFramework::Private::TransformGizmoLocals;
-	const FVector WorldCursorPos = ScreenPointToWorldPoint(InteractionScreenCurrPos, GizmoViewContext);
 
 	UGizmoElementRotateAxis::FDeltaParameters DeltaParameters;
 	DeltaParameters.bIsIndirectManipulation = bIndirectManipulation;
 	DeltaParameters.Transform = GetActiveTransform();
 	DeltaParameters.Angle = (InteractionStartAngle + InteractionDelta.X) - AlignedInteractionAngleOffset;
-	DeltaParameters.CursorLocation = WorldCursorPos;
+	DeltaParameters.CursorLocation = InteractionPlanarCurrPoint;
 	DeltaParameters.CursorLocation2D = InteractionPlanarCurrPoint2D;
 	DeltaParameters.RotateMode = RotateMode;
 
diff --git a/Engine/Source/Editor/Experimental/EditorInteractiveToolsFramework/Private/EditorGizmos/GizmoElementRotateAxis.cpp b/Engine/Source/Editor/Experimental/EditorInteractiveToolsFramework/Private/EditorGizmos/GizmoElementRotateAxis.cpp
--- a/Engine/Source/Editor/Experimental/EditorInteractiveToolsFramework/Private/EditorGizmos/GizmoElementRotateAxis.cpp
+++ b/Engine/Source/Editor/Experimental/EditorInteractiveToolsFramework/Private/EditorGizmos/GizmoElementRotateAxis.cpp
@@ -93,13 +93,6 @@ void UGizmoElementRotateAxis::Render(IToolsContextRenderAPI* RenderAPI, const FR
 				DebugData.LastDeltaTransform = DeltaRenderState.LocalToWorldTransform;
 			}
 
-			FRenderTraversalState UnrotatedDeltaRenderState = DeltaRenderState;
-			UnrotatedDeltaRenderState.LocalToWorldTransform.SetRotation(FQuat::Identity);
-
-			FRenderTraversalState UnrotatedUnscaledDeltaRenderState = DeltaRenderState;
-			UnrotatedUnscaledDeltaRenderState.LocalToWorldTransform.SetRotation(FQuat::Identity);
-			UnrotatedUnscaledDeltaRenderState.LocalToWorldTransform.RemoveScaling();
-
 			ApplyUniformConstantScaleToTransform(DeltaRenderState.PixelToWorldScale, DeltaRenderState.LocalToWorldTransform);
 
 			auto RenderDeltaElementWithState = [&](UGizmoElementBase* InElement, const FRenderTraversalState& InRenderState)
@@ -119,7 +112,12 @@ void UGizmoElementRotateAxis::Render(IToolsContextRenderAPI* RenderAPI, const FR
 
 			if (CurrentState.bShowCursor)
 			{
-				RenderDeltaElementWithState(OriginToCursorLineElement, UnrotatedDeltaRenderState);
+				// Match PixelToWorldScale to the transform scale so the dashed line's
+				// Length (world space) converts to local space and back without distortion.
+				FRenderTraversalState CursorLineRenderState = DeltaRenderState;
+				CursorLineRenderState.LocalToWorldTransform.SetRotation(FQuat::Identity);
+				CursorLineRenderState.PixelToWorldScale = CursorLineRenderState.LocalToWorldTransform.GetScale3D().X;
+				RenderDeltaElementWithState(OriginToCursorLineElement, CursorLineRenderState);
 			}
 
 			RenderDeltaElement(OriginElement);
@@ -316,7 +314,7 @@ void UGizmoElementRotateAxis::DrawDebug(IToolsContextRenderAPI* RenderAPI, const
 		if (InSettings.bDrawCursorRay)
 		{
 			Draw.DrawPoint(DebugData.HitPointStart, FLinearColor::Red, 15.0f, false);
-			Draw.DrawPoint(DebugData.HitPointCurrent, FLinearColor::Yellow, 15.0f, false);
+			Draw.DrawPoint(DebugData.HitPointCurrent, FLinearColor::Blue, 17.0f, false);
 		}
 
 		constexpr bool bDrawCursorLocal = false;
@@ -735,7 +733,6 @@ void UGizmoElementRotateAxis::UpdateDelta(const FDeltaParameters& InParameters)
 
 	// Origin to Cursor Line Element
 	{
-		// @TODO: the incoming CursorLocation2D is in plane space, so need to convert here relative to the Delta Transform!
 		OriginToCursorLineElement->SetDirection(DirectionToCursor);
 		OriginToCursorLineElement->SetLength(static_cast<float>(DistanceToCursor));
 		OriginToCursorLineElement->SetViewDepthOffset(0.0f);
