[Editor TRS] Reverted alt+drag portion of 52508366 until non-regression solution found

#Synced-CL 52474860

diff --git a/Engine/Source/Editor/Experimental/EditorInteractiveToolsFramework/Private/EditorGizmos/EditorTRSGizmo.cpp b/Engine/Source/Editor/Experimental/EditorInteractiveToolsFramework/Private/EditorGizmos/EditorTRSGizmo.cpp
--- a/Engine/Source/Editor/Experimental/EditorInteractiveToolsFramework/Private/EditorGizmos/EditorTRSGizmo.cpp
+++ b/Engine/Source/Editor/Experimental/EditorInteractiveToolsFramework/Private/EditorGizmos/EditorTRSGizmo.cpp
@@ -1413,6 +1413,16 @@ void UEditorTRSGizmo::SetupBehaviors()
 		)
 	);
 
+	MouseBehavior->ModifierCheckFunc = [this](const FInputDeviceState& InputDeviceState)
+	{
+		// Filter out Alt unless we need it for duplication
+		if (FInputDeviceState::IsAltKeyDown(InputDeviceState))
+		{
+			return AllowsDragDuplicate();
+		}
+
+		return true;
+	};
 	AddInputBehavior(MouseBehavior);
 	ClickAndDragBehaviorWeak = MouseBehavior;
 
