[5.8 Gizmos] Fix an issue where direct snapped transtion on an axis would move faster than the mouse.

#rnx
#Synced-CL 52476848

diff --git a/Engine/Source/Editor/Experimental/EditorInteractiveToolsFramework/Private/EditorGizmos/EditorTRSGizmo.cpp b/Engine/Source/Editor/Experimental/EditorInteractiveToolsFramework/Private/EditorGizmos/EditorTRSGizmo.cpp
--- a/Engine/Source/Editor/Experimental/EditorInteractiveToolsFramework/Private/EditorGizmos/EditorTRSGizmo.cpp
+++ b/Engine/Source/Editor/Experimental/EditorInteractiveToolsFramework/Private/EditorGizmos/EditorTRSGizmo.cpp
@@ -5836,7 +5836,12 @@ void UEditorTRSGizmo::OnClickDragTranslateAxis(const FInputDeviceRay& DragPos)
 			
 			if (!Delta.IsNearlyZero())
 			{
-				InteractionAxisCurrParam = DistanceAlongAxis;
+				// Updating by the dot product (equivalent of reprojecting and getting length)
+				// Allows changes in the delta due to snapping to be taken into account.
+				// For example, with a snapping grid of 100, a delta value of 51 will snap to 100.
+				// If in the next mouse movement, the mouse did not move along the axis, we want
+				// the next delta to be -49 (snapped to 0). 
+				InteractionAxisCurrParam += Delta.Dot(AxisDirection);
 			}
 			
 			InteractionScreenCurrPos = DragPos.ScreenPosition;
