[Editor TRS] Fix for rotation cursor projection location for arc mode, screen space rotation origin line visibility fix

#Synced-CL 52513890

diff --git a/Engine/Source/Editor/Experimental/EditorInteractiveToolsFramework/Private/EditorGizmos/EditorTRSGizmo.cpp b/Engine/Source/Editor/Experimental/EditorInteractiveToolsFramework/Private/EditorGizmos/EditorTRSGizmo.cpp
--- a/Engine/Source/Editor/Experimental/EditorInteractiveToolsFramework/Private/EditorGizmos/EditorTRSGizmo.cpp
+++ b/Engine/Source/Editor/Experimental/EditorInteractiveToolsFramework/Private/EditorGizmos/EditorTRSGizmo.cpp
@@ -4121,7 +4121,6 @@ void UEditorTRSGizmo::OnClickDragRotateAxisInternal(
 	FVector HitPoint = FVector::ZeroVector;
 	if (ComputePointOnPlaneFromScreen(GizmoViewContext, InDragPos.WorldRay, InDragPos.ScreenPosition, InteractionPlane, HitPoint))
 	{
-		InteractionPlanarCurrPoint = HitPoint;
 		bHitPlane = true;
 	}
 
@@ -4258,6 +4257,12 @@ void UEditorTRSGizmo::OnClickDragRotateAxisInternal(
 	// Update element deltas
 	if (ensure(InElement))
 	{
+		// Re-calculate based on screen plane to avoid plane intersection failures
+		if (ComputePointOnPlaneFromScreen(GizmoViewContext, InDragPos.WorldRay, InDragPos.ScreenPosition, ScreenPlane, HitPoint))
+		{
+			InteractionPlanarCurrPoint = HitPoint;
+		}
+
 		UpdateRotateAxisDelta(InElement);
 	}
 }
@@ -7335,6 +7340,14 @@ void UEditorTRSGizmo::OnClickPressScreenSpaceRotate(const FInputDeviceRay& InPre
 	ViewToAlignedSign = UE::Math::TIntVector3<int8>(1);
 	ClockwiseCorrectionSignForAxis = 1;
 	AlignedInteractionAngleOffset = 0.0;
+	SetRotateMode(EAxisRotateMode::Pull);
+
+	// initialize arc/mixed data
+	const bool bRotatePull = bIndirectManipulation || DefaultRotateMode == EAxisRotateMode::Pull;
+	if (!bRotatePull)
+	{
+		SetRotateMode(DefaultRotateMode);
+	}
 
 	// Always perform ray/plane intersection
 	FVector HitPoint = FVector::ZeroVector;
@@ -7423,11 +7436,13 @@ void UEditorTRSGizmo::OnClickDragScreenSpaceRotate(const FInputDeviceRay& InDrag
 
 	const double PreviousAngle = InteractionCurrAngle;
 	const FVector2D PreviousScreenPos = InteractionScreenCurrPos;
+	
+	using namespace UE::Editor::GizmoMath;
 
 	// Always perform ray/plane intersection
 	FVector HitPoint = FVector::ZeroVector;
 	double HitAngle = 0.0;
-	bool bHitPlane = UE::Editor::GizmoMath::ComputeAngleInPlaneFromScreen(GizmoViewContext, InDragPos.ScreenPosition, InteractionPlane, HitPoint, HitAngle);
+	bool bHitPlane = ComputeAngleInPlaneFromScreen(GizmoViewContext, InDragPos.ScreenPosition, InteractionPlane, HitPoint, HitAngle);
 	if (bHitPlane)
 	{
 		// Correct to avoid wrapping etc.
@@ -7512,11 +7527,19 @@ void UEditorTRSGizmo::OnClickDragScreenSpaceRotate(const FInputDeviceRay& InDrag
 		// Compensate for GizmoElement screen alignment
 		const double AngleForDelta = (InteractionCurrAngle + ScreenSpacePrecisionVisualOffset) * -1;
 
+		// Re-calculate based on screen plane to avoid plane intersection failures
+		if (ComputePointOnPlaneFromScreen(GizmoViewContext, InDragPos.WorldRay, InDragPos.ScreenPosition, ScreenPlane, HitPoint))
+		{
+			InteractionPlanarCurrPoint = HitPoint;
+		}
+
 		UGizmoElementRotateAxis::FDeltaParameters DeltaParameters;
+		DeltaParameters.bIsIndirectManipulation = bIndirectManipulation;
 		DeltaParameters.Transform = GetActiveTransform();
 		DeltaParameters.Angle = AngleForDelta;
 		DeltaParameters.CursorLocation = InteractionPlanarCurrPoint;
 		DeltaParameters.CursorLocation2D = InteractionPlanarCurrPoint2D;
+		DeltaParameters.RotateMode = RotateMode;
 
 		RotateScreenSpaceElement2->UpdateDelta(DeltaParameters);
 	}
