[ITF Gizmos] Legacy gizmo settings are always disabled

#Synced-CL 52547775

diff --git a/Engine/Source/Editor/Experimental/EditorInteractiveToolsFramework/Private/Customizations/TransformGizmoEditorSettingsCustomization.cpp b/Engine/Source/Editor/Experimental/EditorInteractiveToolsFramework/Private/Customizations/TransformGizmoEditorSettingsCustomization.cpp
--- a/Engine/Source/Editor/Experimental/EditorInteractiveToolsFramework/Private/Customizations/TransformGizmoEditorSettingsCustomization.cpp
+++ b/Engine/Source/Editor/Experimental/EditorInteractiveToolsFramework/Private/Customizations/TransformGizmoEditorSettingsCustomization.cpp
@@ -184,6 +184,12 @@ FTransformGizmoEditorSettingsCustomization::FTransformGizmoEditorSettingsCustomi
 		{
 			return UEditorInteractiveGizmoManager::GetTransformGizmoVersion() >= 1;
 		});
+
+	IsUsingLegacyGizmosAttribute = TAttribute<bool>::Create(
+		[]()
+		{
+			return UEditorInteractiveGizmoManager::GetTransformGizmoVersion() == 0;
+		});
 }
 
 void FTransformGizmoEditorSettingsCustomization::CustomizeDetails(const TSharedPtr<IDetailLayoutBuilder>& InDetailBuilder)
@@ -641,23 +647,28 @@ void FTransformGizmoEditorSettingsCustomization::CustomizeDetails(const TSharedP
 
 		// Gizmo Size
 		TSharedPtr<IPropertyHandle> GizmoSizePropertyHandle = InDetailBuilder->GetProperty(GET_MEMBER_NAME_CHECKED(UTransformGizmoEditorSettings, TransformGizmoSize));
-		LegacyGizmoCategory.AddProperty(GizmoSizePropertyHandle);
+		LegacyGizmoCategory.AddProperty(GizmoSizePropertyHandle)
+		.EditCondition(IsUsingLegacyGizmosAttribute, { });
 
 		// Enable Arcball Rotate
 		TSharedPtr<IPropertyHandle> EnableArcballPropertyHandle = InDetailBuilder->GetProperty(GET_MEMBER_NAME_CHECKED(UTransformGizmoEditorSettings, bEnableArcballRotate));
-		LegacyGizmoCategory.AddProperty(EnableArcballPropertyHandle);
+		LegacyGizmoCategory.AddProperty(EnableArcballPropertyHandle)
+		.EditCondition(IsUsingLegacyGizmosAttribute, { });
 
 		// Enable Screen Rotate
 		TSharedPtr<IPropertyHandle> EnableScreenRotatePropertyHandle = InDetailBuilder->GetProperty(GET_MEMBER_NAME_CHECKED(UTransformGizmoEditorSettings, bEnableScreenRotate));
-		LegacyGizmoCategory.AddProperty(EnableScreenRotatePropertyHandle);
+		LegacyGizmoCategory.AddProperty(EnableScreenRotatePropertyHandle)
+		.EditCondition(IsUsingLegacyGizmosAttribute, { });
 
 		// Enable Axis Drawing for Transform Edit Widget
 		TSharedPtr<IPropertyHandle> EnableAxisDrawingPropertyHandle = InDetailBuilder->GetProperty(GET_MEMBER_NAME_CHECKED(UTransformGizmoEditorSettings, bEnableAxisDrawing));
-		LegacyGizmoCategory.AddProperty(EnableAxisDrawingPropertyHandle);
+		LegacyGizmoCategory.AddProperty(EnableAxisDrawingPropertyHandle)
+		.EditCondition(IsUsingLegacyGizmosAttribute, { });
 
 		// Enable Combined Translate/Rotate Widget
 		TSharedPtr<IPropertyHandle> EnableCombinedTranslateRotatePropertyHandle = InDetailBuilder->GetProperty(GET_MEMBER_NAME_CHECKED(UTransformGizmoEditorSettings, bEnableCombinedTranslateRotate));
-		LegacyGizmoCategory.AddProperty(EnableCombinedTranslateRotatePropertyHandle);
+		LegacyGizmoCategory.AddProperty(EnableCombinedTranslateRotatePropertyHandle)
+		.EditCondition(IsUsingLegacyGizmosAttribute, { });
 	}
 }
 
diff --git a/Engine/Source/Editor/Experimental/EditorInteractiveToolsFramework/Private/Customizations/TransformGizmoEditorSettingsCustomization.h b/Engine/Source/Editor/Experimental/EditorInteractiveToolsFramework/Private/Customizations/TransformGizmoEditorSettingsCustomization.h
--- a/Engine/Source/Editor/Experimental/EditorInteractiveToolsFramework/Private/Customizations/TransformGizmoEditorSettingsCustomization.h
+++ b/Engine/Source/Editor/Experimental/EditorInteractiveToolsFramework/Private/Customizations/TransformGizmoEditorSettingsCustomization.h
@@ -42,5 +42,6 @@ private:
 
 private:
 	TAttribute<bool> IsUsingNewGizmosAttribute;
+	TAttribute<bool> IsUsingLegacyGizmosAttribute;
 	TWeakPtr<IDetailLayoutBuilder> WeakDetailBuilder;
 };
diff --git a/Engine/Source/Editor/Experimental/EditorInteractiveToolsFramework/Public/TransformGizmoEditorSettings.h b/Engine/Source/Editor/Experimental/EditorInteractiveToolsFramework/Public/TransformGizmoEditorSettings.h
--- a/Engine/Source/Editor/Experimental/EditorInteractiveToolsFramework/Public/TransformGizmoEditorSettings.h
+++ b/Engine/Source/Editor/Experimental/EditorInteractiveToolsFramework/Public/TransformGizmoEditorSettings.h
@@ -2,7 +2,6 @@
 
 #pragma once
 
-#include "EditorGizmos/EditorGizmoElementShared.h"
 #include "EditorGizmos/EditorTRSGizmo.h"
 #include "Engine/DeveloperSettings.h"
 #include "EditorGizmos/TransformGizmo.h"
@@ -26,35 +25,35 @@ public:
 	UE_API virtual void PostReloadConfig(class FProperty* PropertyThatWasLoaded) override;
 #endif
 
-	UPROPERTY(Config, EditAnywhere, Category = "Transform Gizmo", DisplayName = "Transform Gizmo Size", meta = (ClampMin = "-10.0", ClampMax = "150.0", EditCondition = "!(bUseExperimentalGizmo || bUseEditorTRSGizmo)"))
+	UPROPERTY(Config, EditAnywhere, Category = "Transform Gizmo", DisplayName = "Transform Gizmo Size", meta = (ClampMin = "-10.0", ClampMax = "150.0"))
 	float TransformGizmoSize = 0.0f;
 
 	/**
 	 * Allow arcball rotation with rotate widget
 	 * (updates the setting with the same name found in Level Editor Viewport Settings)
 	 */
-	UPROPERTY(EditAnywhere, Category = "Transform Gizmo", DisplayName = "Enable Arcball Rotate", meta = (EditCondition = "!(bUseExperimentalGizmo || bUseEditorTRSGizmo)"))
+	UPROPERTY(EditAnywhere, Category = "Transform Gizmo", DisplayName = "Enable Arcball Rotate")
 	bool bEnableArcballRotate;
 
 	/**
 	 * Allow screen rotation with rotate widget
 	 * (updates the setting with the same name found in Level Editor Viewport Settings)
 	 */
-	UPROPERTY(EditAnywhere, Category = "Transform Gizmo", DisplayName = "Enable Screen Rotate", meta = (EditCondition = "!(bUseExperimentalGizmo || bUseEditorTRSGizmo)"))
+	UPROPERTY(EditAnywhere, Category = "Transform Gizmo", DisplayName = "Enable Screen Rotate")
 	bool bEnableScreenRotate;
 
 	/**
 	 * If true, the Edit widget of a transform will display the axis
 	 * (updates the setting with the same name found in Level Editor Viewport Settings)
 	 */
-	UPROPERTY(EditAnywhere, Category = "Transform Gizmo", DisplayName = "Enable Axis drawing for Transform Edit Widget", meta = (EditCondition = "!(bUseExperimentalGizmo || bUseEditorTRSGizmo)"))
+	UPROPERTY(EditAnywhere, Category = "Transform Gizmo", DisplayName = "Enable Axis drawing for Transform Edit Widget")
 	bool bEnableAxisDrawing;
 
 	/**
 	 * Allow translate/rotate widget
 	 * (updates the setting with the same name found in Level Editor Viewport Settings)
 	 */
-	UPROPERTY(EditAnywhere, Category = "Transform Gizmo", DisplayName = "Enable Combined Translate/Rotate Widget", meta = (EditCondition = "!(bUseExperimentalGizmo || bUseEditorTRSGizmo)"))
+	UPROPERTY(EditAnywhere, Category = "Transform Gizmo", DisplayName = "Enable Combined Translate/Rotate Widget")
 	bool bEnableCombinedTranslateRotate;
 	//~ End Transform Gizmo Category
 
@@ -66,7 +65,7 @@ public:
 	UPROPERTY(Config, EditAnywhere, Category = "Experimental", Getter = "UsesNewTRSGizmo", Setter = "SetUseNewTRSGizmo")
 	bool bUseEditorTRSGizmo = true;
 
-	UPROPERTY(Config, EditAnywhere, Category = "Experimental", meta = (ShowOnlyInnerProperties, EditCondition = "bUseExperimentalGizmo || bUseEditorTRSGizmo"))
+	UPROPERTY(Config, EditAnywhere, Category = "Experimental", meta = (ShowOnlyInnerProperties))
 	FGizmosParameters GizmosParameters;
 	//~ End Experimental Gizmo Category
 
