[EditorTRSGizmo] Adjust how indirect translation axes manipulation interacts with snapping. Makes overall object motion match when compared to the same mouse movements while unsnapped.

#rnx
#Synced-CL 52554515

diff --git a/Engine/Source/Editor/Experimental/EditorInteractiveToolsFramework/Private/EditorGizmos/EditorTRSGizmo.cpp b/Engine/Source/Editor/Experimental/EditorInteractiveToolsFramework/Private/EditorGizmos/EditorTRSGizmo.cpp
--- a/Engine/Source/Editor/Experimental/EditorInteractiveToolsFramework/Private/EditorGizmos/EditorTRSGizmo.cpp
+++ b/Engine/Source/Editor/Experimental/EditorInteractiveToolsFramework/Private/EditorGizmos/EditorTRSGizmo.cpp
@@ -92,6 +92,12 @@ TAutoConsoleVariable<int32> CVarAlternateSnappingElements(
 	  TEXT("2 = Enabled, Swapped\n")
 );
 
+TAutoConsoleVariable<bool> CVarUseTranslationSnapRecalibration(
+	TEXT("Editor.Gizmo.UseTranslationSnapRecalibration"),
+	true,
+	TEXT("When enabled, adjust tracked mouse position to account for snapping offsets")
+);
+
 }
 
 extern ENGINE_API float GAverageFPS;
@@ -5781,12 +5787,36 @@ void UEditorTRSGizmo::OnClickDragTranslateAxis(const FInputDeviceRay& DragPos)
 		FVector Delta = PerAxisMultiplier * DeltaDistance;
 		Delta = CurrentRotation * Delta;
 
+		const FVector PreSnapDelta = Delta;
 		SnapTranslate(Delta, GetInteractionAxisList(), DragPos.WorldRay);
 
 		ApplyTranslateDelta(Delta);
 		CumulativeDragDelta += Delta;
-
-		InteractionScreenCurrPos = Delta.IsNearlyZero() ? PreviousScreenPos : DragPos.ScreenPosition;
+		
+		if (!Delta.IsNearlyZero())
+		{
+			if (GizmoLocals::CVarUseTranslationSnapRecalibration.GetValueOnAnyThread())
+			{
+				const FVector DeltaDelta = Delta - PreSnapDelta;
+				if (DeltaDelta.IsNearlyZero())
+				{
+					// No snapping, use the actual screen position
+					InteractionScreenCurrPos = DragPos.ScreenPosition;
+				}
+				else
+				{
+					// After a snap, the "current position" should be the mouse position that would _exactly_ match the current snap position.
+					// This gets there by dead-reckoning the screen position based off the post/pre snap ratio.
+					// This makes the assumption that the system has snapped in the same direction as the pre-snap.
+					double Ratio = Delta.Length() / PreSnapDelta.Length();
+					InteractionScreenCurrPos += DragDir * Ratio;
+				}
+			}
+			else
+			{
+				InteractionScreenCurrPos = DragPos.ScreenPosition;
+			}
+		}
 
 		if (CameraFollowsMovement())
 		{
@@ -5836,12 +5866,19 @@ void UEditorTRSGizmo::OnClickDragTranslateAxis(const FInputDeviceRay& DragPos)
 			
 			if (!Delta.IsNearlyZero())
 			{
-				// Updating by the dot product (equivalent of reprojecting and getting length)
-				// Allows changes in the delta due to snapping to be taken into account.
-				// For example, with a snapping grid of 100, a delta value of 51 will snap to 100.
-				// If in the next mouse movement, the mouse did not move along the axis, we want
-				// the next delta to be -49 (snapped to 0). 
-				InteractionAxisCurrParam += Delta.Dot(AxisDirection);
+				if (GizmoLocals::CVarUseTranslationSnapRecalibration.GetValueOnAnyThread())
+				{
+					// Updating by the dot product (equivalent of reprojecting and getting length)
+					// Allows changes in the delta due to snapping to be taken into account.
+					// For example, with a snapping grid of 100, a delta value of 51 will snap to 100.
+					// If in the next mouse movement, the mouse did not move along the axis, we want
+					// the next delta to be -49 (snapped to 0). 
+					InteractionAxisCurrParam += Delta.Dot(AxisDirection);
+				}
+				else
+				{
+					InteractionAxisCurrParam = DistanceAlongAxis;
+				}
 			}
 			
 			InteractionScreenCurrPos = DragPos.ScreenPosition;
