[Editor TRS] Give the gizmo a chance to call a pre/post interaction delegate before/after a click & drag interaction begins/ends
this gives a way to launch some code before/after a gizmo interactions happens and opens up for some customization
the API is internal for now so that it can be refined

#rnx

#Synced-CL 52588375

diff --git a/Engine/Source/Editor/Experimental/EditorInteractiveToolsFramework/Private/EditorGizmos/EditorTRSGizmo.cpp b/Engine/Source/Editor/Experimental/EditorInteractiveToolsFramework/Private/EditorGizmos/EditorTRSGizmo.cpp
--- a/Engine/Source/Editor/Experimental/EditorInteractiveToolsFramework/Private/EditorGizmos/EditorTRSGizmo.cpp
+++ b/Engine/Source/Editor/Experimental/EditorInteractiveToolsFramework/Private/EditorGizmos/EditorTRSGizmo.cpp
@@ -1543,8 +1543,14 @@ void UEditorTRSGizmo::SetupIndirectBehaviors()
 		FButtonBinding(EKeys::MiddleMouseButton).TriggersStart()
 	});
 
-	MiddleClickDragBehavior->CanBeginClickDragFunc = [this](const FInputDeviceRay&)
+	MiddleClickDragBehavior->CanBeginClickDragFunc = [this](const FInputDeviceRay& PressPos)
 	{
+		if (OnPreCanInteractDelegate.IsBound())
+		{
+			FGizmoInteractionDescription Desc({.Ray=PressPos, .bIndirect=true});
+			OnPreCanInteractDelegate.Broadcast(Desc);
+		}
+			
 		static const FInputRayHit InvalidRayHit;
 		// HACK: This lets this behavior win as a tiebreaker vs viewport interactions with the same binding priorities
 		static const FInputRayHit ValidRayHit(UE_LARGE_WORLD_MAX);
@@ -1613,8 +1619,14 @@ void UEditorTRSGizmo::SetupIndirectBehaviors()
 		FButtonBinding(EKeys::MiddleMouseButton).Required(false).TriggersStart(),
 	});
 
-	IndirectClickDragBehavior->CanBeginClickDragFunc = [this](const FInputDeviceRay&)
+	IndirectClickDragBehavior->CanBeginClickDragFunc = [this](const FInputDeviceRay& PressPos)
 	{
+		if (OnPreCanInteractDelegate.IsBound())
+		{
+			FGizmoInteractionDescription Desc({.Ray=PressPos, .bIndirect=true});
+			OnPreCanInteractDelegate.Broadcast(Desc);
+		}
+		
 		static const FInputRayHit InvalidRayHit;
 		static const FInputRayHit ValidRayHit(TNumericLimits<double>::Max());
 
@@ -4934,6 +4946,12 @@ void UEditorTRSGizmo::OnClickRelease(const FInputDeviceRay& InReleasePos, bool b
 
 	PendingDragFunction.Reset();
 	EndDeltas();
+	
+	if (OnPostInteractionDelegate.IsBound())
+	{
+		FGizmoInteractionDescription Desc({.Ray=InReleasePos, .bIndirect=bIndirectManipulation});
+		OnPostInteractionDelegate.Broadcast(Desc);
+	}
 }
 
 void UEditorTRSGizmo::OnTerminateSingleClickAndDragSequence()
diff --git a/Engine/Source/Editor/Experimental/EditorInteractiveToolsFramework/Private/EditorGizmos/TransformGizmo.cpp b/Engine/Source/Editor/Experimental/EditorInteractiveToolsFramework/Private/EditorGizmos/TransformGizmo.cpp
--- a/Engine/Source/Editor/Experimental/EditorInteractiveToolsFramework/Private/EditorGizmos/TransformGizmo.cpp
+++ b/Engine/Source/Editor/Experimental/EditorInteractiveToolsFramework/Private/EditorGizmos/TransformGizmo.cpp
@@ -262,8 +262,14 @@ void UTransformGizmo::SetupIndirectBehaviors()
 	MiddleClickDragBehavior->Initialize();
 	MiddleClickDragBehavior->SetDefaultPriority(INDIRECT_GIZMO_PRIORITY);
 	MiddleClickDragBehavior->SetUseMiddleMouseButton();
-	MiddleClickDragBehavior->CanBeginClickDragFunc = [this](const FInputDeviceRay&)
+	MiddleClickDragBehavior->CanBeginClickDragFunc = [this](const FInputDeviceRay& PressPos)
 	{
+		if (OnPreCanInteractDelegate.IsBound())
+		{
+			FGizmoInteractionDescription Desc({.Ray=PressPos, .bIndirect=true});
+			OnPreCanInteractDelegate.Broadcast(Desc);
+		}
+		
 		static const FInputRayHit InvalidRayHit;
 		static const FInputRayHit ValidRayHit(TNumericLimits<double>::Max());
 		return CanInteract() ? ValidRayHit : InvalidRayHit;
@@ -271,10 +277,11 @@ void UTransformGizmo::SetupIndirectBehaviors()
 	MiddleClickDragBehavior->OnClickPressFunc = [this](const FInputDeviceRay& InPressPos)
 	{
 		bIndirectManipulation = true;
-		if (LastHitPart == ETransformGizmoPartIdentifier::Default)
+		if (LastHitPart != ETransformGizmoPartIdentifier::Default && LastHitPart != GetCurrentModeLastHitPart())
 		{
-			LastHitPart = GetCurrentModeLastHitPart();
+			UpdateInteractingState(false, LastHitPart, true);
 		}
+		LastHitPart = GetCurrentModeLastHitPart();
 		return OnClickPress(InPressPos);
 	};
 	MiddleClickDragBehavior->OnClickDragFunc = [this](const FInputDeviceRay& InDragPos)
@@ -310,8 +317,14 @@ void UTransformGizmo::SetupIndirectBehaviors()
 	}
 	MultiIndirectClickDragBehavior->EnableButton(EKeys::RightMouseButton);
 	MultiIndirectClickDragBehavior->ModifierCheckFunc = FInputDeviceState::IsCtrlKeyDown;
-	MultiIndirectClickDragBehavior->CanBeginClickDragFunc = [this](const FInputDeviceRay&)
+	MultiIndirectClickDragBehavior->CanBeginClickDragFunc = [this](const FInputDeviceRay& PressPos)
 	{
+		if (OnPreCanInteractDelegate.IsBound())
+		{
+			FGizmoInteractionDescription Desc({.Ray=PressPos, .bIndirect=true});
+			OnPreCanInteractDelegate.Broadcast(Desc);
+		}
+		
 		static const FInputRayHit InvalidRayHit;
 		static const FInputRayHit ValidRayHit(TNumericLimits<double>::Max());
 		return CanInteract() ? ValidRayHit : InvalidRayHit;
@@ -2281,6 +2294,12 @@ void UTransformGizmo::OnClickRelease(const FInputDeviceRay& ReleasePos)
 	}
 
 	PendingDragFunction.Reset();
+	
+	if (OnPostInteractionDelegate.IsBound())
+	{
+		FGizmoInteractionDescription Desc({.Ray=ReleasePos, .bIndirect=bIndirectManipulation});
+		OnPostInteractionDelegate.Broadcast(Desc);
+	}
 }
 
 void UTransformGizmo::OnTerminateDragSequence()
@@ -3767,4 +3786,14 @@ void UTransformGizmo::ApplyScaleDelta(const FVector& InScaleDelta)
 	ActiveTarget->SetTransform(CurrentTransform);
 }
 
+FOnGizmoInteraction& UTransformGizmo::OnPreCanInteract()
+{
+	return OnPreCanInteractDelegate;
+}
+	
+FOnGizmoInteraction& UTransformGizmo::OnPostInteraction()
+{
+	return OnPostInteractionDelegate;
+}
+
 #undef LOCTEXT_NAMESPACE
diff --git a/Engine/Source/Editor/Experimental/EditorInteractiveToolsFramework/Public/EditorGizmos/TransformGizmo.h b/Engine/Source/Editor/Experimental/EditorInteractiveToolsFramework/Public/EditorGizmos/TransformGizmo.h
--- a/Engine/Source/Editor/Experimental/EditorInteractiveToolsFramework/Public/EditorGizmos/TransformGizmo.h
+++ b/Engine/Source/Editor/Experimental/EditorInteractiveToolsFramework/Public/EditorGizmos/TransformGizmo.h
@@ -360,6 +360,19 @@ struct FGizmosParameters
 	FGizmoDebugSettings Debug;
 };
 
+/**
+ * FGizmoInteractionDescription provides a way to pass some interaction description in and out before/after starting/ending a gizmo manipulation
+ */
+struct UE_INTERNAL FGizmoInteractionDescription
+{
+	// ray based on the current input device
+	FInputDeviceRay Ray;
+	// whether the interaction is direct or indirect
+	bool bIndirect = false;
+};
+
+DECLARE_MULTICAST_DELEGATE_OneParam(FOnGizmoInteraction, const FGizmoInteractionDescription&);
+
 /**
  * UTransformGizmo provides standard Transformation Gizmo interactions,
  * applied to a UTransformProxy target object. By default the Gizmo will be
@@ -1295,6 +1308,20 @@ protected:
 	/** Used to check if the gimbal mode is currently active (this is updated when ticking the gizmo) */
 	bool bGimbalRotationMode = false;
 	
+	/** pre-click delegate called before a click & drag interaction can start. */
+	UE_INTERNAL FOnGizmoInteraction OnPreCanInteractDelegate;
+	
+	/** conditional post-release delegate after a click & drag interaction ends. */
+	UE_INTERNAL FOnGizmoInteraction OnPostInteractionDelegate;
+	
+public:
+	
+	/** Delegate to a pre-click called delegate before a click & drag interaction can start. */
+	UE_INTERNAL UE_API FOnGizmoInteraction& OnPreCanInteract();
+	
+	/** Delegate to a conditional post-release delegate after a click & drag interaction ends. */
+	UE_INTERNAL UE_API FOnGizmoInteraction& OnPostInteraction();
+	
 private:
 
 	struct FGizmoDebugData
