[Editor Gizmo Settings] Increase minimum width of gizmo preset comparison table to remove clipping.

#rnx

#Synced-CL 52591508

diff --git a/Engine/Source/Editor/Experimental/EditorInteractiveToolsFramework/Private/Customizations/TransformGizmoEditorSettingsCustomization.cpp b/Engine/Source/Editor/Experimental/EditorInteractiveToolsFramework/Private/Customizations/TransformGizmoEditorSettingsCustomization.cpp
--- a/Engine/Source/Editor/Experimental/EditorInteractiveToolsFramework/Private/Customizations/TransformGizmoEditorSettingsCustomization.cpp
+++ b/Engine/Source/Editor/Experimental/EditorInteractiveToolsFramework/Private/Customizations/TransformGizmoEditorSettingsCustomization.cpp
@@ -771,7 +771,7 @@ TSharedRef<IToolTip> FTransformGizmoEditorSettingsCustomization::GeneratePresetT
 		.Padding(0.0f, 4.0f, 0.0f, 0.0f)
 		[
 			SNew(SBox)
-			.MinDesiredWidth(350.0f)
+			.MinDesiredWidth(450.0f)
 			[
 				ListView
 			]
