[TRS Gizmo] : Fix percentage-based scaling when working on multi-selection.  Fix Negative Scaling handling

#Synced-CL 52609747

diff --git a/Engine/Source/Editor/Experimental/EditorInteractiveToolsFramework/Private/EditorGizmos/EditorTRSGizmo.cpp b/Engine/Source/Editor/Experimental/EditorInteractiveToolsFramework/Private/EditorGizmos/EditorTRSGizmo.cpp
--- a/Engine/Source/Editor/Experimental/EditorInteractiveToolsFramework/Private/EditorGizmos/EditorTRSGizmo.cpp
+++ b/Engine/Source/Editor/Experimental/EditorInteractiveToolsFramework/Private/EditorGizmos/EditorTRSGizmo.cpp
@@ -30,6 +30,7 @@
 #include "EditorGizmos/GizmoElementRotateAxis.h"
 #include "EditorGizmos/GizmoElementTranslate.h"
 #include "EditorGizmos/TransformGizmo.h"
+#include "EditorInteractiveGizmoManager.h"
 #include "EditorInteractiveToolsFrameworkStyle.h"
 #include "EditorModeManager.h"
 #include "TransformGizmoEditorSettings.h"
@@ -6571,7 +6572,12 @@ void UEditorTRSGizmo::OnClickDragScale(const FInputDeviceRay& DragPos)
 		{
 			// The attempt at scaling appeared to have no effect.
 			// The reported scale cannot be trusted for the purposes of feeding percentage-based deltas.
-			bDetectedNonStandardScaling = true;
+			if (!(!GEditor->UsePercentageBasedScaling()
+				&& UEditorInteractiveGizmoManager::AreEditorGizmosAllowed(GetEditorViewportClient() ? GetEditorViewportClient()->GetModeTools() : nullptr)
+				&& GetDefault<UTransformGizmoEditorSettings>()->GizmosParameters.ScaleType == EGizmoTransformScaleType::PercentageBased))
+			{
+				bDetectedNonStandardScaling = true;
+			}
 		}
 
 		UE_CLOGF(bLogScale, LogEditorTRSGizmo, Verbose, "[%lli] UpdateCapture: ScaleDelta: %ls | Pre: %ls | Post: %ls", GFrameCounter, *ScaleDelta.ToString(), *PreDeltaScale.ToString(), *PostDeltaScale.ToString());
@@ -6728,7 +6734,53 @@ void UEditorTRSGizmo::ModifyScaleDeltaForScaleType(const EGizmoTransformScaleTyp
 		// Work out the *intended* scale, assumes UTransformProxy returns a correct Transform Scale
 		const FVector TargetScale = StartObjectScale * (FVector::OneVector + DeltaFromStart);
 
-		InOutScaleDelta = TargetScale - CurrentObjectScale;
+		// Compute per-frame ratio delta from raw cumulative deltas (same formula as the global
+		// percentage path). This is proxy-independent: all actors scale by the same factor
+		// regardless of individual scale, and zero crossing works uniformly.
+		// When (1+prev) is near zero (previous cumulative was at -1, i.e. scale was at the zero
+		// clamp point), fall back to Target/Current to avoid division by zero.
+		for (int32 Index = 0; Index < 3; ++Index)
+		{
+			const double Denominator = 1.0 + InPreviousDelta[Index];
+			if (FMath::IsNearlyZero(Denominator, UE_DOUBLE_KINDA_SMALL_NUMBER))
+			{
+				if (!FMath::IsNearlyZero(CurrentObjectScale[Index], UE_DOUBLE_SMALL_NUMBER))
+				{
+					InOutScaleDelta[Index] = TargetScale[Index] / CurrentObjectScale[Index] - 1.0;
+				}
+				else
+				{
+					InOutScaleDelta[Index] = TargetScale[Index] - CurrentObjectScale[Index];
+				}
+			}
+			else
+			{
+				InOutScaleDelta[Index] = (1.0 + InOutScaleDelta[Index]) / Denominator - 1.0;
+			}
+		}
+
+		// Never allow zero scale: it collapses actor positions to the pivot (0 * anything = 0),
+		// permanently losing spatial offsets. When target would be zero on any axis, skip to the
+		// nearest snap grid value so ValidateScale's grid snap won't round it back to zero.
+		{
+			const bool bSnapEnabled = GetDefault<ULevelEditorViewportSettings>()->SnapScaleEnabled;
+			const double MinAbsScale = bSnapEnabled ? GEditor->GetScaleGridSize() : UE_KINDA_SMALL_NUMBER;
+			for (int32 Index = 0; Index < 3; ++Index)
+			{
+				if (FMath::Abs(TargetScale[Index]) < MinAbsScale)
+				{
+					const double SafeTarget = (DeltaFromStart[Index] < -1.0) ? -MinAbsScale : MinAbsScale;
+					if (!FMath::IsNearlyZero(CurrentObjectScale[Index], UE_DOUBLE_SMALL_NUMBER))
+					{
+						InOutScaleDelta[Index] = SafeTarget / CurrentObjectScale[Index] - 1.0;
+					}
+					else
+					{
+						InOutScaleDelta[Index] = SafeTarget - CurrentObjectScale[Index];
+					}
+				}
+			}
+		}
 
 		// Check if this drag is coming right after 2+ scale axes with different sign crossed the zero (see below, where the flag is turned on)
 		if (bZeroCrossingDeltaScaleSwap)
diff --git a/Engine/Source/Editor/LevelEditor/Private/Elements/Actor/ActorElementLevelEditorViewportInteractionCustomization.cpp b/Engine/Source/Editor/LevelEditor/Private/Elements/Actor/ActorElementLevelEditorViewportInteractionCustomization.cpp
--- a/Engine/Source/Editor/LevelEditor/Private/Elements/Actor/ActorElementLevelEditorViewportInteractionCustomization.cpp
+++ b/Engine/Source/Editor/LevelEditor/Private/Elements/Actor/ActorElementLevelEditorViewportInteractionCustomization.cpp
@@ -9,6 +9,8 @@
 #include "Elements/Component/ComponentElementLevelEditorViewportInteractionCustomization.h"
 
 #include "Editor.h"
+#include "EditorInteractiveGizmoManager.h"
+#include "TransformGizmoEditorSettings.h"
 #include "Engine/Brush.h"
 #include "Editor/GroupActor.h"
 #include "ActorGroupingUtils.h"
@@ -59,7 +61,19 @@ void FActorElementLevelEditorViewportInteractionCustomization::GizmoManipulation
 		}
 		else if (!GEditor->UsePercentageBasedScaling())
 		{
-			ModifyScale(Actor, InDragAxis, AdjustedScale, Actor->IsA<ABrush>());
+			bool bNeedsModifyScale = true;
+			if(const FLevelEditorViewportClient* EditorViewportClient = GetLevelEditorViewportClient())
+			{
+				if(UEditorInteractiveGizmoManager::AreEditorGizmosAllowed(EditorViewportClient->GetModeTools()))
+				{
+					bNeedsModifyScale = GetDefault<UTransformGizmoEditorSettings>()->GizmosParameters.ScaleType != EGizmoTransformScaleType::PercentageBased;
+				}
+			}
+
+			if (bNeedsModifyScale)
+			{
+				ModifyScale(Actor, InDragAxis, AdjustedScale, Actor->IsA<ABrush>());
+			}
 		}
 
 		ModifiedDeltaTransform.SetScale3D(AdjustedScale);
diff --git a/Engine/Source/Editor/UnrealEd/Private/Elements/Actor/ActorElementEditorViewportInteractionCustomization.cpp b/Engine/Source/Editor/UnrealEd/Private/Elements/Actor/ActorElementEditorViewportInteractionCustomization.cpp
--- a/Engine/Source/Editor/UnrealEd/Private/Elements/Actor/ActorElementEditorViewportInteractionCustomization.cpp
+++ b/Engine/Source/Editor/UnrealEd/Private/Elements/Actor/ActorElementEditorViewportInteractionCustomization.cpp
@@ -10,6 +10,8 @@
 #include "Components/BrushComponent.h"
 #include "Components/SkinnedMeshComponent.h"
 #include "EditorGizmos/EditorTransformGizmoUtil.h"
+#include "EditorInteractiveGizmoManager.h"
+#include "TransformGizmoEditorSettings.h"
 #include "Engine/DocumentationActor.h"
 #include "Particles/Emitter.h"
 #include "Elements/Interfaces/TypedElementHierarchyInterface.h"
@@ -185,10 +187,24 @@ void FActorElementEditorViewportInteractionCustomization::GizmoManipulationDelta
 				}
 			}
 		}
-		
-		// Flag actors to use old-style scaling or not
-		// @todo: Remove this hack once we have decided on the scaling method to use.
-		AActor::bUsePercentageBasedScaling = GEditor->UsePercentageBasedScaling();
+
+		AActor::bUsePercentageBasedScaling = [this]
+			{
+				if(GEditor->UsePercentageBasedScaling())
+				{
+					return true;
+				}
+
+				if(const IToolkitHost* Host = GetToolkitHost())
+				{
+					if(UEditorInteractiveGizmoManager::AreEditorGizmosAllowed(&Host->GetEditorModeManager()))
+					{
+						return GetDefault<UTransformGizmoEditorSettings>()->GizmosParameters.ScaleType == EGizmoTransformScaleType::PercentageBased;
+					}
+				}
+
+				return false;
+			}();
 
 		RootActor->EditorApplyScale(
 			ModifiedScale,
