[EditorTRSGizmo] Remove multiple duplicate calls to OnClickPress() and OnClickRelease() caused by managing which axis is targeted by indirect axis control. Fix an issue where indirect manipulation of content in the static mesh editor created a dangling undo transaction.

#rnx
#Synced-CL 52694953

diff --git a/Engine/Source/Editor/Experimental/EditorInteractiveToolsFramework/Private/EditorGizmos/EditorTRSGizmo.cpp b/Engine/Source/Editor/Experimental/EditorInteractiveToolsFramework/Private/EditorGizmos/EditorTRSGizmo.cpp
--- a/Engine/Source/Editor/Experimental/EditorInteractiveToolsFramework/Private/EditorGizmos/EditorTRSGizmo.cpp
+++ b/Engine/Source/Editor/Experimental/EditorInteractiveToolsFramework/Private/EditorGizmos/EditorTRSGizmo.cpp
@@ -1763,36 +1763,64 @@ void UEditorTRSGizmo::SetupIndirectBehaviors()
 		// Disable indirect if the current axis is none
 		if (TargetAxes == EAxisList::None)
 		{
-			OnClickRelease(Input, true);
+			if (bInInteraction)
+			{
+				// Otherwise, the EndDrag handler will call the click release
+				if (!IndirectClickDragBehavior->IsDragging())
+				{
+					UE_LOGF(LogEditorTRSGizmo, Verbose, "Calling OnClickRelease() from non-dragging indirect interaction");
+					OnClickRelease(Input, true);
+				}
+			}
+			else
+			{
+				// Reset interaction visuals
+				UpdateInteractingState(false, LastHitPart);
+			}
 			bIndirectManipulation = false;
 			return;
 		}
 
-		bIndirectManipulation = true;
-
 		// Switch to uniform scaling based on preferences (bUniformIndirectScale) or LMB + RMB being down
 		const bool bUniformScaling = CurrentMode == EGizmoTransformMode::Scale && (bUniformIndirectScale || (!bAdditiveIndirectAxes && Input.Mouse.Left.bDown && Input.Mouse.Right.bDown));
 
 		// If "forcing" uniform scaling, we swap the actual axis with ScaleUniform
 		const ETransformGizmoPartIdentifier HitPart = bUniformScaling ? ETransformGizmoPartIdentifier::ScaleUniform : GetPartForAxisList(TargetAxes);
 
-		if (HitPart != GetCurrentModeLastHitPart())
+		const ETransformGizmoPartIdentifier ModeLastHitPart = GetCurrentModeLastHitPart();
+		if (HitPart != ModeLastHitPart)
 		{
+			UE_LOGF(LogEditorTRSGizmo, Verbose, "Changing indirect target to %i", HitPart);
+			// When changing target hit parts in the middle of an interaction:
+			// 1. Stop the current edit action
+			// 2. Prepare to start a new edit action
+			// Resetting drag accomplishes both: the current edit action is stopped and a new one
+			// will be started when a new drag begins.
 			if (bInInteraction)
 			{
-				// Since state changed, make sure to properly end ongoing interaction (e.g. transactions)
-				EndTransformEditSequence();
-
 				if (!IndirectClickDragBehavior->IsDragging())
 				{
+					// If the behavior is in a drag, the ResetDrag and OnDragEnd will handle this
+					// Otherwise it would not be called
 					OnClickRelease(Input, true);
 				}
+				
+				// Resetting the drag means that the mouse movement threshold is reset
+				// This helps prevent small errors when a user release two mouse buttons at once.
 				IndirectClickDragBehavior->ResetDrag(Input);
 			}
 			
 			LastHitPart = HitPart;
-			OnClickPress(Input);
+			SetModeLastHitPart(CurrentMode, HitPart);
+			
+			// Immediately change selection & interaction visualization
+			// This prevents a UI flicker that would be caused while waiting 
+			// in the interstitial period between resetting the drag and starting a new drag.  
+			UpdateInteractingState(false, ModeLastHitPart); 
+			UpdateInteractingState(true, LastHitPart);
 		}
+		
+		bIndirectManipulation = true;
 	};
 	AddInputBehavior(IndirectClickDragBehavior);
 }
